#UE4 - Fix edit conditons not working in details panel

diff --git a/Engine/Source/Editor/PropertyEditor/Private/DetailPropertyRow.cpp b/Engine/Source/Editor/PropertyEditor/Private/DetailPropertyRow.cpp
--- a/Engine/Source/Editor/PropertyEditor/Private/DetailPropertyRow.cpp
+++ b/Engine/Source/Editor/PropertyEditor/Private/DetailPropertyRow.cpp
@@ -265,25 +265,23 @@ bool FDetailPropertyRow::HasEditCondition() const
 
 bool FDetailPropertyRow::GetEnabledState() const
 {
-	if( IsParentEnabled.IsBound() )
-	{
-		return IsParentEnabled.Get();
-	}
-	else if( HasEditCondition() ) 
+	bool Result = IsParentEnabled.Get();
+
+	if( HasEditCondition() ) 
 	{
-		if( PropertyEditor.IsValid() && PropertyEditor->HasEditCondition() )
+		if (CustomEditCondition.IsValid())
 		{
-			return PropertyEditor->IsEditConditionMet();
+			Result = Result && CustomEditCondition->EditConditionValue.Get();
 		}
 		else
 		{
-			return CustomEditCondition->EditConditionValue.Get();
+			Result = Result && PropertyEditor->IsEditConditionMet();
 		}
 	}
-	else
-	{
-		return CustomIsEnabledAttrib.Get();
-	}
+	
+	Result = Result && CustomIsEnabledAttrib.Get();
+
+	return Result;
 }
 
 void FDetailPropertyRow::MakeNameWidget( FDetailWidgetRow& Row, const TSharedPtr<FDetailWidgetRow> InCustomRow ) const
