Restricting it so you can't set a child-component to the same class as the parent blueprint (preventing a crash in the components view).

diff --git a/Engine/Source/Editor/Kismet/Private/BlueprintDetailsCustomization.cpp b/Engine/Source/Editor/Kismet/Private/BlueprintDetailsCustomization.cpp
--- a/Engine/Source/Editor/Kismet/Private/BlueprintDetailsCustomization.cpp
+++ b/Engine/Source/Editor/Kismet/Private/BlueprintDetailsCustomization.cpp
@@ -5,9 +5,9 @@
 #include "BlueprintUtilities.h"
 #include "ScopedTransaction.h"
 #include "GraphEditor.h"
-
+#include "PropertyRestriction.h"
 #include "BlueprintEditor.h"
-
+#include "BlueprintEditorModes.h"
 #include "Editor/PropertyEditor/Public/PropertyEditing.h"
 
 #include "SKismetInspector.h"
@@ -4657,4 +4657,45 @@ void FBlueprintGraphNodeDetails::OnNameCommitted(const FText& InNewText, ETextCo
 	}
 }
 
+TSharedRef<IDetailCustomization> FChildActorComponentDetails::MakeInstance(TWeakPtr<FBlueprintEditor> BlueprintEditorPtrIn)
+{
+	return MakeShareable(new FChildActorComponentDetails(BlueprintEditorPtrIn));
+}
+
+FChildActorComponentDetails::FChildActorComponentDetails(TWeakPtr<FBlueprintEditor> BlueprintEditorPtrIn)
+	: BlueprintEditorPtr(BlueprintEditorPtrIn)
+{
+}
+
+void FChildActorComponentDetails::CustomizeDetails(IDetailLayoutBuilder& DetailBuilder)
+{
+	TSharedPtr<IPropertyHandle> ActorClassProperty = DetailBuilder.GetProperty("ChildActorClass");
+	if (ActorClassProperty->IsValidHandle())
+	{
+		if (BlueprintEditorPtr.IsValid())
+		{
+			// only restrict for the components view (you can successfully add 
+			// a self child component in the execution graphs)
+			if (BlueprintEditorPtr.Pin()->GetCurrentMode() == FBlueprintEditorApplicationModes::BlueprintComponentsMode)
+			{
+				if (UBlueprint* Blueprint = BlueprintEditorPtr.Pin()->GetBlueprintObj())
+				{
+					FText RestrictReason = LOCTEXT("NoSelfChildActors", "Cannot append a child-actor of this blueprint type (could cause infinite recursion).");
+					TSharedPtr<FPropertyRestriction> ClassRestriction = MakeShareable(new FPropertyRestriction(RestrictReason));
+
+					ClassRestriction->AddValue(Blueprint->GetName());
+					ClassRestriction->AddValue(Blueprint->GetPathName());
+					if (Blueprint->GeneratedClass)
+					{
+						ClassRestriction->AddValue(Blueprint->GeneratedClass->GetName());
+						ClassRestriction->AddValue(Blueprint->GeneratedClass->GetPathName());
+					}
+
+					ActorClassProperty->AddRestriction(ClassRestriction.ToSharedRef());
+				}
+			}
+		}
+	}
+}
+
 #undef LOCTEXT_NAMESPACE
diff --git a/Engine/Source/Editor/Kismet/Private/BlueprintDetailsCustomization.h b/Engine/Source/Editor/Kismet/Private/BlueprintDetailsCustomization.h
--- a/Engine/Source/Editor/Kismet/Private/BlueprintDetailsCustomization.h
+++ b/Engine/Source/Editor/Kismet/Private/BlueprintDetailsCustomization.h
@@ -706,3 +706,21 @@ private:
 	/** Weak reference to the Blueprint editor */
 	TWeakPtr<FBlueprintEditor> BlueprintEditorPtr;
 };
+
+/** Details customization for ChildActorComponents */
+class FChildActorComponentDetails : public IDetailCustomization
+{
+public:
+	/** Makes a new instance of this detail layout class for a specific detail view requesting it */
+	static TSharedRef<IDetailCustomization> MakeInstance(TWeakPtr<FBlueprintEditor> BlueprintEditorPtr);
+
+	/** Constructor */
+	FChildActorComponentDetails(TWeakPtr<FBlueprintEditor> BlueprintEditorPtr);
+
+	/** IDetailCustomization interface */
+	virtual void CustomizeDetails(IDetailLayoutBuilder& DetailBuilder) OVERRIDE;
+
+private:
+	/** Weak reference to the Blueprint editor */
+	TWeakPtr<FBlueprintEditor> BlueprintEditorPtr;
+};
diff --git a/Engine/Source/Editor/Kismet/Private/SKismetInspector.cpp b/Engine/Source/Editor/Kismet/Private/SKismetInspector.cpp
--- a/Engine/Source/Editor/Kismet/Private/SKismetInspector.cpp
+++ b/Engine/Source/Editor/Kismet/Private/SKismetInspector.cpp
@@ -240,6 +240,9 @@ void SKismetInspector::Construct(const FArguments& InArgs)
 
 		FOnGetDetailCustomizationInstance GraphNodeDetails = FOnGetDetailCustomizationInstance::CreateStatic(&FBlueprintGraphNodeDetails::MakeInstance, Kismet2Ptr);
 		PropertyView->RegisterInstancedCustomPropertyLayout(UEdGraphNode::StaticClass(), GraphNodeDetails);
+
+		PropertyView->RegisterInstancedCustomPropertyLayout(UChildActorComponent::StaticClass(),
+			FOnGetDetailCustomizationInstance::CreateStatic(&FChildActorComponentDetails::MakeInstance, Kismet2Ptr));
 	}
 
 	// Create the border that all of the content will get stuffed into
diff --git a/Engine/Source/Editor/PropertyEditor/Private/PropertyEditorHelpers.cpp b/Engine/Source/Editor/PropertyEditor/Private/PropertyEditorHelpers.cpp
--- a/Engine/Source/Editor/PropertyEditor/Private/PropertyEditorHelpers.cpp
+++ b/Engine/Source/Editor/PropertyEditor/Private/PropertyEditorHelpers.cpp
@@ -634,6 +634,100 @@ namespace PropertyEditorHelpers
 		}
 	}
 
+	/**
+	 * A helper function that retrieves the path name of the currently selected 
+	 * item (the value that will be used to set the associated property from the 
+	 * "use selection" button)
+	 * 
+	 * @param  PropertyNode		The associated property that the selection is a candidate for.
+	 * @return Empty if the selection isn't compatible with the specified property, else the path-name of the object/class selected in the editor. 
+	 */
+	static FString GetSelectionPathNameForProperty(TSharedRef<FPropertyNode> PropertyNode)
+	{
+		FString SelectionPathName;
+
+		UProperty* Property = PropertyNode->GetProperty();
+		UClassProperty* ClassProperty = Cast<UClassProperty>(Property);
+		UAssetClassProperty* AssetClassProperty = Cast<UAssetClassProperty>(Property);
+
+		if (ClassProperty || AssetClassProperty)
+		{
+			UClass const* const SelectedClass = GEditor->GetFirstSelectedClass(ClassProperty ? ClassProperty->MetaClass : AssetClassProperty->MetaClass);
+			if (SelectedClass != nullptr)
+			{
+				SelectionPathName = SelectedClass->GetPathName();
+			}
+		}
+		else
+		{
+			UClass* ObjectClass = UObject::StaticClass();
+
+			bool bMustBeLevelActor = false;
+			UClass* RequiredInterface = nullptr;
+
+			if (UObjectPropertyBase* ObjectProperty = Cast<UObjectPropertyBase>(Property))
+			{
+				ObjectClass = ObjectProperty->PropertyClass;
+				bMustBeLevelActor = ObjectProperty->GetOwnerProperty()->GetBoolMetaData(TEXT("MustBeLevelActor"));
+				RequiredInterface = ObjectProperty->GetOwnerProperty()->GetClassMetaData(TEXT("MustImplement"));
+			}
+			else if (UInterfaceProperty* InterfaceProperty = Cast<UInterfaceProperty>(Property))
+			{
+				ObjectClass = InterfaceProperty->InterfaceClass;
+			}
+
+			UObject* SelectedObject = nullptr;
+			if (bMustBeLevelActor)
+			{
+				USelection* const SelectedSet = GEditor->GetSelectedActors();
+				SelectedObject = SelectedSet->GetTop(ObjectClass, RequiredInterface);
+			}
+			else 
+			{
+				USelection* const SelectedSet = GEditor->GetSelectedSet(ObjectClass);
+				SelectedObject = SelectedSet->GetTop(ObjectClass, RequiredInterface);
+			}
+
+			if (SelectedObject != nullptr)
+			{
+				SelectionPathName = SelectedObject->GetPathName();
+			}
+		}
+
+		return SelectionPathName;
+	}
+
+	/**
+	 * A helper method that checks to see if the editor's current selection is 
+	 * compatible with the specified property.
+	 * 
+	 * @param  PropertyNode		The property you desire to set from the "use selected" button.
+	 * @return False if the currently selected object is restricted for the specified property, true otherwise.
+	 */
+	static bool IsUseSelectedUnrestricted(TSharedRef<FPropertyNode> PropertyNode)
+	{
+		return !PropertyNode->IsRestricted(GetSelectionPathNameForProperty(PropertyNode));
+	}
+
+	/**
+	 * A helper method that checks to see if the editor's current selection is 
+	 * restricted, and then returns a tooltip explaining why (otherwise, it 
+	 * returns a default explanation of the "use selected" button).
+	 * 
+	 * @param  PropertyNode		The property that would be set from the "use selected" button.
+	 * @return A tooltip for the "use selected" button.
+	 */
+	static FText GetUseSelectedTooltip(TSharedRef<FPropertyNode> PropertyNode)
+	{
+		FText ToolTip;
+		if (!PropertyNode->GenerateRestrictionToolTip(GetSelectionPathNameForProperty(PropertyNode), ToolTip))
+		{
+			ToolTip = LOCTEXT("UseButtonToolTipText", "Use Selected Asset from Content Browser");
+		}
+
+		return ToolTip;
+	}
+
 	TSharedRef<SWidget> MakePropertyButton( const EPropertyButton::Type ButtonType, const TSharedRef< FPropertyEditor >& PropertyEditor )
 	{
 		TSharedPtr<SWidget> NewButton;
@@ -668,8 +762,14 @@ namespace PropertyEditorHelpers
 			break;
 
 		case EPropertyButton::Use:
-			NewButton = PropertyCustomizationHelpers::MakeUseSelectedButton( FSimpleDelegate::CreateSP( PropertyEditor, &FPropertyEditor::UseSelected ) );
-			break;
+			{
+				FSimpleDelegate OnClickDelegate = FSimpleDelegate::CreateSP(PropertyEditor, &FPropertyEditor::UseSelected);
+				TAttribute<bool>::FGetter EnabledDelegate = TAttribute<bool>::FGetter::CreateStatic(&IsUseSelectedUnrestricted, PropertyEditor->GetPropertyNode());
+				TAttribute<FText>::FGetter TooltipDelegate = TAttribute<FText>::FGetter::CreateStatic(&GetUseSelectedTooltip, PropertyEditor->GetPropertyNode());
+
+				NewButton = PropertyCustomizationHelpers::MakeUseSelectedButton(OnClickDelegate, TAttribute<FText>::Create(TooltipDelegate), TAttribute<bool>::Create(EnabledDelegate));
+				break;
+			}
 
 		case EPropertyButton::PickAsset:
 			NewButton = PropertyCustomizationHelpers::MakeAssetPickerAnchorButton( FOnGetAllowedClasses::CreateSP( PropertyEditor, &FPropertyEditor::OnGetClassesForAssetPicker ), FOnAssetSelected::CreateSP( PropertyEditor, &FPropertyEditor::OnAssetSelected ) );
diff --git a/Engine/Source/Editor/PropertyEditor/Private/PropertyHandleImpl.cpp b/Engine/Source/Editor/PropertyEditor/Private/PropertyHandleImpl.cpp
--- a/Engine/Source/Editor/PropertyEditor/Private/PropertyHandleImpl.cpp
+++ b/Engine/Source/Editor/PropertyEditor/Private/PropertyHandleImpl.cpp
@@ -711,7 +711,19 @@ FPropertyAccess::Result FPropertyValueImpl::OnUseSelected()
 			{
 				if (!InterfaceThatMustBeImplemented || SelectedClass->ImplementsInterface(InterfaceThatMustBeImplemented))
 				{
-					SetValueAsString(SelectedClass->GetPathName(), EPropertyValueSetFlags::DefaultFlags);
+					FString const ClassPathName = SelectedClass->GetPathName();
+
+					TArray<FText> RestrictReasons;
+					if (PropertyNodePin->IsRestricted(ClassPathName, RestrictReasons))
+					{
+						check(RestrictReasons.Num() > 0);
+						FMessageDialog::Open(EAppMsgType::Ok, RestrictReasons[0]);
+					}
+					else 
+					{
+						
+						Res = SetValueAsString(ClassPathName, EPropertyValueSetFlags::DefaultFlags);
+					}
 				}
 			}
 		}
@@ -751,14 +763,20 @@ FPropertyAccess::Result FPropertyValueImpl::OnUseSelected()
 
 			if( SelectedObject )
 			{
-				bool bAllObjectPropertyValuesMatch = SetObject( SelectedObject, EPropertyValueSetFlags::DefaultFlags );
+				FString const ObjPathName = SelectedObject->GetPathName();
 
-				// Warn that some object assignments failed.
-				if ( !bAllObjectPropertyValuesMatch )
+				TArray<FText> RestrictReasons;
+				if (PropertyNodePin->IsRestricted(ObjPathName, RestrictReasons))
+				{
+					check(RestrictReasons.Num() > 0);
+					FMessageDialog::Open(EAppMsgType::Ok, RestrictReasons[0]);
+				}
+				else if (!SetObject(SelectedObject, EPropertyValueSetFlags::DefaultFlags))
 				{
+					// Warn that some object assignments failed.
 					FMessageDialog::Open( EAppMsgType::Ok, FText::Format(
-							NSLOCTEXT("UnrealEd", "ObjectAssignmentsFailed", "Failed to assign {0} to the {1} property, see log for details."),
-							FText::FromString(SelectedObject->GetPathName()), FText::FromString(PropertyNodePin->GetDisplayName())) );
+						NSLOCTEXT("UnrealEd", "ObjectAssignmentsFailed", "Failed to assign {0} to the {1} property, see log for details."),
+						FText::FromString(SelectedObject->GetPathName()), FText::FromString(PropertyNodePin->GetDisplayName())) );
 				}
 				else
 				{
