#UE4 - Expand struct children of arrays when adding a new entry to an array

diff --git a/Engine/Source/Editor/PropertyEditor/Private/PropertyNode.cpp b/Engine/Source/Editor/PropertyEditor/Private/PropertyNode.cpp
--- a/Engine/Source/Editor/PropertyEditor/Private/PropertyNode.cpp
+++ b/Engine/Source/Editor/PropertyEditor/Private/PropertyNode.cpp
@@ -332,6 +332,8 @@ FPropertyNode::DataValidationResult FPropertyNode::EnsureDataIsValid()
 			//assume all arrays match the number of property window children
 			bool bArraysMatchChildNum = true;
 
+			bool bArrayHasNewItem = false;
+
 			//verify that the number of object children are the same too
 			UObjectPropertyBase* ObjectProperty = Cast<UObjectPropertyBase>(MyProperty);
 			//check to see, if this an object property, whether the contents are NULL or not.
@@ -380,6 +382,7 @@ FPropertyNode::DataValidationResult FPropertyNode::EnsureDataIsValid()
 					{
 						NumArrayChildren = ArrayNum;
 					}
+					bArrayHasNewItem = GetNumChildNodes() < ArrayNum;
 					//make sure multiple arrays match
 					bArraysHaveEqualNum = bArraysHaveEqualNum && (NumArrayChildren == ArrayNum);
 					//make sure the array matches the number of property node children
@@ -401,6 +404,18 @@ FPropertyNode::DataValidationResult FPropertyNode::EnsureDataIsValid()
 			if (bArraysHaveEqualNum && !bArraysMatchChildNum)
 			{
 				RebuildChildren();
+
+				if( bArrayHasNewItem && ChildNodes.Num() )
+				{
+					TSharedPtr<FPropertyNode> LastChildNode = ChildNodes.Last();
+					// Don't expand huge children
+					if( LastChildNode->GetNumChildNodes() > 0 && LastChildNode->GetNumChildNodes() < 10 )
+					{
+						// Expand the last item for convenience since generally the user will want to edit the new value they added.
+						LastChildNode->SetNodeFlags(EPropertyNodeFlags::Expanded, true);
+					}
+				}
+
 				return ArraySizeChanged;
 			}
 
