TTP #332943: Using the user friendly display name for enum property value widgets.

diff --git a/Engine/Source/Editor/PropertyEditor/Private/UserInterface/PropertyEditor/SPropertyEditorCombo.cpp b/Engine/Source/Editor/PropertyEditor/Private/UserInterface/PropertyEditor/SPropertyEditorCombo.cpp
--- a/Engine/Source/Editor/PropertyEditor/Private/UserInterface/PropertyEditor/SPropertyEditorCombo.cpp
+++ b/Engine/Source/Editor/PropertyEditor/Private/UserInterface/PropertyEditor/SPropertyEditorCombo.cpp
@@ -12,6 +12,32 @@
 
 #define LOCTEXT_NAMESPACE "PropertyEditor"
 
+static int32 FindEnumValueIndex(UEnum* Enum, FString const& ValueString)
+{
+	int32 Index = INDEX_NONE;
+	for(int32 ValIndex = 0; ValIndex < Enum->NumEnums(); ++ValIndex)
+	{
+		FString const EnumName    = Enum->GetEnumName(ValIndex);
+		FString const DisplayName = Enum->GetDisplayNameText(ValIndex).ToString();
+
+		if (DisplayName.Len() > 0)
+		{
+			if (DisplayName == ValueString)
+			{
+				Index = ValIndex;
+				break;
+			}
+		}
+		
+		if (EnumName == ValueString)
+		{
+			Index = ValIndex;
+			break;
+		}
+	}
+	return Index;
+}
+
 void SPropertyEditorCombo::GetDesiredWidth( float& OutMinDesiredWidth, float& OutMaxDesiredWidth )
 {
 	OutMinDesiredWidth = 125.0f;
@@ -67,7 +93,7 @@ FString SPropertyEditorCombo::GetDisplayValueAsString() const
 {
 	const UProperty* Property = PropertyEditor->GetProperty();
 	const UByteProperty* ByteProperty = Cast<const UByteProperty>( Property );
-	const bool bStringEnumProperty = Property && Property->IsA(UStrProperty::StaticClass()) && Property->HasMetaData(TEXT("Enum"));
+	const bool bStringEnumProperty = Property && Property->IsA(UStrProperty::StaticClass()) && Property->HasMetaData(TEXT("Enum"));	
 
 	if ( !(ByteProperty || bStringEnumProperty) )
 	{
@@ -80,13 +106,34 @@ FString SPropertyEditorCombo::GetDisplayValueAsString() const
 		}
 	}
 
-	return PropertyEditor->GetValueAsString();
+	FString ValueString = PropertyEditor->GetValueAsString();
+
+	if (bUsesAlternateDisplayValues && !Property->IsA(UStrProperty::StaticClass()))
+	{
+		// currently only enum properties can use alternate display values; this 
+		// might change, so assert here so that if support is expanded to other 
+		// property types without updating this block of code, we'll catch it quickly
+		UEnum* Enum = CastChecked<UByteProperty>(Property)->Enum;
+		check(Enum != nullptr);
+
+		int32 Index = FindEnumValueIndex(Enum, ValueString);
+		if (Index != INDEX_NONE)
+		{
+			FString DisplayString = Enum->GetDisplayNameText(Index).ToString();
+			if (DisplayString.Len() > 0)
+			{
+				ValueString = DisplayString;
+			}
+		}
+	}
+
+	return ValueString;
 }
 
 void SPropertyEditorCombo::GenerateComboBoxStrings( TArray< TSharedPtr<FString> >& OutComboBoxStrings, TArray< TSharedPtr<FString> >& OutToolTips, TArray<bool>& OutRestrictedItems )
 {
 	const TSharedRef< IPropertyHandle > PropertyHandle = PropertyEditor->GetPropertyHandle();
-	bUsesAlternateDisplayValues = PropertyHandle->GeneratePossibleValues(OutComboBoxStrings, OutToolTips, OutRestrictedItems );
+	bUsesAlternateDisplayValues = PropertyHandle->GeneratePossibleValues(OutComboBoxStrings, OutToolTips, OutRestrictedItems);
 }
 
 void SPropertyEditorCombo::OnComboSelectionChanged( TSharedPtr<FString> NewValue, ESelectInfo::Type SelectInfo )
@@ -106,30 +153,13 @@ void SPropertyEditorCombo::SendToObjects( const FString& NewValue )
 	FString ToolTipValue;
 	if ( bUsesAlternateDisplayValues && !Property->IsA(UStrProperty::StaticClass()))
 	{
-		// currently only enum properties can use alternate display values; this might change, so assert here so that if support is expanded to other property types
-		// without updating this block of code, we'll catch it quickly
+		// currently only enum properties can use alternate display values; this 
+		// might change, so assert here so that if support is expanded to other 
+		// property types without updating this block of code, we'll catch it quickly
 		UEnum* Enum = CastChecked<UByteProperty>(Property)->Enum;
-		check(Enum);
-		int32 Index = INDEX_NONE;
-		for( int32 ItemIndex = 0; ItemIndex <  Enum->NumEnums(); ++ItemIndex )
-		{
-			const FString EnumName = Enum->GetEnumName(ItemIndex);
-			const FString DisplayName = Enum->GetDisplayNameText(ItemIndex ).ToString();
-			if( DisplayName.Len() > 0 )
-			{
-				if ( DisplayName == NewValue )
-				{
-					Index = ItemIndex;
-					break;
-				}
-			}
-			else if (EnumName == NewValue)
-			{
-				Index = ItemIndex;
-				break;
-			}
-		}
+		check(Enum != nullptr);
 
+		int32 Index = FindEnumValueIndex(Enum, NewValue);
 		check( Index != INDEX_NONE );
 
 		Value = Enum->GetEnumName(Index);
