#UE4 - Added missing file

diff --git a/Engine/Source/Editor/PropertyEditor/Private/UserInterface/PropertyEditor/SPropertyEditorAttribute.h b/Engine/Source/Editor/PropertyEditor/Private/UserInterface/PropertyEditor/SPropertyEditorAttribute.h
new file mode 100644
--- /dev/null
+++ b/Engine/Source/Editor/PropertyEditor/Private/UserInterface/PropertyEditor/SPropertyEditorAttribute.h
@@ -0,0 +1,41 @@
+// Copyright 1998-2014 Epic Games, Inc. All Rights Reserved.
+#pragma once
+
+#include "PropertyEditor.h"
+#include "PropertyEditorConstants.h"
+
+
+#define LOCTEXT_NAMESPACE "PropertyEditor"
+
+class SPropertyEditorAttribute : public SCompoundWidget
+{
+public:
+
+	SLATE_BEGIN_ARGS( SPropertyEditorAttribute )
+		: _Font( FEditorStyle::GetFontStyle( PropertyEditorConstants::PropertyFontStyle ) ) 
+		{}
+		SLATE_ATTRIBUTE( FSlateFontInfo, Font )
+	SLATE_END_ARGS()
+
+	void Construct( const FArguments& InArgs, const TSharedRef<FPropertyEditor>& InPropertyEditor )
+	{
+		PropertyEditor = InPropertyEditor;
+	}
+
+	static bool Supports( const TSharedRef<FPropertyEditor>& InPropertyEditor )
+	{
+		const UProperty* NodeProperty = InPropertyEditor->GetProperty();
+
+		return NodeProperty && NodeProperty->IsA<UAttributeProperty>();
+	}
+
+	void GetDesiredWidth( float& OutMinDesiredWidth, float &OutMaxDesiredWidth )
+	{
+		OutMinDesiredWidth = 130.0f;
+		OutMaxDesiredWidth = 130.0f;
+	}
+private:
+	TSharedPtr<FPropertyEditor> PropertyEditor;
+};
+
+#undef LOCTEXT_NAMESPACE
