#ttp 332226 - EDITOR: BSP: In the brush type drop down menu, the selection highlight can be incorrectly displayed due to the reset to default button
#proj UE4
#branch UE4
#summary SPropertyEditorCombo now updates the selection before opening itself
#extra This ensures that it is always up-to-date with the currently selected value, but also requires that the display name match the names used in the combo list (which wasn't previously the case when using alternate names).
#change I've added a way to get the alternate display name for a property, and now use this in SPropertyEditorCombo when it is displaying a drop down list containing alternate display names.
#change Also improved the functions that return an FText of the display names to use the original FText all the way along the calling chain (where possible), rather than going FText -> FString -> FText as it was doing previously.

diff --git a/Engine/Source/Editor/PropertyEditor/Private/Presentation/PropertyEditor/PropertyEditor.cpp b/Engine/Source/Editor/PropertyEditor/Private/Presentation/PropertyEditor/PropertyEditor.cpp
--- a/Engine/Source/Editor/PropertyEditor/Private/Presentation/PropertyEditor/PropertyEditor.cpp
+++ b/Engine/Source/Editor/PropertyEditor/Private/Presentation/PropertyEditor/PropertyEditor.cpp
@@ -114,9 +114,40 @@ FString FPropertyEditor::GetValueAsString() const
 	return Str;
 }
 
+FString FPropertyEditor::GetValueAsDisplayString() const
+{
+	FString Str;
+
+	if( PropertyHandle->GetValueAsDisplayString( Str ) == FPropertyAccess::MultipleValues )
+	{
+		Str = NSLOCTEXT("PropertyEditor", "MultipleValues", "Multiple Values").ToString();
+	}
+
+	return Str;
+}
+
 FText FPropertyEditor::GetValueAsText() const
 {
-	return(FText::FromString(GetValueAsString()));
+	FText Text;
+
+	if( PropertyHandle->GetValueAsFormattedText( Text ) == FPropertyAccess::MultipleValues )
+	{
+		Text = NSLOCTEXT("PropertyEditor", "MultipleValues", "Multiple Values");
+	}
+
+	return Text;
+}
+
+FText FPropertyEditor::GetValueAsDisplayText() const
+{
+	FText Text;
+
+	if( PropertyHandle->GetValueAsDisplayText( Text ) == FPropertyAccess::MultipleValues )
+	{
+		Text = NSLOCTEXT("PropertyEditor", "MultipleValues", "Multiple Values");
+	}
+
+	return Text;
 }
 
 bool FPropertyEditor::PropertyIsA(const UClass* Class) const
diff --git a/Engine/Source/Editor/PropertyEditor/Private/Presentation/PropertyEditor/PropertyEditor.h b/Engine/Source/Editor/PropertyEditor/Private/Presentation/PropertyEditor/PropertyEditor.h
--- a/Engine/Source/Editor/PropertyEditor/Private/Presentation/PropertyEditor/PropertyEditor.h
+++ b/Engine/Source/Editor/PropertyEditor/Private/Presentation/PropertyEditor/PropertyEditor.h
@@ -15,9 +15,15 @@ public:
 	/** @return The String containing the value of the property */
 	FString GetValueAsString() const;
 
+	/** @return The String containing the value of the property, possibly using an alternate form more suitable for display in the UI */
+	FString GetValueAsDisplayString() const;
+
 	/** @return The String containing the value of the property as Text */
 	FText GetValueAsText() const;
 
+	/** @return The String containing the value of the property as Text, possibly using an alternate form more suitable for display in the UI */
+	FText GetValueAsDisplayText() const;
+
 	bool PropertyIsA(const UClass* Class) const;
 
 	bool IsFavorite() const;
diff --git a/Engine/Source/Editor/PropertyEditor/Private/PropertyHandleImpl.cpp b/Engine/Source/Editor/PropertyEditor/Private/PropertyHandleImpl.cpp
--- a/Engine/Source/Editor/PropertyEditor/Private/PropertyHandleImpl.cpp
+++ b/Engine/Source/Editor/PropertyEditor/Private/PropertyHandleImpl.cpp
@@ -43,7 +43,7 @@ void FPropertyValueImpl::GetObjectsToModify( TArray<FObjectBaseAddress>& Objects
 	}
 }
 
-FPropertyAccess::Result FPropertyValueImpl::GetPropertyValueString( FString& OutString, FPropertyNode* InPropertyNode ) const
+FPropertyAccess::Result FPropertyValueImpl::GetPropertyValueString( FString& OutString, FPropertyNode* InPropertyNode, const bool bAllowAlternateDisplayValue ) const
 {
 	FPropertyAccess::Result Result = FPropertyAccess::Success;
 
@@ -63,11 +63,18 @@ FPropertyAccess::Result FPropertyValueImpl::GetPropertyValueString( FString& Out
 			if( Property != NULL && InPropertyNode->GetParentNode() != NULL )
 			{
 				Property->ExportText_Direct(OutString, ValueAddress, ValueAddress, NULL, PPF_PropertyWindow );
+
 				UByteProperty* ByteProperty = Cast<UByteProperty>(Property);
 				if ( ByteProperty != NULL && ByteProperty->Enum != NULL )
 				{
-					uint8 EnumValue = ByteProperty->GetPropertyValue(ValueAddress);
-					OutString = ByteProperty->Enum->GetEnumName(EnumValue);
+					const uint8 EnumValueIndex = ByteProperty->GetPropertyValue(ValueAddress);
+
+					// See if we specified an alternate name for this value using metadata
+					OutString = ByteProperty->Enum->GetDisplayNameText(EnumValueIndex).ToString();
+					if(!bAllowAlternateDisplayValue || OutString.Len() == 0) 
+					{
+						OutString = ByteProperty->Enum->GetEnumName(EnumValueIndex);
+					}
 				}
 			}
 			else
@@ -85,7 +92,7 @@ FPropertyAccess::Result FPropertyValueImpl::GetPropertyValueString( FString& Out
 	return Result;
 }
 
-FPropertyAccess::Result FPropertyValueImpl::GetPropertyValueText( FText& OutText, FPropertyNode* InPropertyNode ) const
+FPropertyAccess::Result FPropertyValueImpl::GetPropertyValueText( FText& OutText, FPropertyNode* InPropertyNode, const bool bAllowAlternateDisplayValue ) const
 {
 	FPropertyAccess::Result Result = FPropertyAccess::Success;
 
@@ -107,17 +114,25 @@ FPropertyAccess::Result FPropertyValueImpl::GetPropertyValueText( FText& OutText
 			}
 			else
 			{
-				FString OutString;
+				FString ExportedTextString;
+				Property->ExportText_Direct(ExportedTextString, ValueAddress, ValueAddress, NULL, PPF_PropertyWindow );
 
-				Property->ExportText_Direct(OutString, ValueAddress, ValueAddress, NULL, PPF_PropertyWindow );
 				UByteProperty* ByteProperty = Cast<UByteProperty>(Property);
 				if ( ByteProperty != NULL && ByteProperty->Enum != NULL )
 				{
-					uint8 EnumValueIndex = ByteProperty->GetPropertyValue(ValueAddress);
-					OutString = ByteProperty->Enum->GetEnumName(EnumValueIndex);
+					const uint8 EnumValueIndex = ByteProperty->GetPropertyValue(ValueAddress);
+					
+					// See if we specified an alternate name for this value using metadata
+					OutText = ByteProperty->Enum->GetDisplayNameText(EnumValueIndex);
+					if(!bAllowAlternateDisplayValue || OutText.IsEmpty()) 
+					{
+						OutText = ByteProperty->Enum->GetEnumText(EnumValueIndex);
+					}
+				}
+				else
+				{
+					OutText = FText::FromString(ExportedTextString);
 				}
-
-				OutText = FText::FromString(OutString);
 			}
 		}
 
@@ -541,7 +556,25 @@ FPropertyAccess::Result FPropertyValueImpl::GetValueAsString( FString& OutString
 
 	if( PropertyNodePin.IsValid() )
 	{
-		Res = GetPropertyValueString( OutString, PropertyNodePin.Get() );
+		Res = GetPropertyValueString( OutString, PropertyNodePin.Get(), false/*bAllowAlternateDisplayValue*/ );
+	}
+	else
+	{
+		Res = FPropertyAccess::Fail;
+	}
+
+	return Res;
+}
+
+FPropertyAccess::Result FPropertyValueImpl::GetValueAsDisplayString( FString& OutString ) const
+{
+	TSharedPtr<FPropertyNode> PropertyNodePin = PropertyNode.Pin();
+
+	FPropertyAccess::Result Res = FPropertyAccess::Success;
+
+	if( PropertyNodePin.IsValid() )
+	{
+		Res = GetPropertyValueString( OutString, PropertyNodePin.Get(), true/*bAllowAlternateDisplayValue*/ );
 	}
 	else
 	{
@@ -559,7 +592,25 @@ FPropertyAccess::Result FPropertyValueImpl::GetValueAsText( FText& OutText ) con
 
 	if( PropertyNodePin.IsValid() )
 	{
-		Res = GetPropertyValueText( OutText, PropertyNodePin.Get() );
+		Res = GetPropertyValueText( OutText, PropertyNodePin.Get(), false/*bAllowAlternateDisplayValue*/ );
+	}
+	else
+	{
+		Res = FPropertyAccess::Fail;
+	}
+
+	return Res;
+}
+
+FPropertyAccess::Result FPropertyValueImpl::GetValueAsDisplayText( FText& OutText ) const
+{
+	TSharedPtr<FPropertyNode> PropertyNodePin = PropertyNode.Pin();
+
+	FPropertyAccess::Result Res = FPropertyAccess::Success;
+
+	if( PropertyNodePin.IsValid() )
+	{
+		Res = GetPropertyValueText( OutText, PropertyNodePin.Get(), true/*bAllowAlternateDisplayValue*/ );
 	}
 	else
 	{
@@ -1436,6 +1487,21 @@ FPropertyAccess::Result FPropertyHandleBase::GetValueAsFormattedString( FString&
 	return Implementation->GetValueAsString(OutValue);
 }
 
+FPropertyAccess::Result FPropertyHandleBase::GetValueAsDisplayString( FString& OutValue ) const
+{
+	return Implementation->GetValueAsDisplayString(OutValue);
+}
+
+FPropertyAccess::Result FPropertyHandleBase::GetValueAsFormattedText( FText& OutValue ) const
+{
+	return Implementation->GetValueAsText(OutValue);
+}
+
+FPropertyAccess::Result FPropertyHandleBase::GetValueAsDisplayText( FText& OutValue ) const
+{
+	return Implementation->GetValueAsDisplayText(OutValue);
+}
+
 FPropertyAccess::Result FPropertyHandleBase::SetValueFromFormattedString( const FString& InValue, EPropertyValueSetFlags::Type Flags )
 {
 	return Implementation->SetValueAsString( InValue, Flags );
diff --git a/Engine/Source/Editor/PropertyEditor/Private/PropertyHandleImpl.h b/Engine/Source/Editor/PropertyEditor/Private/PropertyHandleImpl.h
--- a/Engine/Source/Editor/PropertyEditor/Private/PropertyHandleImpl.h
+++ b/Engine/Source/Editor/PropertyEditor/Private/PropertyHandleImpl.h
@@ -148,7 +148,15 @@ public:
 	 */
 	FPropertyAccess::Result GetValueAsString( FString& OutString ) const;
 
-/**
+	/**
+	 * Get the value of a property as a formatted string, possibly using an alternate form more suitable for display in the UI
+	 *
+	 * @param OutValue The formatted string value to set
+	 * @return The result of attempting to get the value
+	 */
+	FPropertyAccess::Result GetValueAsDisplayString( FString& OutString ) const;
+
+	/**
 	 * Get the value of a property as FText.
 	 *
 	 * @param OutValue The formatted text value to set
@@ -156,6 +164,14 @@ public:
 	 */
 	FPropertyAccess::Result GetValueAsText( FText& OutText ) const;
 
+	/**
+	 * Get the value of a property as FText, possibly using an alternate form more suitable for display in the UI
+	 *
+	 * @param OutValue The formatted text value to set
+	 * @return The result of attempting to get the value
+	 */
+	FPropertyAccess::Result GetValueAsDisplayText( FText& OutText ) const;
+
 	/**
 	 * Sets the value of a property formatted from a string.
 	 * Each UProperty has a specific string format it requires
@@ -291,16 +307,18 @@ protected:
 	 * 
 	 * @param InPropertyNode	The property to get value from
 	 * @param OutText			The property formatted in a string
+	 * @param bAllowAlternateDisplayValue Allow the function to potentially use an alternate form more suitable for display in the UI
 	 * @return true if the value was retrieved successfully
 	 */
-	FPropertyAccess::Result GetPropertyValueString( FString& OutString, FPropertyNode* InPropertyNode ) const;
+	FPropertyAccess::Result GetPropertyValueString( FString& OutString, FPropertyNode* InPropertyNode, const bool bAllowAlternateDisplayValue ) const;
 
 	/**
 	 * @param InPropertyNode	The property to get value from
 	 * @param OutText			The property formatted in text
+	 * @param bAllowAlternateDisplayValue Allow the function to potentially use an alternate form more suitable for display in the UI
 	 * @return true if the value was retrieved successfully
 	 */
-	FPropertyAccess::Result GetPropertyValueText( FText& OutText, FPropertyNode* InPropertyNode ) const;
+	FPropertyAccess::Result GetPropertyValueText( FText& OutText, FPropertyNode* InPropertyNode, const bool bAllowAlternateDisplayValue ) const;
 
 protected:
 	/** Property node used to access UProperty and address of object to change */
@@ -356,6 +374,9 @@ public:
 	virtual void SetOnPropertyValueChanged( const FSimpleDelegate& InOnPropertyValueChanged ) OVERRIDE;
 	virtual int32 GetIndexInArray() const OVERRIDE;
 	virtual FPropertyAccess::Result GetValueAsFormattedString( FString& OutValue ) const OVERRIDE;
+	virtual FPropertyAccess::Result GetValueAsDisplayString( FString& OutValue ) const OVERRIDE;
+	virtual FPropertyAccess::Result GetValueAsFormattedText( FText& OutValue ) const OVERRIDE;
+	virtual FPropertyAccess::Result GetValueAsDisplayText( FText& OutValue ) const OVERRIDE;
 	virtual FPropertyAccess::Result SetValueFromFormattedString( const FString& InValue,  EPropertyValueSetFlags::Type Flags = EPropertyValueSetFlags::DefaultFlags ) OVERRIDE;
 	virtual TSharedPtr<IPropertyHandle> GetChildHandle( uint32 ChildIndex ) const OVERRIDE;
 	virtual TSharedPtr<IPropertyHandle> GetChildHandle( FName ChildName ) const OVERRIDE;
diff --git a/Engine/Source/Editor/PropertyEditor/Private/UserInterface/PropertyEditor/SPropertyComboBox.cpp b/Engine/Source/Editor/PropertyEditor/Private/UserInterface/PropertyEditor/SPropertyComboBox.cpp
--- a/Engine/Source/Editor/PropertyEditor/Private/UserInterface/PropertyEditor/SPropertyComboBox.cpp
+++ b/Engine/Source/Editor/PropertyEditor/Private/UserInterface/PropertyEditor/SPropertyComboBox.cpp
@@ -36,6 +36,7 @@ void SPropertyComboBox::Construct( const FArguments& InArgs )
 		.OptionsSource(&ComboItemList)
 		.OnGenerateWidget(this, &SPropertyComboBox::OnGenerateComboWidget)
 		.OnSelectionChanged(this, &SPropertyComboBox::OnSelectionChangedInternal)
+		.OnComboBoxOpening(InArgs._OnComboBoxOpening)
 		.InitiallySelectedItem(InitiallySelectedItem)
 		);
 }
diff --git a/Engine/Source/Editor/PropertyEditor/Private/UserInterface/PropertyEditor/SPropertyComboBox.h b/Engine/Source/Editor/PropertyEditor/Private/UserInterface/PropertyEditor/SPropertyComboBox.h
--- a/Engine/Source/Editor/PropertyEditor/Private/UserInterface/PropertyEditor/SPropertyComboBox.h
+++ b/Engine/Source/Editor/PropertyEditor/Private/UserInterface/PropertyEditor/SPropertyComboBox.h
@@ -15,6 +15,7 @@ public:
 		SLATE_ATTRIBUTE( FString, VisibleText )
 		SLATE_ARGUMENT( TArray< TSharedPtr< FString > >, ToolTipList )
 		SLATE_EVENT( FOnSelectionChanged, OnSelectionChanged )
+		SLATE_EVENT( FOnComboBoxOpening, OnComboBoxOpening )
 		SLATE_ARGUMENT( FSlateFontInfo, Font )
 	SLATE_END_ARGS()
 
diff --git a/Engine/Source/Editor/PropertyEditor/Private/UserInterface/PropertyEditor/SPropertyEditorCombo.cpp b/Engine/Source/Editor/PropertyEditor/Private/UserInterface/PropertyEditor/SPropertyEditorCombo.cpp
--- a/Engine/Source/Editor/PropertyEditor/Private/UserInterface/PropertyEditor/SPropertyEditorCombo.cpp
+++ b/Engine/Source/Editor/PropertyEditor/Private/UserInterface/PropertyEditor/SPropertyEditorCombo.cpp
@@ -78,6 +78,7 @@ void SPropertyEditorCombo::Construct( const FArguments& InArgs, const TSharedRef
 		.ComboItemList( ComboItems )
 		.RestrictedList( Restrictions )
 		.OnSelectionChanged( this, &SPropertyEditorCombo::OnComboSelectionChanged )
+		.OnComboBoxOpening( this, &SPropertyEditorCombo::OnComboOpening )
 		.VisibleText( this, &SPropertyEditorCombo::GetDisplayValueAsString )
 		.ToolTipText( InPropertyEditor, &FPropertyEditor::GetValueAsString );
 
@@ -106,28 +107,7 @@ FString SPropertyEditorCombo::GetDisplayValueAsString() const
 		}
 	}
 
-	FString ValueString = PropertyEditor->GetValueAsString();
-
-	if (bUsesAlternateDisplayValues && !Property->IsA(UStrProperty::StaticClass()))
-	{
-		// currently only enum properties can use alternate display values; this 
-		// might change, so assert here so that if support is expanded to other 
-		// property types without updating this block of code, we'll catch it quickly
-		UEnum* Enum = CastChecked<UByteProperty>(Property)->Enum;
-		check(Enum != nullptr);
-
-		int32 Index = FindEnumValueIndex(Enum, ValueString);
-		if (Index != INDEX_NONE)
-		{
-			FString DisplayString = Enum->GetDisplayNameText(Index).ToString();
-			if (DisplayString.Len() > 0)
-			{
-				ValueString = DisplayString;
-			}
-		}
-	}
-
-	return ValueString;
+	return (bUsesAlternateDisplayValues) ? PropertyEditor->GetValueAsDisplayString() : PropertyEditor->GetValueAsString();
 }
 
 void SPropertyEditorCombo::GenerateComboBoxStrings( TArray< TSharedPtr<FString> >& OutComboBoxStrings, TArray< TSharedPtr<FString> >& OutToolTips, TArray<bool>& OutRestrictedItems )
@@ -144,6 +124,14 @@ void SPropertyEditorCombo::OnComboSelectionChanged( TSharedPtr<FString> NewValue
 	}
 }
 
+void SPropertyEditorCombo::OnComboOpening()
+{
+	// try and re-sync the selection in the combo list in case it was changed since Construct was called
+	// this would fail if the displayed value doesn't match the equivalent value in the combo list
+	FString CurrentDisplayValue = GetDisplayValueAsString();
+	ComboBox->SetSelectedItem(CurrentDisplayValue);
+}
+
 void SPropertyEditorCombo::SendToObjects( const FString& NewValue )
 {
 	const TSharedRef< FPropertyNode > PropertyNode = PropertyEditor->GetPropertyNode();
@@ -159,7 +147,7 @@ void SPropertyEditorCombo::SendToObjects( const FString& NewValue )
 		UEnum* Enum = CastChecked<UByteProperty>(Property)->Enum;
 		check(Enum != nullptr);
 
-		int32 Index = FindEnumValueIndex(Enum, NewValue);
+		const int32 Index = FindEnumValueIndex(Enum, NewValue);
 		check( Index != INDEX_NONE );
 
 		Value = Enum->GetEnumName(Index);
diff --git a/Engine/Source/Editor/PropertyEditor/Private/UserInterface/PropertyEditor/SPropertyEditorCombo.h b/Engine/Source/Editor/PropertyEditor/Private/UserInterface/PropertyEditor/SPropertyEditorCombo.h
--- a/Engine/Source/Editor/PropertyEditor/Private/UserInterface/PropertyEditor/SPropertyEditorCombo.h
+++ b/Engine/Source/Editor/PropertyEditor/Private/UserInterface/PropertyEditor/SPropertyEditorCombo.h
@@ -22,6 +22,7 @@ private:
 
 	void GenerateComboBoxStrings( TArray< TSharedPtr<FString> >& OutComboBoxStrings, TArray< TSharedPtr<FString> >& OutToolTips, TArray<bool>& OutRestrictedItems );
 	void OnComboSelectionChanged( TSharedPtr<FString> NewValue, ESelectInfo::Type SelectInfo );
+	void OnComboOpening();
 
 	virtual void SendToObjects( const FString& NewValue );
 
diff --git a/Engine/Source/Editor/PropertyEditor/Public/PropertyHandle.h b/Engine/Source/Editor/PropertyEditor/Public/PropertyHandle.h
--- a/Engine/Source/Editor/PropertyEditor/Public/PropertyHandle.h
+++ b/Engine/Source/Editor/PropertyEditor/Public/PropertyHandle.h
@@ -66,6 +66,30 @@ public:
 	 */
 	virtual FPropertyAccess::Result GetValueAsFormattedString( FString& OutValue ) const = 0;
 
+	/**
+	 * Gets the value formatted as a string, possibly using an alternate form more suitable for display in the UI
+	 *
+	 * @param OutValue	String where the value is stored.  Remains unchanged if the value could not be set
+	 * @return The result of attempting to get the value
+	 */
+	virtual FPropertyAccess::Result GetValueAsDisplayString( FString& OutValue ) const = 0;
+
+	/**
+	 * Gets the value formatted as a string, as Text.
+	 *
+	 * @param OutValue	Text where the value is stored.  Remains unchanged if the value could not be set
+	 * @return The result of attempting to get the value
+	 */
+	virtual FPropertyAccess::Result GetValueAsFormattedText( FText& OutValue ) const = 0;
+
+	/**
+	 * Gets the value formatted as a string, as Text, possibly using an alternate form more suitable for display in the UI
+	 *
+	 * @param OutValue	Text where the value is stored.  Remains unchanged if the value could not be set
+	 * @return The result of attempting to get the value
+	 */
+	virtual FPropertyAccess::Result GetValueAsDisplayText( FText& OutValue ) const = 0;
+
 	/**
 	 * Sets the value formatted as a string.
 	 *
