Fixing a crash bug in the property grid editor, where if a tooltip spawns while saving it illegally attempts to access the object, now it aborts getting the value during those times.

diff --git a/Engine/Source/Editor/PropertyEditor/Private/UserInterface/PropertyEditor/SPropertyEditorAsset.cpp b/Engine/Source/Editor/PropertyEditor/Private/UserInterface/PropertyEditor/SPropertyEditorAsset.cpp
--- a/Engine/Source/Editor/PropertyEditor/Private/UserInterface/PropertyEditor/SPropertyEditorAsset.cpp
+++ b/Engine/Source/Editor/PropertyEditor/Private/UserInterface/PropertyEditor/SPropertyEditorAsset.cpp
@@ -384,7 +384,15 @@ void SPropertyEditorAsset::SetValue( const UObject* InObject )
 
 FPropertyAccess::Result SPropertyEditorAsset::GetValue( FObjectOrAssetData& OutValue ) const
 {
+	// Potentially accessing the value while garbage collecting or saving the package could trigger a crash.
+	// so we fail to get the value when that is occuring.
+	if ( GIsSavingPackage || GIsGarbageCollecting )
+	{
+		return FPropertyAccess::Fail;
+	}
+
 	FPropertyAccess::Result Result = FPropertyAccess::Fail;
+
 	if( PropertyEditor.IsValid() && PropertyEditor->GetPropertyHandle()->IsValidHandle() )
 	{
 		UObject* Object = NULL;
