Add the ability to not generate headers for custom builder

diff --git a/Engine/Source/Editor/PropertyEditor/Public/PropertyCustomizationHelpers.h b/Engine/Source/Editor/PropertyEditor/Public/PropertyCustomizationHelpers.h
--- a/Engine/Source/Editor/PropertyEditor/Public/PropertyCustomizationHelpers.h
+++ b/Engine/Source/Editor/PropertyEditor/Public/PropertyCustomizationHelpers.h
@@ -182,9 +182,10 @@ class FDetailArrayBuilder : public IDetailCustomNodeBuilder
 {
 public:
 
-	FDetailArrayBuilder( TSharedRef<IPropertyHandle> InBaseProperty )
+	FDetailArrayBuilder( TSharedRef<IPropertyHandle> InBaseProperty, bool InGenerateHeader = true )
 		: ArrayProperty( InBaseProperty->AsArray() )
 		, BaseProperty( InBaseProperty )
+		, bGenerateHeader( InGenerateHeader)
 	{
 		check( ArrayProperty.IsValid() );
 
@@ -224,30 +225,33 @@ public:
 
 	virtual void GenerateHeaderRowContent( FDetailWidgetRow& NodeRow ) OVERRIDE
 	{
-		TSharedPtr<SResetToDefaultMenu> ResetToDefaultMenu;
-		bool bDisplayResetToDefault = false;
-		NodeRow
-		.FilterString( DisplayName.Len() > 0 ? DisplayName : BaseProperty->GetPropertyDisplayName() )
-		.NameContent()
-		[
-			BaseProperty->CreatePropertyNameWidget( DisplayName, bDisplayResetToDefault )
-		]
-		.ValueContent()
-		[
-			SNew( SHorizontalBox )
-			+ SHorizontalBox::Slot()
-			[
-				BaseProperty->CreatePropertyValueWidget()
-			]
-			+ SHorizontalBox::Slot()
-			.AutoWidth()
-			.Padding( FMargin( 2.0f, 0.0f, 0.0f, 0.0f ) )
-			[
-				SAssignNew( ResetToDefaultMenu, SResetToDefaultMenu )
-			]
-		];
-
-		ResetToDefaultMenu->AddProperty( BaseProperty );
+		if (bGenerateHeader)
+		{
+			TSharedPtr<SResetToDefaultMenu> ResetToDefaultMenu;
+			bool bDisplayResetToDefault = false;
+			NodeRow
+				.FilterString(DisplayName.Len() > 0 ? DisplayName : BaseProperty->GetPropertyDisplayName())
+				.NameContent()
+				[
+					BaseProperty->CreatePropertyNameWidget(DisplayName, bDisplayResetToDefault)
+				]
+			.ValueContent()
+				[
+					SNew(SHorizontalBox)
+					+ SHorizontalBox::Slot()
+					[
+						BaseProperty->CreatePropertyValueWidget()
+					]
+					+ SHorizontalBox::Slot()
+						.AutoWidth()
+						.Padding(FMargin(2.0f, 0.0f, 0.0f, 0.0f))
+						[
+							SAssignNew(ResetToDefaultMenu, SResetToDefaultMenu)
+						]
+				];
+
+			ResetToDefaultMenu->AddProperty(BaseProperty);
+		}
 	}
 
 	virtual void GenerateChildContent( IDetailChildrenBuilder& ChildrenBuilder ) OVERRIDE
@@ -281,6 +285,7 @@ private:
 	TSharedPtr<IPropertyHandleArray> ArrayProperty;
 	TSharedRef<IPropertyHandle> BaseProperty;
 	FSimpleDelegate OnRebuildChildren;
+	bool bGenerateHeader;
 };
 
 /**
