You can now use an actor picker to choose which actor to attach to

TTP# 279062 - ROCKET: TASK: PUNTABLE: Viewports: "Drag to Attach" feature

Moved FEdModePicker out of the Private section of the PropertyEditor module, and into its own module so that it can be shared between the PropertyEditor and LevelEditor modules. Also renamed it to FEdModeActorPicker.

Added a way to filter FEdModeActorPicker by checking the actor directly (this complements the class type check it already had).

Added an actor picker button next to the scene outliner in the "Attach To" menu on actors. This is also mapped to the Alt+A command.

Built and tested in unity and non-unity.

diff --git a/Engine/Source/Editor/ActorPickerMode/ActorPickerMode.Build.cs b/Engine/Source/Editor/ActorPickerMode/ActorPickerMode.Build.cs
new file mode 100644
--- /dev/null
+++ b/Engine/Source/Editor/ActorPickerMode/ActorPickerMode.Build.cs
@@ -0,0 +1,22 @@
+// Copyright 1998-2014 Epic Games, Inc. All Rights Reserved.
+
+using UnrealBuildTool;
+
+public class ActorPickerMode : ModuleRules
+{
+    public ActorPickerMode(TargetInfo Target)
+	{
+		PrivateDependencyModuleNames.AddRange(
+			new string[] {
+				"Core",
+				"CoreUObject",
+				"Engine",
+                "InputCore",
+                "Slate",
+                "SlateCore",
+                "EditorStyle",
+				"UnrealEd",
+			}
+		);
+	}
+}
diff --git a/Engine/Source/Editor/ActorPickerMode/Private/ActorPickerMode.cpp b/Engine/Source/Editor/ActorPickerMode/Private/ActorPickerMode.cpp
new file mode 100644
--- /dev/null
+++ b/Engine/Source/Editor/ActorPickerMode/Private/ActorPickerMode.cpp
@@ -0,0 +1,64 @@
+// Copyright 1998-2014 Epic Games, Inc. All Rights Reserved.
+
+#include "ActorPickerModePrivatePCH.h"
+
+IMPLEMENT_MODULE( FActorPickerModeModule, ActorPickerMode );
+
+void FActorPickerModeModule::StartupModule()
+{
+	EdModeActorPicker = MakeShareable(new FEdModeActorPicker);
+	GEditorModeTools().RegisterMode(EdModeActorPicker.ToSharedRef());
+
+	// we want to be able to perform this action with all the built-in editor modes
+	GEditorModeTools().RegisterCompatibleModes(EdModeActorPicker->GetID(), FBuiltinEditorModes::EM_Default);
+	GEditorModeTools().RegisterCompatibleModes(EdModeActorPicker->GetID(), FBuiltinEditorModes::EM_Placement);
+	GEditorModeTools().RegisterCompatibleModes(EdModeActorPicker->GetID(), FBuiltinEditorModes::EM_Bsp);
+	GEditorModeTools().RegisterCompatibleModes(EdModeActorPicker->GetID(), FBuiltinEditorModes::EM_Geometry);
+	GEditorModeTools().RegisterCompatibleModes(EdModeActorPicker->GetID(), FBuiltinEditorModes::EM_InterpEdit);
+	GEditorModeTools().RegisterCompatibleModes(EdModeActorPicker->GetID(), FBuiltinEditorModes::EM_Texture);
+	GEditorModeTools().RegisterCompatibleModes(EdModeActorPicker->GetID(), FBuiltinEditorModes::EM_MeshPaint);
+	GEditorModeTools().RegisterCompatibleModes(EdModeActorPicker->GetID(), FBuiltinEditorModes::EM_Landscape);
+	GEditorModeTools().RegisterCompatibleModes(EdModeActorPicker->GetID(), FBuiltinEditorModes::EM_Foliage);
+	GEditorModeTools().RegisterCompatibleModes(EdModeActorPicker->GetID(), FBuiltinEditorModes::EM_Level);
+	GEditorModeTools().RegisterCompatibleModes(EdModeActorPicker->GetID(), FBuiltinEditorModes::EM_StreamingLevel);
+	GEditorModeTools().RegisterCompatibleModes(EdModeActorPicker->GetID(), FBuiltinEditorModes::EM_Physics);
+}
+
+void FActorPickerModeModule::ShutdownModule()
+{
+	EndActorPickingMode();
+
+	GEditorModeTools().UnregisterCompatibleModes(EdModeActorPicker->GetID(), FBuiltinEditorModes::EM_Default);
+	GEditorModeTools().UnregisterCompatibleModes(EdModeActorPicker->GetID(), FBuiltinEditorModes::EM_Placement);
+	GEditorModeTools().UnregisterCompatibleModes(EdModeActorPicker->GetID(), FBuiltinEditorModes::EM_Bsp);
+	GEditorModeTools().UnregisterCompatibleModes(EdModeActorPicker->GetID(), FBuiltinEditorModes::EM_Geometry);
+	GEditorModeTools().UnregisterCompatibleModes(EdModeActorPicker->GetID(), FBuiltinEditorModes::EM_InterpEdit);
+	GEditorModeTools().UnregisterCompatibleModes(EdModeActorPicker->GetID(), FBuiltinEditorModes::EM_Texture);
+	GEditorModeTools().UnregisterCompatibleModes(EdModeActorPicker->GetID(), FBuiltinEditorModes::EM_MeshPaint);
+	GEditorModeTools().UnregisterCompatibleModes(EdModeActorPicker->GetID(), FBuiltinEditorModes::EM_Landscape);
+	GEditorModeTools().UnregisterCompatibleModes(EdModeActorPicker->GetID(), FBuiltinEditorModes::EM_Foliage);
+	GEditorModeTools().UnregisterCompatibleModes(EdModeActorPicker->GetID(), FBuiltinEditorModes::EM_Level);
+	GEditorModeTools().UnregisterCompatibleModes(EdModeActorPicker->GetID(), FBuiltinEditorModes::EM_StreamingLevel);
+	GEditorModeTools().UnregisterCompatibleModes(EdModeActorPicker->GetID(), FBuiltinEditorModes::EM_Physics);
+	GEditorModeTools().UnregisterMode(EdModeActorPicker.ToSharedRef());
+	EdModeActorPicker = nullptr;
+}
+
+void FActorPickerModeModule::BeginActorPickingMode(FOnGetAllowedClasses InOnGetAllowedClasses, FOnShouldFilterActor InOnShouldFilterActor, FOnActorSelected InOnActorSelected)
+{
+	EndActorPickingMode();
+	EdModeActorPicker->BeginMode(InOnGetAllowedClasses, InOnShouldFilterActor, InOnActorSelected);
+}
+
+void FActorPickerModeModule::EndActorPickingMode()
+{
+	if(IsInActorPickingMode())
+	{
+		EdModeActorPicker->EndMode();
+	}
+}
+
+bool FActorPickerModeModule::IsInActorPickingMode() const
+{
+	return GEditorModeTools().IsModeActive(EdModeActorPicker->GetID());
+}
diff --git a/Engine/Source/Editor/ActorPickerMode/Private/ActorPickerModePrivatePCH.h b/Engine/Source/Editor/ActorPickerMode/Private/ActorPickerModePrivatePCH.h
new file mode 100644
--- /dev/null
+++ b/Engine/Source/Editor/ActorPickerMode/Private/ActorPickerModePrivatePCH.h
@@ -0,0 +1,8 @@
+// Copyright 1998-2014 Epic Games, Inc. All Rights Reserved.
+
+#pragma once
+
+#include "UnrealEd.h"
+
+#include "ActorPickerMode.h"
+#include "EditorModeActorPicker.h"
diff --git a/Engine/Source/Editor/ActorPickerMode/Private/EditorModeActorPicker.cpp b/Engine/Source/Editor/ActorPickerMode/Private/EditorModeActorPicker.cpp
new file mode 100644
--- /dev/null
+++ b/Engine/Source/Editor/ActorPickerMode/Private/EditorModeActorPicker.cpp
@@ -0,0 +1,242 @@
+// Copyright 1998-2014 Epic Games, Inc. All Rights Reserved.
+
+#include "ActorPickerModePrivatePCH.h"
+
+#define LOCTEXT_NAMESPACE "PropertyPicker"
+
+FEdModeActorPicker::FEdModeActorPicker()
+{
+	ID = FBuiltinEditorModes::EM_ActorPicker;
+	PickState = EPickState::NotOverViewport;
+	HoveredActor.Reset();
+}
+	
+void FEdModeActorPicker::Tick(FLevelEditorViewportClient* ViewportClient, float DeltaTime)
+{
+	if(CursorDecoratorWindow.IsValid())
+	{
+		CursorDecoratorWindow->MoveWindowTo(FSlateApplication::Get().GetCursorPos() + FSlateApplication::Get().GetCursorSize());
+	}
+
+	FEdMode::Tick(ViewportClient, DeltaTime);
+}
+
+bool FEdModeActorPicker::MouseEnter(FLevelEditorViewportClient* ViewportClient, FViewport* Viewport,int32 x, int32 y)
+{
+	PickState = EPickState::OverViewport;
+	HoveredActor.Reset();
+	return FEdMode::MouseEnter(ViewportClient, Viewport, x, y);
+}
+
+bool FEdModeActorPicker::MouseLeave(FLevelEditorViewportClient* ViewportClient, FViewport* Viewport)
+{
+	PickState = EPickState::NotOverViewport;
+	HoveredActor.Reset();
+	return FEdMode::MouseLeave(ViewportClient, Viewport);
+}
+
+bool FEdModeActorPicker::MouseMove(FLevelEditorViewportClient* ViewportClient, FViewport* Viewport, int32 x, int32 y)
+{
+	if (ViewportClient == GCurrentLevelEditingViewportClient)
+	{
+		PickState = EPickState::OverViewport;
+		HoveredActor.Reset();
+
+		int32 HitX = Viewport->GetMouseX();
+		int32 HitY = Viewport->GetMouseY();
+		HHitProxy* HitProxy = Viewport->GetHitProxy(HitX, HitY);
+		if (HitProxy != NULL && HitProxy->IsA(HActor::StaticGetType()))
+		{
+			HActor* ActorHit = static_cast<HActor*>(HitProxy);
+			if(ActorHit->Actor != NULL)
+			{
+				AActor* Actor = ActorHit->Actor;
+				HoveredActor = Actor;
+				PickState =  IsActorValid(Actor) ? EPickState::OverActor : EPickState::OverIncompatibleActor;
+			}
+		}
+	}
+	else
+	{
+		PickState = EPickState::NotOverViewport;
+		HoveredActor.Reset();
+	}
+
+	return true;
+}
+
+bool FEdModeActorPicker::LostFocus(FLevelEditorViewportClient* ViewportClient, FViewport* Viewport)
+{
+	if (ViewportClient == GCurrentLevelEditingViewportClient)
+	{
+		// Make sure actor picking mode is disabled once the active viewport loses focus
+		EndMode();
+		return true;
+	}
+
+	return false;
+}
+
+bool FEdModeActorPicker::InputKey(FLevelEditorViewportClient* ViewportClient, FViewport* Viewport, FKey Key, EInputEvent Event)
+{
+	if (ViewportClient == GCurrentLevelEditingViewportClient)
+	{
+		if (Key == EKeys::LeftMouseButton && Event == IE_Pressed)
+		{
+			// See if we clicked on an actor
+			int32 HitX = Viewport->GetMouseX();
+			int32 HitY = Viewport->GetMouseY();
+			HHitProxy*	HitProxy = Viewport->GetHitProxy(HitX, HitY);
+			if (HitProxy != NULL && HitProxy->IsA(HActor::StaticGetType()))
+			{
+				HActor* ActorHit = static_cast<HActor*>(HitProxy);
+				if(IsActorValid(ActorHit->Actor))
+				{
+					OnActorSelected.ExecuteIfBound(ActorHit->Actor);
+					EndMode();
+				}
+			}
+			return true;
+		}
+		else if(Key == EKeys::Escape && Event == IE_Pressed)
+		{
+			EndMode();
+			return true;
+		}
+	}
+	else
+	{
+		EndMode();
+	}
+
+	return false;
+}
+
+bool FEdModeActorPicker::GetCursor(EMouseCursor::Type& OutCursor) const
+{
+	if(HoveredActor.IsValid() && PickState == EPickState::OverActor)
+	{
+		OutCursor = EMouseCursor::EyeDropper;
+	}
+	else
+	{
+		OutCursor = EMouseCursor::SlashedCircle;
+	}
+	
+	return true;
+}
+
+bool FEdModeActorPicker::UsesToolkits() const
+{
+	return false;
+}
+
+
+void FEdModeActorPicker::BeginMode(FOnGetAllowedClasses InOnGetAllowedClasses, FOnShouldFilterActor InOnShouldFilterActor, FOnActorSelected InOnActorSelected )
+{
+	OnActorSelected = InOnActorSelected;
+	OnGetAllowedClasses = InOnGetAllowedClasses;
+	OnShouldFilterActor = InOnShouldFilterActor;
+	GEditorModeTools().ActivateMode(GetID());
+	CursorDecoratorWindow = SWindow::MakeCursorDecorator();
+	FSlateApplication::Get().AddWindow(CursorDecoratorWindow.ToSharedRef(), true);
+	CursorDecoratorWindow->SetContent(
+		SNew(SToolTip)
+		.Text(this, &FEdModeActorPicker::GetCursorDecoratorText)
+		);
+
+	HoveredActor.Reset();
+	PickState = EPickState::NotOverViewport;
+}
+
+void FEdModeActorPicker::EndMode()
+{
+	OnActorSelected = FOnActorSelected();
+	OnGetAllowedClasses = FOnGetAllowedClasses();
+	OnShouldFilterActor = FOnShouldFilterActor();
+	GEditorModeTools().DeactivateMode(GetID());
+
+	if (CursorDecoratorWindow.IsValid())
+	{
+		CursorDecoratorWindow->RequestDestroyWindow();
+		CursorDecoratorWindow.Reset();
+	}
+
+	HoveredActor.Reset();
+	PickState = EPickState::NotOverViewport;
+}
+
+FString FEdModeActorPicker::GetCursorDecoratorText() const
+{
+	switch(PickState)
+	{
+	default:
+	case EPickState::NotOverViewport:
+		return LOCTEXT("PickActor_NotOverViewport", "Pick an actor by clicking on it in the active level viewport").ToString();
+	case EPickState::OverViewport:
+		return LOCTEXT("PickActor_NotOverActor", "Pick an actor by clicking on it").ToString();
+	case EPickState::OverIncompatibleActor:
+		{
+			if(HoveredActor.IsValid())
+			{
+				FFormatNamedArguments Arguments;
+				Arguments.Add(TEXT("Actor"), FText::FromString(HoveredActor.Get()->GetName()));
+				return FText::Format(LOCTEXT("PickActor_OverIncompatibleActor", "{Actor} is incompatible"), Arguments).ToString();
+			}
+			else
+			{
+				return LOCTEXT("PickActor_NotOverActor", "Pick an actor by clicking on it").ToString();
+			}
+		}
+	case EPickState::OverActor:
+		{
+			if(HoveredActor.IsValid())
+			{
+				FFormatNamedArguments Arguments;
+				Arguments.Add(TEXT("Actor"), FText::FromString(HoveredActor.Get()->GetName()));
+				return FText::Format(LOCTEXT("PickActor_OverActor", "Pick {Actor}"), Arguments).ToString();
+			}
+			else
+			{
+				return LOCTEXT("PickActor_NotOverActor", "Pick an actor by clicking on it").ToString();
+			}
+		}
+	}
+}
+
+bool FEdModeActorPicker::IsActorValid(const AActor *const Actor) const
+{
+	bool bIsValid = false;
+
+	if(Actor)
+	{
+		bool bHasValidClass = true;
+		if(OnGetAllowedClasses.IsBound())
+		{
+			bHasValidClass = false;
+
+			TArray<const UClass*> AllowedClasses;
+			OnGetAllowedClasses.Execute(AllowedClasses);
+			for(const UClass* AllowedClass : AllowedClasses)
+			{
+				if(Actor->IsA(AllowedClass))
+				{
+					bHasValidClass = true;
+					break;
+				}
+			}
+		}
+
+		bool bHasValidActor = true;
+		if(OnShouldFilterActor.IsBound())
+		{
+			bHasValidActor = OnShouldFilterActor.Execute(Actor);
+		}
+
+		bIsValid = bHasValidClass && bHasValidActor;
+	}
+
+	return bIsValid;
+}
+
+#undef LOCTEXT_NAMESPACE
diff --git a/Engine/Source/Editor/ActorPickerMode/Private/EditorModeActorPicker.h b/Engine/Source/Editor/ActorPickerMode/Private/EditorModeActorPicker.h
new file mode 100644
--- /dev/null
+++ b/Engine/Source/Editor/ActorPickerMode/Private/EditorModeActorPicker.h
@@ -0,0 +1,65 @@
+// Copyright 1998-2014 Epic Games, Inc. All Rights Reserved.
+
+#pragma once
+
+namespace EPickState
+{
+	enum Type
+	{
+		NotOverViewport,
+		OverViewport,
+		OverIncompatibleActor,
+		OverActor,
+	};
+}
+
+/**
+ * Editor mode used to interactively pick actors in the level editor viewports.
+ */
+class FEdModeActorPicker : public FEdMode
+{
+public:
+	FEdModeActorPicker();
+	
+	/** Register the picker editor mode */
+	static void RegisterEditMode();
+
+	/** Unregister the picker editor mode */
+	static void UnregisterEditMode();
+
+	/** Begin FEdMode interface */
+	virtual void Tick(FLevelEditorViewportClient* ViewportClient, float DeltaTime) OVERRIDE;
+	virtual bool MouseEnter(FLevelEditorViewportClient* ViewportClient, FViewport* Viewport,int32 x, int32 y) OVERRIDE;
+	virtual bool MouseLeave(FLevelEditorViewportClient* ViewportClient, FViewport* Viewport) OVERRIDE;
+	virtual bool MouseMove(FLevelEditorViewportClient* ViewportClient, FViewport* Viewport, int32 x, int32 y) OVERRIDE;
+	virtual bool LostFocus(FLevelEditorViewportClient* ViewportClient, FViewport* Viewport) OVERRIDE;
+	virtual bool InputKey(FLevelEditorViewportClient* ViewportClient, FViewport* Viewport, FKey Key, EInputEvent Event) OVERRIDE;
+	virtual bool GetCursor(EMouseCursor::Type& OutCursor) const OVERRIDE;
+	virtual bool UsesToolkits() const OVERRIDE;
+	/** End FEdMode interface */
+
+	/** Start up the mode, reset state etc. */
+	void BeginMode(FOnGetAllowedClasses InOnGetAllowedClasses, FOnShouldFilterActor InOnShouldFilterActor, FOnActorSelected InOnActorSelected );
+
+	/** End the mode */
+	void EndMode();
+
+	/** Delegate used to display information about picking near the cursor */
+	FString GetCursorDecoratorText() const;
+
+	bool IsActorValid(const AActor *const Actor) const;
+
+	/** Flag for display state */
+	TWeakObjectPtr<AActor> HoveredActor;
+
+	/** Flag for display state */
+	EPickState::Type PickState;
+
+	/** The window that owns the decorator widget */
+	TSharedPtr<SWindow> CursorDecoratorWindow;
+
+	/** Delegates used to pick actors */
+	FOnActorSelected OnActorSelected;
+	FOnGetAllowedClasses OnGetAllowedClasses;
+	FOnShouldFilterActor OnShouldFilterActor;
+};
diff --git a/Engine/Source/Editor/ActorPickerMode/Public/ActorPickerMode.h b/Engine/Source/Editor/ActorPickerMode/Public/ActorPickerMode.h
new file mode 100644
--- /dev/null
+++ b/Engine/Source/Editor/ActorPickerMode/Public/ActorPickerMode.h
@@ -0,0 +1,40 @@
+// Copyright 1998-2014 Epic Games, Inc. All Rights Reserved.
+
+#pragma once
+
+#include "IFilter.h"
+#include "FilterCollection.h"
+
+DECLARE_DELEGATE_OneParam( FOnGetAllowedClasses, TArray<const UClass*>& );
+DECLARE_DELEGATE_OneParam( FOnActorSelected, AActor* );
+DECLARE_DELEGATE_RetVal_OneParam(bool, FOnShouldFilterActor, const AActor*);
+
+/**
+ * Actor picker mode module
+ */
+class ACTORPICKERMODE_API FActorPickerModeModule : public IModuleInterface
+{
+public:
+	// IModuleInterface
+	virtual void StartupModule() OVERRIDE;
+	virtual void ShutdownModule() OVERRIDE;
+	// End of IModuleInterface
+
+	/** 
+	 * Enter actor picking mode (note: will cancel any current actor picking)
+	 * @param	InOnGetAllowedClasses	Delegate used to only allow actors using a particular set of classes (empty to accept all actor classes; works alongside InOnShouldFilterActor)
+	 * @param	InOnShouldFilterActor	Delegate used to only allow particular actors (empty to accept all actors; works alongside InOnGetAllowedClasses)
+	 * @param	InOnActorSelected		Delegate to call when a valid actor is selected
+	 */
+	void BeginActorPickingMode(FOnGetAllowedClasses InOnGetAllowedClasses, FOnShouldFilterActor InOnShouldFilterActor, FOnActorSelected InOnActorSelected);
+
+	/** Exit actor picking mode */
+	void EndActorPickingMode();
+
+	/** @return Whether or not actor picking mode is currently active */
+	bool IsInActorPickingMode() const;
+
+private:
+	/** The one and only actor picking mode */
+	TSharedPtr<class FEdModeActorPicker> EdModeActorPicker;
+};
diff --git a/Engine/Source/Editor/LevelEditor/LevelEditor.Build.cs b/Engine/Source/Editor/LevelEditor/LevelEditor.Build.cs
--- a/Engine/Source/Editor/LevelEditor/LevelEditor.Build.cs
+++ b/Engine/Source/Editor/LevelEditor/LevelEditor.Build.cs
@@ -42,7 +42,8 @@ public class LevelEditor : ModuleRules
 				"RenderCore",
 				"DeviceProfileServices",
 				"ContentBrowser",
-                "SceneOutliner"
+                "SceneOutliner",
+                "ActorPickerMode"
 			}
 		);
 
diff --git a/Engine/Source/Editor/LevelEditor/Private/LevelEditor.cpp b/Engine/Source/Editor/LevelEditor/Private/LevelEditor.cpp
--- a/Engine/Source/Editor/LevelEditor/Private/LevelEditor.cpp
+++ b/Engine/Source/Editor/LevelEditor/Private/LevelEditor.cpp
@@ -751,6 +751,11 @@ void FLevelEditorModule::BindGlobalLevelEditorCommands()
 		FExecuteAction::CreateStatic( &FLevelEditorActionCallbacks::AttachSelectedActors )
 		);
 
+	ActionList.MapAction(
+		Commands.AttachActorIteractive,
+		FExecuteAction::CreateStatic( &FLevelEditorActionCallbacks::AttachActorIteractive )
+		);
+
 	ActionList.MapAction(
 		Commands.CreateNewOutlinerFolder,
 		FExecuteAction::CreateStatic( &FLevelEditorActionCallbacks::CreateNewOutlinerFolder_Clicked )
diff --git a/Engine/Source/Editor/LevelEditor/Private/LevelEditorActions.cpp b/Engine/Source/Editor/LevelEditor/Private/LevelEditorActions.cpp
--- a/Engine/Source/Editor/LevelEditor/Private/LevelEditorActions.cpp
+++ b/Engine/Source/Editor/LevelEditor/Private/LevelEditorActions.cpp
@@ -40,6 +40,7 @@
 #include "IAnalyticsProvider.h"
 
 #include "EditorActorFolders.h"
+#include "ActorPickerMode.h"
 
 DEFINE_LOG_CATEGORY_STATIC(LevelEditorActions, Log, All);
 
@@ -1019,6 +1020,41 @@ void FLevelEditorActionCallbacks::AttachSelectedActors()
 	GUnrealEd->AttachSelectedActors();
 }
 
+void FLevelEditorActionCallbacks::AttachActorIteractive()
+{
+	if(GUnrealEd->GetSelectedActorCount())
+	{
+		FActorPickerModeModule& ActorPickerMode = FModuleManager::Get().GetModuleChecked<FActorPickerModeModule>("ActorPickerMode");
+
+		ActorPickerMode.BeginActorPickingMode(
+			FOnGetAllowedClasses(), 
+			FOnShouldFilterActor::CreateStatic(&FLevelEditorActionCallbacks::IsAttachableActor), 
+			FOnActorSelected::CreateStatic(&FLevelEditorActionCallbacks::AttachToActor)
+			);
+	}
+}
+
+bool FLevelEditorActionCallbacks::IsAttachableActor( const AActor* const ParentActor )
+{
+	for ( FSelectionIterator It( GEditor->GetSelectedActorIterator() ) ; It ; ++It )
+	{
+		AActor* Actor = static_cast<AActor*>( *It );
+		if (!GEditor->CanParentActors(ParentActor, Actor))
+		{
+			return false;
+		}
+
+		USceneComponent* ChildRoot = Actor->GetRootComponent();
+		USceneComponent* ParentRoot = ParentActor->GetRootComponent();
+
+		if (ChildRoot != nullptr && ParentRoot != nullptr && ChildRoot->IsAttachedTo(ParentRoot))
+		{
+			return false;
+		}
+	}
+	return true;
+}
+
 void FLevelEditorActionCallbacks::CreateNewOutlinerFolder_Clicked()
 {
 	const FName NewFolderName = FActorFolders::Get().GetDefaultFolderNameForSelection(*GetWorld());
@@ -2560,6 +2596,7 @@ void FLevelEditorCommands::RegisterCommands()
 	UI_COMMAND( LockActorMovement, "Lock Actor Movement", "Locks the actor so it cannot be moved", EUserInterfaceActionType::ToggleButton, FInputGesture() );
 	UI_COMMAND( DetachFromParent, "Detach", "Detach the actor from its parent", EUserInterfaceActionType::Button, FInputGesture() );
 	UI_COMMAND( AttachSelectedActors, "Attach Selected Actors", "Attach the selected actors to the last selected actor", EUserInterfaceActionType::Button, FInputGesture(EModifierKey::Alt, EKeys::B) );
+	UI_COMMAND( AttachActorIteractive, "Attach Actor Interactive", "Start an interactive actor picker to let you choose a parent for the currently selected actor", EUserInterfaceActionType::Button, FInputGesture(EModifierKey::Alt, EKeys::A) );
 	UI_COMMAND( CreateNewOutlinerFolder, "Create Folder", "Place the selected actors in a new folder", EUserInterfaceActionType::Button, FInputGesture() );
 	UI_COMMAND( HoldToEnableVertexSnapping, "Hold to Enable Vertex Snapping", "When the key binding is pressed and held vertex snapping will be enabled", EUserInterfaceActionType::ToggleButton, FInputGesture(EKeys::V) );
 
diff --git a/Engine/Source/Editor/LevelEditor/Private/LevelEditorContextMenu.cpp b/Engine/Source/Editor/LevelEditor/Private/LevelEditorContextMenu.cpp
--- a/Engine/Source/Editor/LevelEditor/Private/LevelEditorContextMenu.cpp
+++ b/Engine/Source/Editor/LevelEditor/Private/LevelEditorContextMenu.cpp
@@ -704,83 +704,77 @@ void FLevelEditorContextMenuImpl::FillTransformMenu( FMenuBuilder& MenuBuilder )
 
 void FLevelEditorContextMenuImpl::FillActorMenu( FMenuBuilder& MenuBuilder )
 {
+	struct Local
+	{
+		static FReply OnInteractiveActorPickerClicked()
+		{
+			FSlateApplication::Get().DismissAllMenus();
+			FLevelEditorActionCallbacks::AttachActorIteractive();
+			return FReply::Handled();
+		}
+	};
+
 	FSceneOutlinerInitializationOptions InitOptions;
 	{
 		InitOptions.Mode = ESceneOutlinerMode::ActorPicker;			
 		InitOptions.bShowHeaderRow = false;
 		InitOptions.bFocusSearchBoxWhenOpened = true;
 
-		struct Local
-		{
-			static bool IsAttachableActor( const AActor* const ParentActor )
-			{					
-				for ( FSelectionIterator It( GEditor->GetSelectedActorIterator() ) ; It ; ++It )
-				{
-					AActor* Actor = static_cast<AActor*>( *It );
-					if (!GEditor->CanParentActors(ParentActor, Actor))
-					{
-						return false;
-					}
-
-					USceneComponent* ChildRoot = Actor->GetRootComponent();
-					USceneComponent* ParentRoot = ParentActor->GetRootComponent();
-
-					if (ChildRoot != nullptr && ParentRoot != nullptr && ChildRoot->IsAttachedTo(ParentRoot))
-					{
-						return false;
-					}
-				}
-				return true;
-			}
-		};
-
 		// Only display Actors that we can attach too
-		InitOptions.Filters->AddFilterPredicate( SceneOutliner::FActorFilterPredicate::CreateStatic( &Local::IsAttachableActor ) );
+		InitOptions.Filters->AddFilterPredicate( SceneOutliner::FActorFilterPredicate::CreateStatic( &FLevelEditorActionCallbacks::IsAttachableActor) );
 	}		
 
+	if(SelectionInfo.bHaveAttachedActor)
+	{
+		MenuBuilder.AddMenuEntry( FLevelEditorCommands::Get().DetachFromParent, NAME_None, LOCTEXT( "None", "None" ) );
+	}
+
 	// Actor selector to allow the user to choose a parent actor
 	FSceneOutlinerModule& SceneOutlinerModule = FModuleManager::LoadModuleChecked<FSceneOutlinerModule>( "SceneOutliner" );
-	TSharedRef< SWidget > MiniSceneOutliner =
-		SNew( SVerticalBox )
-		+SVerticalBox::Slot()
-		.MaxHeight(400.0f)
-		[
-			SceneOutlinerModule.CreateSceneOutliner(
-			InitOptions,
-			FOnActorPicked::CreateStatic( &FLevelEditorActionCallbacks::AttachToActor ) )
-		];
 
-	// Give the scene outliner a border and background
-	const FSlateBrush* BackgroundBrush = FEditorStyle::GetBrush( "Menu.Background" );
-	TSharedRef< SBorder > RootBorder =
-		SNew( SBorder )
-		.Padding(3)
-		.BorderImage( BackgroundBrush )
-		.ForegroundColor( FEditorStyle::GetSlateColor("DefaultForeground") )
+	TSharedRef< SWidget > MenuWidget = 
+		SNew(SHorizontalBox)
 
-		// Assign the box panel as the child
+		+SHorizontalBox::Slot()
+		.AutoWidth()
 		[
-			SNew( SVerticalBox )
+			SNew(SVerticalBox)
 			+SVerticalBox::Slot()
-			.Padding( 5 )
-			.HAlign( HAlign_Center )
+			.MaxHeight(400.0f)
 			[
-				SNew( STextBlock )
-				.Text( LOCTEXT( "SelectMatineeActorToEdit", "Select a Matinee actor" ) )
+				SceneOutlinerModule.CreateSceneOutliner(
+					InitOptions,
+					FOnActorPicked::CreateStatic( &FLevelEditorActionCallbacks::AttachToActor )
+					)
 			]
+		]
+	
+		+SHorizontalBox::Slot()
+		.VAlign(VAlign_Top)
+		.AutoWidth()
+		[
+			SNew(SVerticalBox)
 
 			+SVerticalBox::Slot()
-			.Padding( 2 )
+			.AutoHeight()
+			.Padding(4.0f, 0.0f, 0.0f, 0.0f)
 			[
-				MiniSceneOutliner
+				SNew(SButton)
+				.ToolTipText( LOCTEXT( "PickButtonLabel", "Pick a parent actor to attach to").ToString() )
+				.ButtonStyle(FEditorStyle::Get(), "HoverHintOnly")
+				.OnClicked(FOnClicked::CreateStatic(&Local::OnInteractiveActorPickerClicked))
+				.ContentPadding(4.0f)
+				.ForegroundColor(FSlateColor::UseForeground())
+				.IsFocusable(false)
+				[
+					SNew(SImage)
+					.Image(FEditorStyle::GetBrush("PropertyWindow.Button_PickActorInteractive"))
+					.ColorAndOpacity(FSlateColor::UseForeground())
+				]
 			]
 		];
-	if(SelectionInfo.bHaveAttachedActor)
-	{
-		MenuBuilder.AddMenuEntry( FLevelEditorCommands::Get().DetachFromParent, NAME_None, LOCTEXT( "None", "None" ) );
-	}
 
-	MenuBuilder.AddWidget(MiniSceneOutliner, FText::GetEmpty(), false);
+	MenuBuilder.AddWidget(MenuWidget, FText::GetEmpty(), false);
 }
 
 void FLevelEditorContextMenuImpl::FillSnapAlignMenu( FMenuBuilder& MenuBuilder )
diff --git a/Engine/Source/Editor/LevelEditor/Public/LevelEditorActions.h b/Engine/Source/Editor/LevelEditor/Public/LevelEditorActions.h
--- a/Engine/Source/Editor/LevelEditor/Public/LevelEditorActions.h
+++ b/Engine/Source/Editor/LevelEditor/Public/LevelEditorActions.h
@@ -211,6 +211,8 @@ public:
 
 	TSharedPtr< FUICommandInfo > AttachSelectedActors;
 
+	TSharedPtr< FUICommandInfo > AttachActorIteractive;
+
 	TSharedPtr< FUICommandInfo > CreateNewOutlinerFolder;
 
 	TSharedPtr< FUICommandInfo > HoldToEnableVertexSnapping;
@@ -755,6 +757,12 @@ public:
 	/** Called when attach selected actors is pressed */
 	static void AttachSelectedActors();
 
+	/** Called when the actor picker needs to be used to select a new parent actor */
+	static void AttachActorIteractive();
+
+	/** @return true if the selected actor can be attached to the given parent actor */
+	static bool IsAttachableActor( const AActor* const ParentActor );
+
 	/** Called when create new outliner folder is clicked */
 	static void CreateNewOutlinerFolder_Clicked();
 
diff --git a/Engine/Source/Editor/PropertyEditor/Private/PropertyCustomizationHelpers.cpp b/Engine/Source/Editor/PropertyEditor/Private/PropertyCustomizationHelpers.cpp
--- a/Engine/Source/Editor/PropertyEditor/Private/PropertyCustomizationHelpers.cpp
+++ b/Engine/Source/Editor/PropertyEditor/Private/PropertyCustomizationHelpers.cpp
@@ -219,12 +219,13 @@ namespace PropertyCustomizationHelpers
 			.OnUseSelected(OnUseSelected);
 	}
 
-	TSharedRef<SWidget> MakeInteractiveActorPicker( FOnGetAllowedClasses OnGetAllowedClasses, FOnActorSelected OnActorSelectedFromPicker )
+	TSharedRef<SWidget> MakeInteractiveActorPicker( FOnGetAllowedClasses OnGetAllowedClasses, FOnShouldFilterActor OnShouldFilterActor, FOnActorSelected OnActorSelectedFromPicker )
 	{
 		return 
 			SNew( SPropertyEditorInteractiveActorPicker )
 			.ToolTipText( LOCTEXT( "PickButtonLabel", "Pick Actor from scene").ToString() )
 			.OnGetAllowedClasses( OnGetAllowedClasses )
+			.OnShouldFilterActor( OnShouldFilterActor )
 			.OnActorSelected( OnActorSelectedFromPicker );
 	}
 }
diff --git a/Engine/Source/Editor/PropertyEditor/Private/PropertyEditorHelpers.cpp b/Engine/Source/Editor/PropertyEditor/Private/PropertyEditorHelpers.cpp
--- a/Engine/Source/Editor/PropertyEditor/Private/PropertyEditorHelpers.cpp
+++ b/Engine/Source/Editor/PropertyEditor/Private/PropertyEditorHelpers.cpp
@@ -789,7 +789,7 @@ namespace PropertyEditorHelpers
 			break;
 
 		case EPropertyButton::PickActorInteractive:
-			NewButton = PropertyCustomizationHelpers::MakeInteractiveActorPicker( FOnGetAllowedClasses::CreateSP( PropertyEditor, &FPropertyEditor::OnGetClassesForAssetPicker ), FOnActorSelected::CreateSP( PropertyEditor, &FPropertyEditor::OnActorSelected ) );
+			NewButton = PropertyCustomizationHelpers::MakeInteractiveActorPicker( FOnGetAllowedClasses::CreateSP( PropertyEditor, &FPropertyEditor::OnGetClassesForAssetPicker ), FOnShouldFilterActor(), FOnActorSelected::CreateSP( PropertyEditor, &FPropertyEditor::OnActorSelected ) );
 			break;
 
 		default:
diff --git a/Engine/Source/Editor/PropertyEditor/Private/PropertyEditorModule.cpp b/Engine/Source/Editor/PropertyEditor/Private/PropertyEditorModule.cpp
--- a/Engine/Source/Editor/PropertyEditor/Private/PropertyEditorModule.cpp
+++ b/Engine/Source/Editor/PropertyEditor/Private/PropertyEditorModule.cpp
@@ -19,7 +19,6 @@
 #include "PropertyTable.h"
 #include "SPropertyTable.h"
 #include "TextPropertyTableCellPresenter.h"
-#include "SPropertyEditorInteractiveActorPicker.h"
 
 #include "PropertyTableConstants.h"
 
@@ -27,7 +26,6 @@ IMPLEMENT_MODULE( FPropertyEditorModule, PropertyEditor );
 
 void FPropertyEditorModule::StartupModule()
 {
-	SPropertyEditorInteractiveActorPicker::RegisterEditMode();
 }
 
 void FPropertyEditorModule::ShutdownModule()
@@ -40,8 +38,6 @@ void FPropertyEditorModule::ShutdownModule()
 
 	AllDetailViews.Empty();
 	AllSinglePropertyViews.Empty();
-
-	SPropertyEditorInteractiveActorPicker::UnregisterEditMode();
 }
 
 void FPropertyEditorModule::NotifyCustomizationModuleChanged()
diff --git a/Engine/Source/Editor/PropertyEditor/Private/UserInterface/PropertyEditor/SPropertyEditorAsset.cpp b/Engine/Source/Editor/PropertyEditor/Private/UserInterface/PropertyEditor/SPropertyEditorAsset.cpp
--- a/Engine/Source/Editor/PropertyEditor/Private/UserInterface/PropertyEditor/SPropertyEditorAsset.cpp
+++ b/Engine/Source/Editor/PropertyEditor/Private/UserInterface/PropertyEditor/SPropertyEditorAsset.cpp
@@ -201,7 +201,7 @@ void SPropertyEditorAsset::Construct( const FArguments& InArgs, const TSharedPtr
 		.Padding( 2.0f, 0.0f )
 		.VAlign(VAlign_Center)
 		[
-			PropertyCustomizationHelpers::MakeInteractiveActorPicker( FOnGetAllowedClasses::CreateSP(this, &SPropertyEditorAsset::OnGetAllowedClasses), FOnActorSelected::CreateSP( this, &SPropertyEditorAsset::OnActorSelected ) )
+			PropertyCustomizationHelpers::MakeInteractiveActorPicker( FOnGetAllowedClasses::CreateSP(this, &SPropertyEditorAsset::OnGetAllowedClasses), FOnShouldFilterActor(), FOnActorSelected::CreateSP( this, &SPropertyEditorAsset::OnActorSelected ) )
 		];
 	}
 
diff --git a/Engine/Source/Editor/PropertyEditor/Private/UserInterface/PropertyEditor/SPropertyEditorInteractiveActorPicker.cpp b/Engine/Source/Editor/PropertyEditor/Private/UserInterface/PropertyEditor/SPropertyEditorInteractiveActorPicker.cpp
--- a/Engine/Source/Editor/PropertyEditor/Private/UserInterface/PropertyEditor/SPropertyEditorInteractiveActorPicker.cpp
+++ b/Engine/Source/Editor/PropertyEditor/Private/UserInterface/PropertyEditor/SPropertyEditorInteractiveActorPicker.cpp
@@ -9,290 +9,19 @@
 
 #define LOCTEXT_NAMESPACE "PropertyPicker"
 
-namespace EPickState
-{
-	enum Type
-	{
-		NotOverViewport,
-		OverViewport,
-		OverIncompatibleActor,
-		OverActor,
-	};
-}
-
-/**
- * Editor mode used to interactively pick actors in the level editor viewports.
- */
-struct FEdModePicker : public FEdMode
-{
-	FEdModePicker()
-	{
-		ID = FBuiltinEditorModes::EM_ActorPicker;
-		PickState = EPickState::NotOverViewport;
-		HoveredActor.Reset();
-	}
-	
-	/** Begin FEdMode interface */
-	virtual void Tick(FLevelEditorViewportClient* ViewportClient, float DeltaTime) OVERRIDE
-	{
-		if(CursorDecoratorWindow.IsValid())
-		{
-			CursorDecoratorWindow->MoveWindowTo(FSlateApplication::Get().GetCursorPos() + FSlateApplication::Get().GetCursorSize());
-		}
-
-		FEdMode::Tick(ViewportClient, DeltaTime);
-	}
-
-	virtual bool MouseEnter(FLevelEditorViewportClient* ViewportClient, FViewport* Viewport,int32 x, int32 y) OVERRIDE
-	{
-		PickState = EPickState::OverViewport;
-		HoveredActor.Reset();
-		return FEdMode::MouseEnter(ViewportClient, Viewport, x, y);
-	}
-
-	virtual bool MouseLeave(FLevelEditorViewportClient* ViewportClient, FViewport* Viewport) OVERRIDE
-	{
-		PickState = EPickState::NotOverViewport;
-		HoveredActor.Reset();
-		return FEdMode::MouseLeave(ViewportClient, Viewport);
-	}
-
-	virtual bool MouseMove(FLevelEditorViewportClient* ViewportClient, FViewport* Viewport, int32 x, int32 y) OVERRIDE
-	{
-		PickState = EPickState::OverViewport;
-		HoveredActor.Reset();
-
-		int32 HitX = Viewport->GetMouseX();
-		int32 HitY = Viewport->GetMouseY();
-		HHitProxy* HitProxy = Viewport->GetHitProxy(HitX, HitY);
-		if (HitProxy != NULL && HitProxy->IsA(HActor::StaticGetType()))
-		{
-			HActor* ActorHit = static_cast<HActor*>(HitProxy);
-			if(ActorHit->Actor != NULL)
-			{
-				AActor* Actor = ActorHit->Actor;
-				HoveredActor = Actor;
-				PickState = EPickState::OverIncompatibleActor;
-
-				TArray<const UClass*> AllowedClasses;
-				OnGetAllowedClasses.Execute(AllowedClasses);
-				for(auto It(AllowedClasses.CreateConstIterator()); It; It++)
-				{
-					if(Actor->IsA(*It))
-					{
-						PickState = EPickState::OverActor;
-						break;
-					}
-				}
-			}
-		}
-
-		return true;
-	}
-
-	virtual bool InputKey(FLevelEditorViewportClient* ViewportClient, FViewport* Viewport, FKey Key, EInputEvent Event) OVERRIDE
-	{
-		if (Key == EKeys::LeftMouseButton && Event == IE_Pressed)
-		{
-			// See if we clicked on an actor
-			int32 HitX = Viewport->GetMouseX();
-			int32 HitY = Viewport->GetMouseY();
-			HHitProxy*	HitProxy = Viewport->GetHitProxy(HitX, HitY);
-			if (HitProxy != NULL && HitProxy->IsA(HActor::StaticGetType()))
-			{
-				HActor* ActorHit = static_cast<HActor*>(HitProxy);
-				if(ActorHit->Actor != NULL)
-				{
-					AActor* Actor = ActorHit->Actor;
-					TArray<const UClass*> AllowedClasses;
-					OnGetAllowedClasses.Execute(AllowedClasses);
-					for(auto It(AllowedClasses.CreateConstIterator()); It; It++)
-					{
-						if(Actor->IsA(*It))
-						{
-							OnActorSelected.ExecuteIfBound(ActorHit->Actor);
-							EndMode();
-							break;
-						}
-					}
-				}
-			}
-			return true;
-		}
-		else if(Key == EKeys::Escape && Event == IE_Pressed)
-		{
-			EndMode();
-			return true;
-		}
-		return false;
-	}
-
-	virtual bool GetCursor(EMouseCursor::Type& OutCursor) const OVERRIDE
-	{
-		if(HoveredActor.IsValid())
-		{
-			OutCursor = EMouseCursor::EyeDropper;
-		}
-		else
-		{
-			OutCursor = EMouseCursor::SlashedCircle;
-		}
-		
-		return true;
-	}
-
-	virtual bool UsesToolkits() const OVERRIDE
-	{
-		return false;
-	}
-	/** End FEdMode interface */
-
-	/** Start up the mode, reset state etc. */
-	void BeginMode(FOnGetAllowedClasses InOnGetAllowedClasses, FOnActorSelected InOnActorSelected )
-	{
-		OnActorSelected = InOnActorSelected;
-		OnGetAllowedClasses = InOnGetAllowedClasses;
-		GEditorModeTools().ActivateMode(GetID());
-		CursorDecoratorWindow = SWindow::MakeCursorDecorator();
-		FSlateApplication::Get().AddWindow(CursorDecoratorWindow.ToSharedRef(), true);
-		CursorDecoratorWindow->SetContent(
-			SNew(SToolTip)
-			.Text(this, &FEdModePicker::GetCursorDecoratorText)
-			);
-
-		HoveredActor.Reset();
-		PickState = EPickState::NotOverViewport;
-	}
-
-	/** End the mode */
-	void EndMode()
-	{
-		OnActorSelected = FOnActorSelected();
-		OnGetAllowedClasses = FOnGetAllowedClasses();
-		GEditorModeTools().DeactivateMode(GetID());
-
-		if (CursorDecoratorWindow.IsValid())
-		{
-			CursorDecoratorWindow->RequestDestroyWindow();
-			CursorDecoratorWindow.Reset();
-		}
-
-		HoveredActor.Reset();
-		PickState = EPickState::NotOverViewport;
-	}
-
-	/** Delegate used to display information about picking near the cursor */
-	FString GetCursorDecoratorText() const
-	{
-		switch(PickState)
-		{
-		default:
-		case EPickState::NotOverViewport:
-			return LOCTEXT("PickActor_NotOverViewport", "Pick an actor by clicking on it in a level viewport").ToString();
-		case EPickState::OverViewport:
-			return LOCTEXT("PickActor_NotOverActor", "Pick an actor by clicking on it").ToString();
-		case EPickState::OverIncompatibleActor:
-			{
-				if(HoveredActor.IsValid())
-				{
-					FFormatNamedArguments Arguments;
-					Arguments.Add(TEXT("Actor"), FText::FromString(HoveredActor.Get()->GetName()));
-					return FText::Format(LOCTEXT("PickActor_OverIncompatibleActor", "{Actor} is incompatible"), Arguments).ToString();
-				}
-				else
-				{
-					return LOCTEXT("PickActor_NotOverActor", "Pick an actor by clicking on it").ToString();
-				}
-			}
-		case EPickState::OverActor:
-			{
-				if(HoveredActor.IsValid())
-				{
-					FFormatNamedArguments Arguments;
-					Arguments.Add(TEXT("Actor"), FText::FromString(HoveredActor.Get()->GetName()));
-					return FText::Format(LOCTEXT("PickActor_OverActor", "Pick {Actor}"), Arguments).ToString();
-				}
-				else
-				{
-					return LOCTEXT("PickActor_NotOverActor", "Pick an actor by clicking on it").ToString();
-				}
-			}
-		}
-	}
-
-	/** Flag for display state */
-	TWeakObjectPtr<AActor> HoveredActor;
-
-	/** Flag for display state */
-	EPickState::Type PickState;
-
-	/** The window that owns the decorator widget */
-	TSharedPtr<SWindow> CursorDecoratorWindow;
-
-	/** Delegates used to pick actors */
-	FOnActorSelected OnActorSelected;
-	FOnGetAllowedClasses OnGetAllowedClasses;
-};
-
-/** The one and only picker edit mode */
-static TSharedPtr<FEdModePicker> EdModePicker = NULL;
-
-void SPropertyEditorInteractiveActorPicker::RegisterEditMode()
-{	
-	EdModePicker = MakeShareable(new FEdModePicker);
-	GEditorModeTools().RegisterMode(EdModePicker.ToSharedRef());
-
-	// we want to be able to perform this action with all the built-in editor modes
-	GEditorModeTools().RegisterCompatibleModes(EdModePicker->GetID(), FBuiltinEditorModes::EM_Default);
-	GEditorModeTools().RegisterCompatibleModes(EdModePicker->GetID(), FBuiltinEditorModes::EM_Placement);
-	GEditorModeTools().RegisterCompatibleModes(EdModePicker->GetID(), FBuiltinEditorModes::EM_Bsp);
-	GEditorModeTools().RegisterCompatibleModes(EdModePicker->GetID(), FBuiltinEditorModes::EM_Geometry);
-	GEditorModeTools().RegisterCompatibleModes(EdModePicker->GetID(), FBuiltinEditorModes::EM_InterpEdit);
-	GEditorModeTools().RegisterCompatibleModes(EdModePicker->GetID(), FBuiltinEditorModes::EM_Texture);
-	GEditorModeTools().RegisterCompatibleModes(EdModePicker->GetID(), FBuiltinEditorModes::EM_MeshPaint);
-	GEditorModeTools().RegisterCompatibleModes(EdModePicker->GetID(), FBuiltinEditorModes::EM_Landscape);
-	GEditorModeTools().RegisterCompatibleModes(EdModePicker->GetID(), FBuiltinEditorModes::EM_Foliage);
-	GEditorModeTools().RegisterCompatibleModes(EdModePicker->GetID(), FBuiltinEditorModes::EM_Level);
-	GEditorModeTools().RegisterCompatibleModes(EdModePicker->GetID(), FBuiltinEditorModes::EM_StreamingLevel);
-	GEditorModeTools().RegisterCompatibleModes(EdModePicker->GetID(), FBuiltinEditorModes::EM_Physics);
-}
-
-void SPropertyEditorInteractiveActorPicker::UnregisterEditMode()
-{		
-	if(GEditorModeTools().IsModeActive(EdModePicker->GetID()))
-	{
-		EdModePicker->EndMode();
-	}
-
-	GEditorModeTools().UnregisterCompatibleModes(EdModePicker->GetID(), FBuiltinEditorModes::EM_Default);
-	GEditorModeTools().UnregisterCompatibleModes(EdModePicker->GetID(), FBuiltinEditorModes::EM_Placement);
-	GEditorModeTools().UnregisterCompatibleModes(EdModePicker->GetID(), FBuiltinEditorModes::EM_Bsp);
-	GEditorModeTools().UnregisterCompatibleModes(EdModePicker->GetID(), FBuiltinEditorModes::EM_Geometry);
-	GEditorModeTools().UnregisterCompatibleModes(EdModePicker->GetID(), FBuiltinEditorModes::EM_InterpEdit);
-	GEditorModeTools().UnregisterCompatibleModes(EdModePicker->GetID(), FBuiltinEditorModes::EM_Texture);
-	GEditorModeTools().UnregisterCompatibleModes(EdModePicker->GetID(), FBuiltinEditorModes::EM_MeshPaint);
-	GEditorModeTools().UnregisterCompatibleModes(EdModePicker->GetID(), FBuiltinEditorModes::EM_Landscape);
-	GEditorModeTools().UnregisterCompatibleModes(EdModePicker->GetID(), FBuiltinEditorModes::EM_Foliage);
-	GEditorModeTools().UnregisterCompatibleModes(EdModePicker->GetID(), FBuiltinEditorModes::EM_Level);
-	GEditorModeTools().UnregisterCompatibleModes(EdModePicker->GetID(), FBuiltinEditorModes::EM_StreamingLevel);
-	GEditorModeTools().UnregisterCompatibleModes(EdModePicker->GetID(), FBuiltinEditorModes::EM_Physics);
-	GEditorModeTools().UnregisterMode(EdModePicker.ToSharedRef());
-	EdModePicker = NULL;
-}
-
 SPropertyEditorInteractiveActorPicker::~SPropertyEditorInteractiveActorPicker()
 {
+	FActorPickerModeModule& ActorPickerMode = FModuleManager::Get().GetModuleChecked<FActorPickerModeModule>("ActorPickerMode");
+
 	// make sure we are unregistered when this widget goes away
-	if(GEditorModeTools().IsModeActive(EdModePicker->GetID()))
-	{
-		EdModePicker->EndMode();
-	}	
+	ActorPickerMode.EndActorPickingMode();
 }
 
 void SPropertyEditorInteractiveActorPicker::Construct( const FArguments& InArgs )
 {
 	OnActorSelected = InArgs._OnActorSelected;
 	OnGetAllowedClasses = InArgs._OnGetAllowedClasses;
+	OnShouldFilterActor = InArgs._OnShouldFilterActor;
 
 	SButton::Construct(
 		SButton::FArguments()
@@ -313,9 +42,11 @@ FReply SPropertyEditorInteractiveActorPicker::OnKeyDown( const FGeometry& MyGeom
 {
 	if(InKeyboardEvent.GetKey() == EKeys::Escape)
 	{
-		if(GEditorModeTools().IsModeActive(EdModePicker->GetID()))
+		FActorPickerModeModule& ActorPickerMode = FModuleManager::Get().GetModuleChecked<FActorPickerModeModule>("ActorPickerMode");
+
+		if(ActorPickerMode.IsInActorPickingMode())
 		{
-			EdModePicker->EndMode();
+			ActorPickerMode.EndActorPickingMode();
 			return FReply::Handled();
 		}
 	}
@@ -330,16 +61,18 @@ bool SPropertyEditorInteractiveActorPicker::SupportsKeyboardFocus() const
 
 FReply SPropertyEditorInteractiveActorPicker::OnClicked()
 {
-	if(GEditorModeTools().IsModeActive(EdModePicker->GetID()))
+	FActorPickerModeModule& ActorPickerMode = FModuleManager::Get().GetModuleChecked<FActorPickerModeModule>("ActorPickerMode");
+
+	if(ActorPickerMode.IsInActorPickingMode())
 	{
-		EdModePicker->EndMode();
+		ActorPickerMode.EndActorPickingMode();
 	}
 	else
 	{
-		EdModePicker->BeginMode(OnGetAllowedClasses, OnActorSelected);
+		ActorPickerMode.BeginActorPickingMode(OnGetAllowedClasses, OnShouldFilterActor, OnActorSelected);
 	}
 
 	return FReply::Handled();
 }
 
-#undef LOCTEXT_NAMESPACE
\ No newline at end of file
+#undef LOCTEXT_NAMESPACE
diff --git a/Engine/Source/Editor/PropertyEditor/Private/UserInterface/PropertyEditor/SPropertyEditorInteractiveActorPicker.h b/Engine/Source/Editor/PropertyEditor/Private/UserInterface/PropertyEditor/SPropertyEditorInteractiveActorPicker.h
--- a/Engine/Source/Editor/PropertyEditor/Private/UserInterface/PropertyEditor/SPropertyEditorInteractiveActorPicker.h
+++ b/Engine/Source/Editor/PropertyEditor/Private/UserInterface/PropertyEditor/SPropertyEditorInteractiveActorPicker.h
@@ -9,6 +9,9 @@ public:
 	/** Delegate used to filter allowed actors */
 	SLATE_EVENT( FOnGetAllowedClasses, OnGetAllowedClasses )
 
+	/** Delegate used to filter allowed actors */
+	SLATE_EVENT( FOnShouldFilterActor, OnShouldFilterActor )
+
 	/** Delegate called when an actor is selected */
 	SLATE_EVENT( FOnActorSelected, OnActorSelected )
 
@@ -23,12 +26,6 @@ public:
 	virtual bool SupportsKeyboardFocus() const OVERRIDE;
 	/** End SWidget interface */
 
-	/** Register the picker editor mode */
-	static void RegisterEditMode();
-
-	/** Unregister the picker editor mode */
-	static void UnregisterEditMode();
-
 private:
 	/** Delegate for when the button is clicked */
 	FReply OnClicked();
@@ -36,6 +33,9 @@ private:
 	/** Delegate used to filter allowed actors */
 	FOnGetAllowedClasses OnGetAllowedClasses;
 
+	/** Delegate used to filter allowed actors */
+	FOnShouldFilterActor OnShouldFilterActor;
+
 	/** Delegate called when an actor is selected */
 	FOnActorSelected OnActorSelected;
 };
\ No newline at end of file
diff --git a/Engine/Source/Editor/PropertyEditor/PropertyEditor.Build.cs b/Engine/Source/Editor/PropertyEditor/PropertyEditor.Build.cs
--- a/Engine/Source/Editor/PropertyEditor/PropertyEditor.Build.cs
+++ b/Engine/Source/Editor/PropertyEditor/PropertyEditor.Build.cs
@@ -8,7 +8,8 @@ public class PropertyEditor : ModuleRules
 	{
 		PublicDependencyModuleNames.AddRange(
 			new string[] {
-				"UnrealEd",		
+				"UnrealEd",
+                "ActorPickerMode",
 			}
 		);
 
diff --git a/Engine/Source/Editor/PropertyEditor/Public/PropertyCustomizationHelpers.h b/Engine/Source/Editor/PropertyEditor/Public/PropertyCustomizationHelpers.h
--- a/Engine/Source/Editor/PropertyEditor/Public/PropertyCustomizationHelpers.h
+++ b/Engine/Source/Editor/PropertyEditor/Public/PropertyCustomizationHelpers.h
@@ -5,13 +5,12 @@
 #include "IFilter.h"
 #include "FilterCollection.h"
 #include "SResetToDefaultMenu.h"
+#include "ActorPickerMode.h"
 
 namespace SceneOutliner { struct FOutlinerFilters; }
 
 DECLARE_DELEGATE_OneParam(FOnAssetSelected, const class FAssetData& /*AssetData*/);
-DECLARE_DELEGATE_OneParam( FOnGetAllowedClasses, TArray<const UClass*>& );
 DECLARE_DELEGATE_RetVal_OneParam(bool, FOnShouldFilterAsset, const class FAssetData& /*AssetData*/);
-DECLARE_DELEGATE_OneParam( FOnActorSelected, AActor* );
 DECLARE_DELEGATE_OneParam( FOnGetActorFilters, TSharedPtr<SceneOutliner::FOutlinerFilters>& );
 
 namespace PropertyCustomizationHelpers
@@ -28,7 +27,7 @@ namespace PropertyCustomizationHelpers
 	PROPERTYEDITOR_API TSharedRef<SWidget> MakeAssetPickerWithMenu( UObject* const InitialObject, const bool AllowClear, const TArray<const UClass*>* const AllowedClasses, FOnShouldFilterAsset OnShouldFilterAsset, FOnAssetSelected OnSet, FSimpleDelegate OnClose );
 	PROPERTYEDITOR_API TSharedRef<SWidget> MakeActorPickerAnchorButton( FOnGetActorFilters OnGetActorFilters, FOnActorSelected OnActorSelectedFromPicker );
 	PROPERTYEDITOR_API TSharedRef<SWidget> MakeActorPickerWithMenu( AActor* const InitialActor, const bool AllowClear, const TSharedPtr< SceneOutliner::FOutlinerFilters >& ActorFilters, FOnActorSelected OnSet, FSimpleDelegate OnClose, FSimpleDelegate OnUseSelected );
-	PROPERTYEDITOR_API TSharedRef<SWidget> MakeInteractiveActorPicker( FOnGetAllowedClasses OnClearClicked, FOnActorSelected OnActorSelectedFromPicker );
+	PROPERTYEDITOR_API TSharedRef<SWidget> MakeInteractiveActorPicker( FOnGetAllowedClasses OnGetAllowedClasses, FOnShouldFilterActor OnShouldFilterActor, FOnActorSelected OnActorSelectedFromPicker );
 }
 
 /** Delegate used to get a generic object */
diff --git a/Engine/Source/Editor/UnrealEd/Private/EditorModes.cpp b/Engine/Source/Editor/UnrealEd/Private/EditorModes.cpp
--- a/Engine/Source/Editor/UnrealEd/Private/EditorModes.cpp
+++ b/Engine/Source/Editor/UnrealEd/Private/EditorModes.cpp
@@ -10,6 +10,7 @@
 #include "Editor/BspMode/Public/IBspModeModule.h"
 #include "Editor/MeshPaint/Public/MeshPaintModule.h"
 #include "Editor/GeometryMode/Public/GeometryEdMode.h"
+#include "Editor/ActorPickerMode/Public/ActorPickerMode.h"
 #include "Editor/TextureAlignMode/Public/TextureAlignEdMode.h"
 #include "LevelEditor.h"
 #include "Toolkits/ToolkitManager.h"
@@ -234,6 +235,25 @@ bool FEdMode::MouseMove(FLevelEditorViewportClient* ViewportClient,FViewport* Vi
 	return false;
 }
 
+bool FEdMode::ReceivedFocus(FLevelEditorViewportClient* ViewportClient,FViewport* Viewport)
+{
+	if( GetCurrentTool() )
+	{
+		return GetCurrentTool()->ReceivedFocus( ViewportClient, Viewport );
+	}
+
+	return false;
+}
+
+bool FEdMode::LostFocus(FLevelEditorViewportClient* ViewportClient,FViewport* Viewport)
+{
+	if( GetCurrentTool() )
+	{
+		return GetCurrentTool()->LostFocus( ViewportClient, Viewport );
+	}
+
+	return false;
+}
 
 /**
  * Called when the mouse is moved while a window input capture is in effect
@@ -1204,6 +1224,7 @@ void FEditorModeTools::Init()
 	FModuleManager::LoadModuleChecked<IBspModeModule>(TEXT("BspMode"));
 	FModuleManager::LoadModuleChecked<FTextureAlignModeModule>(TEXT("TextureAlignMode"));
 	FModuleManager::LoadModuleChecked<FGeometryModeModule>(TEXT("GeometryMode"));
+	FModuleManager::LoadModuleChecked<FActorPickerModeModule>(TEXT("ActorPickerMode"));
 	FModuleManager::LoadModuleChecked<IMeshPaintModule>(TEXT("MeshPaint"));
 	FModuleManager::LoadModuleChecked<ILandscapeEditorModule>(TEXT("LandscapeEditor"));
 	FModuleManager::LoadModuleChecked<IFoliageEditModule>(TEXT("FoliageEdit"));
@@ -1817,6 +1838,26 @@ bool FEditorModeTools::MouseMove( FLevelEditorViewportClient* InViewportClient,
 	return bHandled;
 }
 
+bool FEditorModeTools::ReceivedFocus( FLevelEditorViewportClient* InViewportClient, FViewport* Viewport )
+{
+	bool bHandled = false;
+	for( int32 ModeIndex = 0; ModeIndex < ActiveModes.Num(); ++ModeIndex )
+	{
+		bHandled |= ActiveModes[ModeIndex]->ReceivedFocus( InViewportClient, Viewport );
+	}
+	return bHandled;
+}
+
+bool FEditorModeTools::LostFocus( FLevelEditorViewportClient* InViewportClient, FViewport* Viewport )
+{
+	bool bHandled = false;
+	for( int32 ModeIndex = 0; ModeIndex < ActiveModes.Num(); ++ModeIndex )
+	{
+		bHandled |= ActiveModes[ModeIndex]->LostFocus( InViewportClient, Viewport );
+	}
+	return bHandled;
+}
+
 /** Draws all active mode components */	
 void FEditorModeTools::DrawActiveModes( const FSceneView* InView, FPrimitiveDrawInterface* PDI )
 {
diff --git a/Engine/Source/Editor/UnrealEd/Private/LevelEditorViewport.cpp b/Engine/Source/Editor/UnrealEd/Private/LevelEditorViewport.cpp
--- a/Engine/Source/Editor/UnrealEd/Private/LevelEditorViewport.cpp
+++ b/Engine/Source/Editor/UnrealEd/Private/LevelEditorViewport.cpp
@@ -4460,6 +4460,18 @@ void FLevelEditorViewportClient::SetCameraSpeedSetting(int32 SpeedSetting)
 	GetMutableDefault<ULevelEditorViewportSettings>()->CameraSpeed = SpeedSetting;
 }
 
+void FLevelEditorViewportClient::ReceivedFocus(FViewport* Viewport)
+{
+	FEditorViewportClient::ReceivedFocus(Viewport);
+	GEditorModeTools().ReceivedFocus(this, Viewport);
+}
+
+void FLevelEditorViewportClient::LostFocus(FViewport* Viewport)
+{
+	FEditorViewportClient::LostFocus(Viewport);
+	GEditorModeTools().LostFocus(this, Viewport);
+}
+
 bool FLevelEditorViewportClient::OverrideHighResScreenshotCaptureRegion(FIntRect& OutCaptureRegion)
 {
 	FSlateRect Rect;
diff --git a/Engine/Source/Editor/UnrealEd/Public/EditorModeTools.h b/Engine/Source/Editor/UnrealEd/Public/EditorModeTools.h
--- a/Engine/Source/Editor/UnrealEd/Public/EditorModeTools.h
+++ b/Engine/Source/Editor/UnrealEd/Public/EditorModeTools.h
@@ -50,6 +50,10 @@ public:
 
 	virtual bool MouseMove(FLevelEditorViewportClient* ViewportClient,FViewport* Viewport,int32 x, int32 y) { return false; }
 
+	virtual bool ReceivedFocus(FLevelEditorViewportClient* ViewportClient,FViewport* Viewport) { return false; }
+
+	virtual bool LostFocus(FLevelEditorViewportClient* ViewportClient,FViewport* Viewport) { return false; }
+
 	/**
 	 * Called when the mouse is moved while a window input capture is in effect
 	 *
diff --git a/Engine/Source/Editor/UnrealEd/Public/EditorModes.h b/Engine/Source/Editor/UnrealEd/Public/EditorModes.h
--- a/Engine/Source/Editor/UnrealEd/Public/EditorModes.h
+++ b/Engine/Source/Editor/UnrealEd/Public/EditorModes.h
@@ -106,6 +106,10 @@ public:
 
 	virtual bool MouseMove(FLevelEditorViewportClient* ViewportClient,FViewport* Viewport,int32 x, int32 y);
 
+	virtual bool ReceivedFocus(FLevelEditorViewportClient* ViewportClient,FViewport* Viewport);
+
+	virtual bool LostFocus(FLevelEditorViewportClient* ViewportClient,FViewport* Viewport);
+
 	/**
 	 * Called when the mouse is moved while a window input capture is in effect
 	 *
@@ -603,6 +607,12 @@ public:
 	/** Notifies all active modes that the mouse has moved */
 	bool MouseMove( FLevelEditorViewportClient* InViewportClient, FViewport* Viewport, int32 X, int32 Y );
 
+	/** Notifies all active modes that a viewport has received focus */
+	bool ReceivedFocus( FLevelEditorViewportClient* InViewportClient, FViewport* Viewport );
+
+	/** Notifies all active modes that a viewport has lost focus */
+	bool LostFocus( FLevelEditorViewportClient* InViewportClient, FViewport* Viewport );
+
 	/** Draws all active modes */	
 	void DrawActiveModes( const FSceneView* InView, FPrimitiveDrawInterface* PDI );
 
diff --git a/Engine/Source/Editor/UnrealEd/Public/LevelEditorViewport.h b/Engine/Source/Editor/UnrealEd/Public/LevelEditorViewport.h
--- a/Engine/Source/Editor/UnrealEd/Public/LevelEditorViewport.h
+++ b/Engine/Source/Editor/UnrealEd/Public/LevelEditorViewport.h
@@ -173,6 +173,8 @@ public:
 	virtual bool IsAspectRatioConstrained() const OVERRIDE;
 	virtual int32 GetCameraSpeedSetting() const OVERRIDE;
 	virtual void SetCameraSpeedSetting(int32 SpeedSetting) OVERRIDE;
+	virtual void ReceivedFocus(FViewport* Viewport) OVERRIDE;
+	virtual void LostFocus(FViewport* Viewport) OVERRIDE;
 
 	virtual bool OverrideHighResScreenshotCaptureRegion(FIntRect& OutCaptureRegion) OVERRIDE;
 
