SLATE - TAttribute<FText> Tooltips are now settable post construct on SWidgets.

diff --git a/Engine/Source/Developer/TaskGraph/Private/SGraphBar.cpp b/Engine/Source/Developer/TaskGraph/Private/SGraphBar.cpp
--- a/Engine/Source/Developer/TaskGraph/Private/SGraphBar.cpp
+++ b/Engine/Source/Developer/TaskGraph/Private/SGraphBar.cpp
@@ -200,7 +200,7 @@ FReply SGraphBar::OnMouseMove( const FGeometry& MyGeometry, const FPointerEvent&
 	{
 		if( HoveredEventIndex != INDEX_NONE )
 		{
-			this->SetToolTipText( Events[ HoveredEventIndex ]->EventName );
+			this->SetToolTipText(TAttribute<FString>(Events[HoveredEventIndex]->EventName));
 		}
 		else
 		{
diff --git a/Engine/Source/Editor/ClassViewer/Private/SClassViewer.cpp b/Engine/Source/Editor/ClassViewer/Private/SClassViewer.cpp
--- a/Engine/Source/Editor/ClassViewer/Private/SClassViewer.cpp
+++ b/Engine/Source/Editor/ClassViewer/Private/SClassViewer.cpp
@@ -2336,7 +2336,7 @@ TSharedRef< ITableRow > SClassViewer::OnGenerateRowForClassViewer( TSharedPtr<FC
 
 	if ( !Item->GeneratedClassPackage.IsEmpty() )
 	{
-		ReturnRow->SetToolTipText(Item->GeneratedClassPackage);
+		ReturnRow->SetToolTipText(TAttribute<FString>(Item->GeneratedClassPackage));
 	}
 
 	// Expand the item if needed.
diff --git a/Engine/Source/Editor/MaterialEditor/Private/MaterialEditorDetailCustomization.cpp b/Engine/Source/Editor/MaterialEditor/Private/MaterialEditorDetailCustomization.cpp
--- a/Engine/Source/Editor/MaterialEditor/Private/MaterialEditorDetailCustomization.cpp
+++ b/Engine/Source/Editor/MaterialEditor/Private/MaterialEditorDetailCustomization.cpp
@@ -237,7 +237,7 @@ void FMaterialExpressionCollectionParameterDetails::CustomizeDetails( IDetailLay
 	ParameterComboButton = NewComboButton;
 	ParameterListView = NewListView;
 
-	NewComboButton->SetToolTipText(GetToolTipText());
+	NewComboButton->SetToolTipText(TAttribute<FString>(GetToolTipText()));
 }
 
 void FMaterialExpressionCollectionParameterDetails::PopulateParameters()
diff --git a/Engine/Source/Editor/PropertyEditor/Private/PropertyEditorHelpers.cpp b/Engine/Source/Editor/PropertyEditor/Private/PropertyEditorHelpers.cpp
--- a/Engine/Source/Editor/PropertyEditor/Private/PropertyEditorHelpers.cpp
+++ b/Engine/Source/Editor/PropertyEditor/Private/PropertyEditorHelpers.cpp
@@ -76,7 +76,7 @@ void SPropertyValueWidget::Construct( const FArguments& InArgs, TSharedPtr<FProp
 
 	ValueEditorWidget = ConstructPropertyEditorWidget( PropertyEditor, InPropertyUtilities );
 
-	ValueEditorWidget->SetToolTipText( PropertyEditor->GetToolTipText() );
+	ValueEditorWidget->SetToolTipText( TAttribute<FString>(PropertyEditor->GetToolTipText()) );
 
 
 	if( InArgs._ShowPropertyButtons )
diff --git a/Engine/Source/Editor/PropertyEditor/Private/UserInterface/PropertyEditor/SPropertyComboBox.cpp b/Engine/Source/Editor/PropertyEditor/Private/UserInterface/PropertyEditor/SPropertyComboBox.cpp
--- a/Engine/Source/Editor/PropertyEditor/Private/UserInterface/PropertyEditor/SPropertyComboBox.cpp
+++ b/Engine/Source/Editor/PropertyEditor/Private/UserInterface/PropertyEditor/SPropertyComboBox.cpp
@@ -20,7 +20,7 @@ void SPropertyComboBox::Construct( const FArguments& InArgs )
 	{
 		if(*ComboItemList[ItemIndex].Get() == VisibleText)
 		{
-			SetToolTipText(*ToolTipList[ItemIndex]);
+			SetToolTipText(TAttribute<FString>(*ToolTipList[ItemIndex]));
 			InitiallySelectedItem = ComboItemList[ItemIndex];
 			break;
 		}
@@ -56,7 +56,7 @@ void SPropertyComboBox::SetSelectedItem( const FString& InSelectedItem )
 	{
 		if(*ComboItemList[ItemIndex].Get() == InSelectedItem)
 		{
-			SetToolTipText(*ToolTipList[ItemIndex]);
+			SetToolTipText(TAttribute<FString>(*ToolTipList[ItemIndex]));
 
 			SComboBox< TSharedPtr<FString> >::SetSelectedItem(ComboItemList[ItemIndex]);
 			return;
diff --git a/Engine/Source/Editor/PropertyEditor/Private/UserInterface/PropertyEditor/SPropertyEditorCombo.cpp b/Engine/Source/Editor/PropertyEditor/Private/UserInterface/PropertyEditor/SPropertyEditorCombo.cpp
--- a/Engine/Source/Editor/PropertyEditor/Private/UserInterface/PropertyEditor/SPropertyEditorCombo.cpp
+++ b/Engine/Source/Editor/PropertyEditor/Private/UserInterface/PropertyEditor/SPropertyEditorCombo.cpp
@@ -158,7 +158,7 @@ void SPropertyEditorCombo::SendToObjects( const FString& NewValue )
 		{
 			ToolTipText = FString::Printf(TEXT("%s\n\n%s"), *ToolTipText, *ToolTipValue);
 		}
-		SetToolTipText(ToolTipText);
+		SetToolTipText(TAttribute<FString>(ToolTipText));
 	}
 	else
 	{
diff --git a/Engine/Source/Editor/PropertyEditor/Private/UserInterface/PropertyEditor/SPropertyEditorTableRow.cpp b/Engine/Source/Editor/PropertyEditor/Private/UserInterface/PropertyEditor/SPropertyEditorTableRow.cpp
--- a/Engine/Source/Editor/PropertyEditor/Private/UserInterface/PropertyEditor/SPropertyEditorTableRow.cpp
+++ b/Engine/Source/Editor/PropertyEditor/Private/UserInterface/PropertyEditor/SPropertyEditorTableRow.cpp
@@ -34,7 +34,7 @@ void SPropertyEditorTableRow::Construct( const FArguments& InArgs, const TShared
 
 	PropertyPath = FPropertyNode::CreatePropertyPath( PropertyEditor->GetPropertyNode() );
 
-	this->SetToolTipText( PropertyEditor->GetToolTipText() );
+	this->SetToolTipText(TAttribute<FString>(PropertyEditor->GetToolTipText()));
 
 	SMultiColumnTableRow< TSharedPtr<FPropertyNode*> >::Construct( FSuperRowType::FArguments(), InOwnerTable );
 }
@@ -289,7 +289,7 @@ TSharedRef<SWidget> SPropertyEditorTableRow::ConstructPropertyEditorWidget()
 		PropertyWidget = SNew( SPropertyEditor, PropertyEditorRef );
 	}
 
-	PropertyWidget->SetToolTipText( PropertyEditor->GetToolTipText() );
+	PropertyWidget->SetToolTipText( TAttribute<FString>(PropertyEditor->GetToolTipText()) );
 
 	return PropertyWidget.ToSharedRef();
 }
\ No newline at end of file
diff --git a/Engine/Source/Editor/PropertyEditor/Private/UserInterface/PropertyTable/TextPropertyTableCellPresenter.cpp b/Engine/Source/Editor/PropertyEditor/Private/UserInterface/PropertyTable/TextPropertyTableCellPresenter.cpp
--- a/Engine/Source/Editor/PropertyEditor/Private/UserInterface/PropertyTable/TextPropertyTableCellPresenter.cpp
+++ b/Engine/Source/Editor/PropertyEditor/Private/UserInterface/PropertyTable/TextPropertyTableCellPresenter.cpp
@@ -156,7 +156,7 @@ TSharedRef<SWidget> FTextPropertyTableCellPresenter::ConstructEditModeCellWidget
 		HasReadOnlyEditingWidget = true;
 	}
 
-	PropertyWidget->SetToolTipText( PropertyEditor->GetToolTipText() );
+	PropertyWidget->SetToolTipText( TAttribute<FString>(PropertyEditor->GetToolTipText()) );
 
 	return PropertyWidget.ToSharedRef();
 }
diff --git a/Engine/Source/Runtime/Slate/Private/Framework/Application/SlateApplication.cpp b/Engine/Source/Runtime/Slate/Private/Framework/Application/SlateApplication.cpp
--- a/Engine/Source/Runtime/Slate/Private/Framework/Application/SlateApplication.cpp
+++ b/Engine/Source/Runtime/Slate/Private/Framework/Application/SlateApplication.cpp
@@ -3030,6 +3030,13 @@ TSharedRef<IToolTip> FSlateApplication::MakeToolTip( const TAttribute<FString>&
 }
 
 
+TSharedRef<IToolTip> FSlateApplication::MakeToolTip(const TAttribute<FText>& ToolTipText)
+{
+	return SNew(SToolTip)
+		.Text(ToolTipText);
+}
+
+
 TSharedRef<IToolTip> FSlateApplication::MakeToolTip( const FText& ToolTipText )
 {
 	return SNew(SToolTip)
diff --git a/Engine/Source/Runtime/Slate/Private/Widgets/Notifications/SErrorHint.cpp b/Engine/Source/Runtime/Slate/Private/Widgets/Notifications/SErrorHint.cpp
--- a/Engine/Source/Runtime/Slate/Private/Widgets/Notifications/SErrorHint.cpp
+++ b/Engine/Source/Runtime/Slate/Private/Widgets/Notifications/SErrorHint.cpp
@@ -33,7 +33,7 @@ void SErrorHint::SetError( const FText& InErrorText )
 	}
 
 	ErrorText = InErrorText;
-	SetToolTipText( ErrorText.ToString() );
+	SetToolTipText( ErrorText );
 }
 
 void SErrorHint::SetError( const FString& InErrorText )
diff --git a/Engine/Source/Runtime/Slate/Public/Framework/Application/SlateApplication.h b/Engine/Source/Runtime/Slate/Public/Framework/Application/SlateApplication.h
--- a/Engine/Source/Runtime/Slate/Public/Framework/Application/SlateApplication.h
+++ b/Engine/Source/Runtime/Slate/Public/Framework/Application/SlateApplication.h
@@ -1015,6 +1015,8 @@ public:
 
 	virtual TSharedRef<IToolTip> MakeToolTip( const TAttribute<FString>& ToolTipString ) OVERRIDE;
 
+	virtual TSharedRef<IToolTip> MakeToolTip( const TAttribute<FText>& ToolTipText ) OVERRIDE;
+
 	virtual TSharedRef<IToolTip> MakeToolTip( const FText& ToolTipText ) OVERRIDE;
 
 	virtual void RequestDestroyWindow( TSharedRef<SWindow> WindowToDestroy ) OVERRIDE;
diff --git a/Engine/Source/Runtime/SlateCore/Private/Widgets/SWidget.cpp b/Engine/Source/Runtime/SlateCore/Private/Widgets/SWidget.cpp
--- a/Engine/Source/Runtime/SlateCore/Private/Widgets/SWidget.cpp
+++ b/Engine/Source/Runtime/SlateCore/Private/Widgets/SWidget.cpp
@@ -435,6 +435,12 @@ void SWidget::SetToolTipText( const TAttribute<FString>& ToolTipString )
 }
 
 
+void SWidget::SetToolTipText(const TAttribute<FText>& ToolTipText)
+{
+	ToolTip = FSlateApplicationBase::Get().MakeToolTip(ToolTipText);
+}
+
+
 void SWidget::SetToolTipText( const FText& ToolTipText )
 {
 	ToolTip = FSlateApplicationBase::Get().MakeToolTip(ToolTipText);
diff --git a/Engine/Source/Runtime/SlateCore/Public/Application/SlateApplicationBase.h b/Engine/Source/Runtime/SlateCore/Public/Application/SlateApplicationBase.h
--- a/Engine/Source/Runtime/SlateCore/Public/Application/SlateApplicationBase.h
+++ b/Engine/Source/Runtime/SlateCore/Public/Application/SlateApplicationBase.h
@@ -184,6 +184,15 @@ public:
 	 */
 	virtual TSharedRef<IToolTip> MakeToolTip( const TAttribute<FString>& ToolTipString ) = 0;
 
+	/**
+	 * Creates a tool tip with the specified text.
+	 *
+	 * @param ToolTipText The text attribute to assign to the tool tip.
+	 *
+	 * @return The tool tip.
+	 */
+	virtual TSharedRef<IToolTip> MakeToolTip(const TAttribute<FText>& ToolTipText) = 0;
+
 	/**
 	 * Creates a tool tip with the specified text.
 	 *
diff --git a/Engine/Source/Runtime/SlateCore/Public/Widgets/SWidget.h b/Engine/Source/Runtime/SlateCore/Public/Widgets/SWidget.h
--- a/Engine/Source/Runtime/SlateCore/Public/Widgets/SWidget.h
+++ b/Engine/Source/Runtime/SlateCore/Public/Widgets/SWidget.h
@@ -581,6 +581,8 @@ public:
 	 */
 	void SetToolTipText( const TAttribute<FString>& InToolTipText );
 
+	void SetToolTipText(const TAttribute<FText>& ToolTipText);
+
 	void SetToolTipText( const FText& InToolTipText );
 
 	/**
