Copy/pasting from a property matrix which doesn't have common pins no longer crashes

#ttp 334645 - Editor: Property Matrix: Crash when copying/pasting multiple cells when some properties are invalid for some actors

#branch UE4

#change Prevented crash by stopping trying to access data when the cell is unbound

#change PasteTextAtCell was ignoring the cells in the first row with no string value (L419), but allows those in other rows (L462). Since unbound cells return empty strings this was changed to always allow (plus, it looked like a copy/paste error from L415). The alternative to changing this is to implement something similar to CL#2066840 which returns a string, but then ignores it.

diff --git a/Engine/Source/Editor/PropertyEditor/Private/Presentation/PropertyTable/PropertyTable.cpp b/Engine/Source/Editor/PropertyEditor/Private/Presentation/PropertyTable/PropertyTable.cpp
--- a/Engine/Source/Editor/PropertyEditor/Private/Presentation/PropertyTable/PropertyTable.cpp
+++ b/Engine/Source/Editor/PropertyEditor/Private/Presentation/PropertyTable/PropertyTable.cpp
@@ -416,7 +416,7 @@ void FPropertyTable::PasteTextAtCell( const FString& Text, const TSharedRef< IPr
 
 	// Parse row strings into individual cell strings
 	TArray<FString> CellStrings;
-	RowStrings[CurrentRowIdx++].ParseIntoArray(&CellStrings, TEXT("\t"), true);
+	RowStrings[CurrentRowIdx++].ParseIntoArray(&CellStrings, TEXT("\t"), false);
 
 	// Get the maximum paste operations before displaying the slow task
 	int32 NumPasteOperationsBeforeWarning = GEditor->AccessEditorUserSettings().PropertyMatrix_NumberOfPasteOperationsBeforeWarning;
diff --git a/Engine/Source/Editor/PropertyEditor/Private/Presentation/PropertyTable/PropertyTableCell.cpp b/Engine/Source/Editor/PropertyEditor/Private/Presentation/PropertyTable/PropertyTableCell.cpp
--- a/Engine/Source/Editor/PropertyEditor/Private/Presentation/PropertyTable/PropertyTableCell.cpp
+++ b/Engine/Source/Editor/PropertyEditor/Private/Presentation/PropertyTable/PropertyTableCell.cpp
@@ -114,12 +114,12 @@ bool FPropertyTableCell::IsValid() const
 
 FString FPropertyTableCell::GetValueAsString() const
 {
-	return PropertyEditor->GetValueAsString();
+	return IsBound() ? PropertyEditor->GetValueAsString() : FString();
 }
 
 FText FPropertyTableCell::GetValueAsText() const
 {
-	return PropertyEditor->GetValueAsText();
+	return IsBound() ? PropertyEditor->GetValueAsText() : FText::GetEmpty();
 }
 
 void FPropertyTableCell::SetValueFromString( const FString& InString )
@@ -150,7 +150,7 @@ TSharedRef< class IPropertyTable > FPropertyTableCell::GetTable() const
 
 TSharedPtr< class IPropertyHandle > FPropertyTableCell::GetPropertyHandle() const 
 { 
-	if( PropertyEditor.IsValid() ) 
+	if( IsBound() ) 
 	{
 		return PropertyEditor->GetPropertyHandle(); 
 	}
