Fix properties of structs not being considered when a property is checked for EditOnInstance/EditOnDefault visibility

diff --git a/Engine/Source/Editor/PropertyEditor/Private/DetailLayoutBuilderImpl.cpp b/Engine/Source/Editor/PropertyEditor/Private/DetailLayoutBuilderImpl.cpp
--- a/Engine/Source/Editor/PropertyEditor/Private/DetailLayoutBuilderImpl.cpp
+++ b/Engine/Source/Editor/PropertyEditor/Private/DetailLayoutBuilderImpl.cpp
@@ -407,6 +407,16 @@ TSharedPtr<FAssetThumbnailPool> FDetailLayoutBuilderImpl::GetThumbnailPool() con
 	return DetailsView.GetThumbnailPool();
 }
 
+bool FDetailLayoutBuilderImpl::IsPropertyVisible( TSharedRef<IPropertyHandle> PropertyHandle ) const
+{
+	return PropertyHandle->IsValidHandle() ? IsPropertyVisible( PropertyHandle->GetProperty() ) : false;
+}
+
+bool FDetailLayoutBuilderImpl::IsPropertyVisible( UProperty* Property ) const
+{
+	return DetailsView.IsPropertyVisible( Property );
+}
+
 const IDetailsView& FDetailLayoutBuilderImpl::GetDetailsView() const
 { 
 	return DetailsView; 
diff --git a/Engine/Source/Editor/PropertyEditor/Private/DetailLayoutBuilderImpl.h b/Engine/Source/Editor/PropertyEditor/Private/DetailLayoutBuilderImpl.h
--- a/Engine/Source/Editor/PropertyEditor/Private/DetailLayoutBuilderImpl.h
+++ b/Engine/Source/Editor/PropertyEditor/Private/DetailLayoutBuilderImpl.h
@@ -17,7 +17,9 @@ public:
 	virtual void HideProperty( FName PropertyPath, const UClass* ClassOutermost = NULL, FName InstanceName = NAME_None ) OVERRIDE;
 	virtual void ForceRefreshDetails() OVERRIDE;
 	virtual TSharedPtr<FAssetThumbnailPool> GetThumbnailPool() const OVERRIDE;
-	
+	virtual bool IsPropertyVisible( TSharedRef<IPropertyHandle> PropertyHandle ) const OVERRIDE;
+	virtual bool IsPropertyVisible( UProperty* PropertyHandle ) const OVERRIDE;
+
 	virtual const TSharedRef< IPropertyUtilities >& GetPropertyUtilities() const OVERRIDE; 
 
 	/**
diff --git a/Engine/Source/Editor/PropertyEditor/Private/DetailPropertyRow.cpp b/Engine/Source/Editor/PropertyEditor/Private/DetailPropertyRow.cpp
--- a/Engine/Source/Editor/PropertyEditor/Private/DetailPropertyRow.cpp
+++ b/Engine/Source/Editor/PropertyEditor/Private/DetailPropertyRow.cpp
@@ -240,6 +240,12 @@ void FDetailPropertyRow::GenerateChildrenForPropertyNode( TSharedPtr<FPropertyNo
 	}
 	else if (bShowCustomPropertyChildren || !CustomPropertyWidget.IsValid() )
 	{
+		TSharedRef<FDetailCategoryImpl> ParentCategoryRef = ParentCategory.Pin().ToSharedRef();
+		IDetailLayoutBuilder& LayoutBuilder = ParentCategoryRef->GetParentLayout();
+		UProperty* Property = RootPropertyNode->GetProperty();
+
+		const bool bStructProperty = Property && Property->IsA<UStructProperty>();
+
 		for( int32 ChildIndex = 0; ChildIndex < RootPropertyNode->GetNumChildNodes(); ++ChildIndex )
 		{
 			TSharedPtr<FPropertyNode> ChildNode = RootPropertyNode->GetChildNode(ChildIndex);
@@ -251,11 +257,12 @@ void FDetailPropertyRow::GenerateChildrenForPropertyNode( TSharedPtr<FPropertyNo
 					// Skip over object nodes and generate their children.  Object nodes are not visible
 					GenerateChildrenForPropertyNode( ChildNode, OutChildren );
 				}
-				else
-				{
+				// Only struct children can have custom visibility that is different from their parent.
+				else if ( !bStructProperty || LayoutBuilder.IsPropertyVisible( ChildNode->GetProperty() ) )
+				{			
 					FDetailLayoutCustomization Customization;
-					Customization.PropertyRow = MakeShareable( new FDetailPropertyRow( ChildNode, ParentCategory.Pin().ToSharedRef() ) );
-					TSharedRef<FDetailItemNode> ChildNodeItem = MakeShareable( new FDetailItemNode( Customization, ParentCategory.Pin().ToSharedRef(), ParentEnabledState ) );
+					Customization.PropertyRow = MakeShareable( new FDetailPropertyRow( ChildNode, ParentCategoryRef ) );
+					TSharedRef<FDetailItemNode> ChildNodeItem = MakeShareable( new FDetailItemNode( Customization, ParentCategoryRef, ParentEnabledState ) );
 					ChildNodeItem->Initialize();
 					OutChildren.Add( ChildNodeItem );
 				}
diff --git a/Engine/Source/Editor/PropertyEditor/Private/SDetailsView.cpp b/Engine/Source/Editor/PropertyEditor/Private/SDetailsView.cpp
--- a/Engine/Source/Editor/PropertyEditor/Private/SDetailsView.cpp
+++ b/Engine/Source/Editor/PropertyEditor/Private/SDetailsView.cpp
@@ -473,7 +473,7 @@ void SDetailsView::UpdatePropertyMapRecursive( FPropertyNode& InNode, FDetailLay
 				bVisibleByDefault &= !bIsChildOfAttribute;
 
 				// See if the user requested specific property visibility 
-				const bool bIsUserVisible = IsPropertyVisible.IsBound() ? IsPropertyVisible.Execute( Property ) : true;
+				const bool bIsUserVisible = IsPropertyVisible( Property ); 
 
 				// Inners of customized in structs should not be taken into consideration for customizing.  They are not designed to be individually customized when their parent is already customized
 				if( !bIsChildOfCustomizedStruct )
@@ -1470,7 +1470,7 @@ void SDetailsView::UnregisterInstancedCustomPropertyLayout( UClass* Class )
 
 void SDetailsView::SetIsPropertyVisibleDelegate( FIsPropertyVisible InIsPropertyVisible )
 {
-	IsPropertyVisible = InIsPropertyVisible;
+	IsPropertyVisibleDelegate = InIsPropertyVisible;
 }
 
 void SDetailsView::SetIsPropertyEditingEnabledDelegate( FIsPropertyEditingEnabled IsPropertyEditingEnabled )
@@ -1563,4 +1563,9 @@ bool SDetailsView::IsCategoryHiddenByClass( FName CategoryName ) const
 	return RootPropertyNode->GetHiddenCategories().Contains( CategoryName );
 }
 
+bool SDetailsView::IsPropertyVisible( const UProperty* Property ) const
+{
+	return IsPropertyVisibleDelegate.IsBound() ? IsPropertyVisibleDelegate.Execute( Property ) : true;
+}
+
 #undef LOCTEXT_NAMESPACE
diff --git a/Engine/Source/Editor/PropertyEditor/Private/SDetailsView.h b/Engine/Source/Editor/PropertyEditor/Private/SDetailsView.h
--- a/Engine/Source/Editor/PropertyEditor/Private/SDetailsView.h
+++ b/Engine/Source/Editor/PropertyEditor/Private/SDetailsView.h
@@ -290,6 +290,12 @@ public:
 	 * @return True if a category is hidden by any of the uobject classes currently in view by this details panel
 	 */
 	bool IsCategoryHiddenByClass( FName CategoryName ) const;
+
+	/**
+	 * @return True if the property is visible
+	 */
+	bool IsPropertyVisible( const UProperty* Property ) const;
+
 private:
 	void RegisterInstancedCustomPropertyLayout( UClass* Class, FOnGetDetailCustomizationInstance DetailLayoutDelegate );
 	void UnregisterInstancedCustomPropertyLayout( UClass* Class );
@@ -501,7 +507,7 @@ private:
 	/** Callback to send when the property view changes */
 	FOnObjectArrayChanged OnObjectArrayChanged;
 	/** Delegate executed to determine if a property should be visible */
-	FIsPropertyVisible IsPropertyVisible;
+	FIsPropertyVisible IsPropertyVisibleDelegate;
 	/** Delegate called to see if a property editing is enabled */
 	FIsPropertyEditingEnabled IsPropertyEditingEnabledDelegate;
 	/** Delegate called when the details panel finishes editing a property (after post edit change is called) */
diff --git a/Engine/Source/Editor/PropertyEditor/Public/DetailLayoutBuilder.h b/Engine/Source/Editor/PropertyEditor/Public/DetailLayoutBuilder.h
--- a/Engine/Source/Editor/PropertyEditor/Public/DetailLayoutBuilder.h
+++ b/Engine/Source/Editor/PropertyEditor/Public/DetailLayoutBuilder.h
@@ -129,4 +129,14 @@ public:
 	 * Gets the thumbnail pool that should be used for rendering thumbnails in the details view                   
 	 */
 	virtual TSharedPtr<class FAssetThumbnailPool> GetThumbnailPool() const = 0;
+
+	/**
+	 * @return true if the property should be visible in the details panel or false if the specific details panel is not showing this property
+	 */
+	virtual bool IsPropertyVisible( TSharedRef<IPropertyHandle> PropertyHandle ) const = 0;
+
+	/**
+	 * @return true if the property should be visible in the details panel or false if the specific details panel is not showing this property
+	 */
+	virtual bool IsPropertyVisible( UProperty* Property ) const = 0;
 };
