PropertyEditor refactored

diff --git a/Engine/Source/Editor/PropertyEditor/Private/DetailAdvancedDropdownNode.h b/Engine/Source/Editor/PropertyEditor/Private/DetailAdvancedDropdownNode.h
--- a/Engine/Source/Editor/PropertyEditor/Private/DetailAdvancedDropdownNode.h
+++ b/Engine/Source/Editor/PropertyEditor/Private/DetailAdvancedDropdownNode.h
@@ -25,7 +25,7 @@ public:
 	{}
 private:
 	/** IDetailTreeNode Interface */
-	virtual SDetailsView& GetDetailsView() const OVERRIDE { return ParentCategory.GetDetailsView(); }
+	virtual IDetailsViewPrivate& GetDetailsView() const OVERRIDE{ return ParentCategory.GetDetailsView(); }
 	virtual TSharedRef< ITableRow > GenerateNodeWidget( const TSharedRef<STableViewBase>& OwnerTable, const FDetailColumnSizeData& ColumnSizeData, const TSharedRef<IPropertyUtilities>& PropertyUtilities ) OVERRIDE;
 	virtual void GetChildren( TArray< TSharedRef<IDetailTreeNode> >& OutChildren )  OVERRIDE {}
 	virtual void OnItemExpansionChanged( bool bIsExpanded ) OVERRIDE {}
diff --git a/Engine/Source/Editor/PropertyEditor/Private/DetailCategoryBuilderImpl.h b/Engine/Source/Editor/PropertyEditor/Private/DetailCategoryBuilderImpl.h
--- a/Engine/Source/Editor/PropertyEditor/Private/DetailCategoryBuilderImpl.h
+++ b/Engine/Source/Editor/PropertyEditor/Private/DetailCategoryBuilderImpl.h
@@ -145,7 +145,7 @@ public:
 	virtual const FString& GetDisplayName() const OVERRIDE { return DisplayName; }
 
 	/** IDetailTreeNode interface */
-	virtual SDetailsView& GetDetailsView() const OVERRIDE { return DetailLayoutBuilder.Pin()->GetDetailsView(); }
+	virtual IDetailsViewPrivate& GetDetailsView() const OVERRIDE{ return DetailLayoutBuilder.Pin()->GetDetailsView(); }
 	virtual TSharedRef< ITableRow > GenerateNodeWidget( const TSharedRef<STableViewBase>& OwnerTable, const FDetailColumnSizeData& ColumnSizeData, const TSharedRef<IPropertyUtilities>& PropertyUtilities ) OVERRIDE;
 	virtual void GetChildren( TArray< TSharedRef<IDetailTreeNode> >& OutChildren ) OVERRIDE;
 	virtual bool ShouldBeExpanded() const OVERRIDE;
diff --git a/Engine/Source/Editor/PropertyEditor/Private/DetailCategoryGroupNode.h b/Engine/Source/Editor/PropertyEditor/Private/DetailCategoryGroupNode.h
--- a/Engine/Source/Editor/PropertyEditor/Private/DetailCategoryGroupNode.h
+++ b/Engine/Source/Editor/PropertyEditor/Private/DetailCategoryGroupNode.h
@@ -33,7 +33,7 @@ public:
 	FDetailCategoryGroupNode( const FDetailNodeList& InChildNodes, FName InGroupName, FDetailCategoryImpl& InParentCategory );
 
 private:
-	virtual SDetailsView& GetDetailsView() const OVERRIDE { return ParentCategory.GetDetailsView(); }
+	virtual IDetailsViewPrivate& GetDetailsView() const OVERRIDE{ return ParentCategory.GetDetailsView(); }
 	virtual void OnItemExpansionChanged( bool bIsExpanded ) OVERRIDE {}
 	virtual bool ShouldBeExpanded() const OVERRIDE { return true; }
 	virtual ENodeVisibility::Type GetVisibility() const OVERRIDE { return bShouldBeVisible ? ENodeVisibility::Visible : ENodeVisibility::HiddenDueToFiltering; }
diff --git a/Engine/Source/Editor/PropertyEditor/Private/DetailItemNode.h b/Engine/Source/Editor/PropertyEditor/Private/DetailItemNode.h
--- a/Engine/Source/Editor/PropertyEditor/Private/DetailItemNode.h
+++ b/Engine/Source/Editor/PropertyEditor/Private/DetailItemNode.h
@@ -31,7 +31,7 @@ public:
 	bool HasMultiColumnWidget() const;
 
 	/** IDetailTreeNode interface */
-	virtual SDetailsView& GetDetailsView() const OVERRIDE { return ParentCategory.Pin()->GetDetailsView(); }
+	virtual IDetailsViewPrivate& GetDetailsView() const OVERRIDE{ return ParentCategory.Pin()->GetDetailsView(); }
 	virtual TSharedRef< ITableRow > GenerateNodeWidget( const TSharedRef<STableViewBase>& OwnerTable, const FDetailColumnSizeData& ColumnSizeData, const TSharedRef<IPropertyUtilities>& PropertyUtilities ) OVERRIDE;
 	virtual void GetChildren( FDetailNodeList& OutChildren )  OVERRIDE;
 	virtual void OnItemExpansionChanged( bool bInIsExpanded ) OVERRIDE;
diff --git a/Engine/Source/Editor/PropertyEditor/Private/DetailLayoutBuilderImpl.cpp b/Engine/Source/Editor/PropertyEditor/Private/DetailLayoutBuilderImpl.cpp
--- a/Engine/Source/Editor/PropertyEditor/Private/DetailLayoutBuilderImpl.cpp
+++ b/Engine/Source/Editor/PropertyEditor/Private/DetailLayoutBuilderImpl.cpp
@@ -7,7 +7,7 @@
 #include "IPropertyUtilities.h"
 
 
-FDetailLayoutBuilderImpl::FDetailLayoutBuilderImpl( FClassToPropertyMap& InPropertyMap, const TSharedRef< class IPropertyUtilities >& InPropertyUtilities, const TSharedRef< SDetailsView >& InDetailsView )
+FDetailLayoutBuilderImpl::FDetailLayoutBuilderImpl(FClassToPropertyMap& InPropertyMap, const TSharedRef< class IPropertyUtilities >& InPropertyUtilities, const TSharedRef< IDetailsViewPrivate >& InDetailsView)
 	: PropertyMap( InPropertyMap )
 	, PropertyDetailsUtilities( InPropertyUtilities )
 	, CurrentCustomizationClass( NULL )
diff --git a/Engine/Source/Editor/PropertyEditor/Private/DetailLayoutBuilderImpl.h b/Engine/Source/Editor/PropertyEditor/Private/DetailLayoutBuilderImpl.h
--- a/Engine/Source/Editor/PropertyEditor/Private/DetailLayoutBuilderImpl.h
+++ b/Engine/Source/Editor/PropertyEditor/Private/DetailLayoutBuilderImpl.h
@@ -6,7 +6,7 @@ class FDetailCategoryImpl;
 class FDetailLayoutBuilderImpl : public IDetailLayoutBuilder, public TSharedFromThis<FDetailLayoutBuilderImpl>
 {
 public:
-	FDetailLayoutBuilderImpl( FClassToPropertyMap& InPropertyMap, const TSharedRef< class IPropertyUtilities >& InPropertyUtilities, const TSharedRef< SDetailsView >& InDetailsView );
+	FDetailLayoutBuilderImpl(FClassToPropertyMap& InPropertyMap, const TSharedRef< class IPropertyUtilities >& InPropertyUtilities, const TSharedRef< IDetailsViewPrivate >& InDetailsView);
 
 	/** IDetailLayoutBuilder Interface */
 	virtual const IDetailsView& GetDetailsView() const OVERRIDE;
@@ -141,7 +141,7 @@ public:
 	void AddExternalRootPropertyNode( TSharedRef<FObjectPropertyNode> InExternalRootNode );
 
 	/** @return The details view that owns this layout */
-	SDetailsView& GetDetailsView() { return DetailsView; }
+	IDetailsViewPrivate& GetDetailsView() { return DetailsView; }
 private:
 	/**
 	 * Finds a property node for the current property by searching in a fast lookup map or a path search if required
@@ -178,7 +178,7 @@ private:
 	/** The global property settings */
 	const TSharedRef< class IPropertyUtilities > PropertyDetailsUtilities;
 	/** The view where this detail customizer resides */
-	class SDetailsView& DetailsView;
+	class IDetailsViewPrivate& DetailsView;
 	/** The current class being customized */
 	UClass* CurrentCustomizationClass;
 };
diff --git a/Engine/Source/Editor/PropertyEditor/Private/IDetailTreeNode.h b/Engine/Source/Editor/PropertyEditor/Private/IDetailTreeNode.h
--- a/Engine/Source/Editor/PropertyEditor/Private/IDetailTreeNode.h
+++ b/Engine/Source/Editor/PropertyEditor/Private/IDetailTreeNode.h
@@ -25,7 +25,7 @@ public:
 	virtual ~IDetailTreeNode() {}
 
 	/** @return The details view that this node is in */
-	virtual SDetailsView& GetDetailsView() const = 0;
+	virtual IDetailsViewPrivate& GetDetailsView() const = 0;
 
 	/**
 	 * Generates the widget representing this node
diff --git a/Engine/Source/Editor/PropertyEditor/Private/IDetailsViewPrivate.h b/Engine/Source/Editor/PropertyEditor/Private/IDetailsViewPrivate.h
new file mode 100644
--- /dev/null
+++ b/Engine/Source/Editor/PropertyEditor/Private/IDetailsViewPrivate.h
@@ -0,0 +1,90 @@
+// Copyright 1998-2014 Epic Games, Inc. All Rights Reserved.
+
+#pragma once
+
+#include "IDetailsView.h"
+
+class FNotifyHook;
+class IDetailTreeNode;
+class IPropertyUtilities;
+
+class IDetailsViewPrivate : public IDetailsView
+{
+public:
+	/**
+	* Sets the expansion state for a node and optionally all of its children
+	*
+	* @param InTreeNode		The node to change expansion state on
+	* @param bIsItemExpanded	The new expansion state
+	* @param bRecursive		Whether or not to apply the expansion change to any children
+	*/
+	virtual void SetNodeExpansionState(TSharedRef<IDetailTreeNode> InTreeNode, bool bIsItemExpanded, bool bRecursive) = 0;
+
+	/**
+	 * Requests that an item in the tree be expanded or collapsed
+	 *
+	 * @param TreeNode	The tree node to expand
+	 * @param bExpand	true if the item should be expanded, false otherwise
+	 */
+	virtual void RequestItemExpanded(TSharedRef<IDetailTreeNode> TreeNode, bool bExpand) = 0;
+
+	/**
+	 * Refreshes the detail's treeview
+	 */
+	virtual void RefreshTree() = 0;
+
+	/** Returns the notify hook to use when properties change */
+	virtual FNotifyHook* GetNotifyHook() const = 0;
+
+	/**
+	* Returns the property utilities for this view
+	*/
+	virtual TSharedPtr<IPropertyUtilities> GetPropertyUtilities() = 0;
+
+
+	/** Gets the base class being viewed */
+	virtual const UClass* GetBaseClass() const = 0;
+	virtual UClass* GetBaseClass() = 0;
+
+	/**
+	* @return True if a category is hidden by any of the uobject classes currently in view by this details panel
+	*/
+	virtual bool IsCategoryHiddenByClass(FName CategoryName) const = 0;
+
+	/** Causes the details view to be refreshed (new widgets generated) with the current set of objects */
+	virtual void ForceRefresh() = 0;
+
+	/**
+	* Saves the expansion state of a tree node
+	*
+	* @param NodePath	The path to the detail node to save
+	* @param bIsExpanded	true if the node is expanded, false otherwise
+	*/
+	virtual void SaveCustomExpansionState(const FString& NodePath, bool bIsExpanded) = 0;
+
+	/**
+	* Gets the saved expansion state of a tree node in this category
+	*
+	* @param NodePath	The path to the detail node to get
+	* @return true if the node should be expanded, false otherwise
+	*/
+	virtual bool GetCustomSavedExpansionState(const FString& NodePath) const = 0;
+
+	/**
+	 * @return True if the property is visible
+	 */
+	virtual bool IsPropertyVisible(const UProperty* Property) const = 0;
+
+	/**
+	 * @return The thumbnail pool that should be used for thumbnails being rendered in this view
+	 */
+	virtual TSharedPtr<class FAssetThumbnailPool> GetThumbnailPool() const = 0;
+
+	/**
+	 * Adds an external property root node to the list of root nodes that the details new needs to manage
+	 *
+	 * @param InExternalRootNode	The node to add
+	 */
+	virtual void AddExternalRootPropertyNode(TSharedRef<FObjectPropertyNode> ExternalRootNode) = 0;
+
+};
\ No newline at end of file
diff --git a/Engine/Source/Editor/PropertyEditor/Private/SDetailsView.h b/Engine/Source/Editor/PropertyEditor/Private/SDetailsView.h
--- a/Engine/Source/Editor/PropertyEditor/Private/SDetailsView.h
+++ b/Engine/Source/Editor/PropertyEditor/Private/SDetailsView.h
@@ -4,6 +4,7 @@
 
 #include "AssetSelection.h"
 #include "IPropertyUtilities.h"
+#include "IDetailsViewPrivate.h"
 
 class FPropertyNode;
 class FObjectPropertyNode;
@@ -74,7 +75,7 @@ struct FDetailColumnSizeData
 	void SetColumnWidth( float InWidth ) { OnWidthChanged.ExecuteIfBound(InWidth); }
 };
 
-class SDetailsView : public IDetailsView
+class SDetailsView : public IDetailsViewPrivate
 {
 	friend class FPropertyDetailsUtilities;
 public:
