CIS fix - Fixed Mac build error about creating a const FString

diff --git a/Engine/Source/Editor/PropertyEditor/Private/PropertyHandleImpl.cpp b/Engine/Source/Editor/PropertyEditor/Private/PropertyHandleImpl.cpp
--- a/Engine/Source/Editor/PropertyEditor/Private/PropertyHandleImpl.cpp
+++ b/Engine/Source/Editor/PropertyEditor/Private/PropertyHandleImpl.cpp
@@ -1655,7 +1655,7 @@ bool FPropertyHandleBase::HasMetaData(const FName& Key) const
 const FString& FPropertyHandleBase::GetMetaData(const FName& Key) const
 {
 	// if not found, return a static empty string
-	static const FString EmptyString;
+	static const FString EmptyString = TEXT("");
 
 	UProperty* const MetaDataProperty = GetMetaDataProperty();
 	return (MetaDataProperty) ? MetaDataProperty->GetMetaData(Key) : EmptyString;
