PropertyEditor refactored. SDetailView split.

diff --git a/Engine/Source/Editor/PropertyEditor/Private/IDetailsViewPrivate.h b/Engine/Source/Editor/PropertyEditor/Private/IDetailsViewPrivate.h
--- a/Engine/Source/Editor/PropertyEditor/Private/IDetailsViewPrivate.h
+++ b/Engine/Source/Editor/PropertyEditor/Private/IDetailsViewPrivate.h
@@ -7,6 +7,8 @@
 class FNotifyHook;
 class IDetailTreeNode;
 class IPropertyUtilities;
+class FPropertyNode;
+class FObjectPropertyNode;
 
 class IDetailsViewPrivate : public IDetailsView
 {
@@ -87,4 +89,4 @@ public:
 	 */
 	virtual void AddExternalRootPropertyNode(TSharedRef<FObjectPropertyNode> ExternalRootNode) = 0;
 
-};
\ No newline at end of file
+};
diff --git a/Engine/Source/Editor/PropertyEditor/Private/SDetailsView.cpp b/Engine/Source/Editor/PropertyEditor/Private/SDetailsView.cpp
--- a/Engine/Source/Editor/PropertyEditor/Private/SDetailsView.cpp
+++ b/Engine/Source/Editor/PropertyEditor/Private/SDetailsView.cpp
@@ -22,11 +22,6 @@
 SDetailsView::~SDetailsView()
 {
 	SaveExpandedItems();
-
-	if( ThumbnailPool.IsValid() )
-	{
-		ThumbnailPool->ReleaseResources();
-	}
 };
 
 /**
@@ -212,58 +207,6 @@ TSharedRef<SDetailTree> SDetailsView::ConstructTreeView( TSharedRef<SScrollBar>&
 	.ExternalScrollbar( ScrollBar );
 }
 
-void SDetailsView::OnGetChildrenForDetailTree( TSharedRef<IDetailTreeNode> InTreeNode, TArray< TSharedRef<IDetailTreeNode> >& OutChildren )
-{
-	InTreeNode->GetChildren( OutChildren );
-}
-
-
-TSharedRef<ITableRow> SDetailsView::OnGenerateRowForDetailTree( TSharedRef<IDetailTreeNode> InTreeNode, const TSharedRef<STableViewBase>& OwnerTable )
-{
-	return InTreeNode->GenerateNodeWidget( OwnerTable, ColumnSizeData, PropertyUtilities.ToSharedRef() );
-}
-
-void SDetailsView::SetRootExpansionStates( const bool bExpand, const bool bRecurse )
-{
-	for (auto Iter = RootTreeNodes.CreateIterator(); Iter; ++Iter)
-	{
-		SetNodeExpansionState(*Iter, bExpand, bRecurse );
-	}
-}
-
-void SDetailsView::SetNodeExpansionState( TSharedRef<IDetailTreeNode> InTreeNode, bool bIsItemExpanded, bool bRecursive )
-{
-	TArray< TSharedRef<IDetailTreeNode> > Children;
-	InTreeNode->GetChildren( Children );
-
-	if( Children.Num() )
-	{
-		RequestItemExpanded( InTreeNode, bIsItemExpanded );
-		InTreeNode->OnItemExpansionChanged( bIsItemExpanded );
-
-		if( bRecursive )
-		{
-			for( int32 ChildIndex = 0; ChildIndex < Children.Num(); ++ChildIndex )
-			{
-				TSharedRef<IDetailTreeNode> Child = Children[ChildIndex];
-
-				SetNodeExpansionState( Child, bIsItemExpanded, bRecursive );
-			}
-		}
-	}
-}
-
-void SDetailsView::SetNodeExpansionStateRecursive( TSharedRef<IDetailTreeNode> InTreeNode, bool bIsItemExpanded )
-{
-	SetNodeExpansionState( InTreeNode, bIsItemExpanded, true );
-}
-
-void SDetailsView::OnItemExpansionChanged( TSharedRef<IDetailTreeNode> InTreeNode, bool bIsItemExpanded )
-{
-	SetNodeExpansionState( InTreeNode, bIsItemExpanded, false );
-}
-
-
 void SDetailsView::SetColorPropertyFromColorPicker(FLinearColor NewColor)
 {
 	const TSharedPtr< FPropertyNode > PinnedColorPropertyNode = ColorPropertyNode.Pin();
@@ -294,12 +237,6 @@ FReply SDetailsView::OnOpenRawPropertyEditorClicked()
 	return FReply::Handled();
 }
 
-FReply SDetailsView::OnLockButtonClicked()
-{
-	bIsLocked = !bIsLocked;
-	return FReply::Handled();
-}
-
 void SDetailsView::OnShowOnlyModifiedClicked()
 {
 	CurrentFilter.bShowOnlyModifiedProperties = !CurrentFilter.bShowOnlyModifiedProperties;
@@ -333,35 +270,12 @@ EVisibility SDetailsView::GetActorNameAreaVisibility() const
 	return bVisible ? EVisibility::Visible : EVisibility::Collapsed; 
 }
 
-void SDetailsView::HideFilterArea(bool bHide)
-{
-	DetailsViewArgs.bAllowSearch = !bHide;
-}
-
 EVisibility SDetailsView::GetFilterBoxVisibility() const
 {
 	// Visible if we allow search and we have anything to search otherwise collapsed so it doesn't take up room
 	return DetailsViewArgs.bAllowSearch && RootPropertyNode->GetNumObjects() > 0 ? EVisibility::Visible : EVisibility::Collapsed;
 }
 
-EVisibility SDetailsView::GetTreeVisibility() const
-{
-	return DetailLayout.IsValid() && DetailLayout->HasDetails() ? EVisibility::Visible : EVisibility::Collapsed;
-}
-
-/** Returns the image used for the icon on the filter button */ 
-const FSlateBrush* SDetailsView::OnGetFilterButtonImageResource() const
-{
-	if( bHasActiveFilter )
-	{
-		return FEditorStyle::GetBrush(TEXT("PropertyWindow.FilterCancel"));
-	}
-	else
-	{
-		return FEditorStyle::GetBrush(TEXT("PropertyWindow.FilterSearch"));
-	}
-}
-
 /**
  * Determines whether or not a property should be visible in the default generated detail layout
  *
@@ -632,7 +546,6 @@ void SDetailsView::SetObject( UObject* InObject, bool bForceRefresh )
 	SetObjects( ObjectWeakPtrs, bForceRefresh );
 }
 
-
 bool SDetailsView::ShouldSetNewObjects( const TArray< TWeakObjectPtr< UObject > >& InObjects ) const
 {
 	bool bShouldSetObjects = false;
@@ -764,12 +677,6 @@ void SDetailsView::SetObjectArrayPrivate( const TArray< TWeakObjectPtr< UObject
 	double ElapsedTime = FPlatformTime::Seconds() - StartTime;
 }
 
-void SDetailsView::EnqueueDeferredAction( FSimpleDelegate& DeferredAction )
-{
-	DeferredActions.AddUnique( DeferredAction );
-}
-
-
 void SDetailsView::ReplaceObjects( const TMap<UObject*, UObject*>& OldToNewObjectMap )
 {
 	TArray< TWeakObjectPtr< UObject > > NewObjectList;
@@ -1234,7 +1141,6 @@ void SDetailsView::UpdateFilteredDetails()
 	DetailTree->RequestTreeRefresh();
 }
 
-
 /** Ticks the property view.  This function performs a data consistency check */
 void SDetailsView::Tick( const FGeometry& AllottedGeometry, const double InCurrentTime, const float InDeltaTime )
 {
@@ -1361,13 +1267,6 @@ void SDetailsView::Tick( const FGeometry& AllottedGeometry, const double InCurre
 	}
 }
 
-void SDetailsView::OnColorPickerWindowClosed( const TSharedRef<SWindow>& Window ) 
-{
-	// A color picker window is no longer open
-	bHasOpenColorPicker = false;
-	ColorPropertyNode.Reset();
-}
-
 /** 
  * Creates the color picker window for this property view.
  *
@@ -1418,11 +1317,6 @@ void SDetailsView::CreateColorPickerWindow(const TSharedRef< FPropertyEditor >&
 	OpenColorPicker(PickerArgs);
 }
 
-TSharedPtr<IPropertyUtilities> SDetailsView::GetPropertyUtilities() 
-{
-	return PropertyUtilities;
-}
-
 void SDetailsView::SetOnObjectArrayChanged(FOnObjectArrayChanged OnObjectArrayChangedDelegate)
 {
 	OnObjectArrayChanged = OnObjectArrayChangedDelegate;
@@ -1448,7 +1342,6 @@ UClass* SDetailsView::GetBaseClass()
 	return NULL;
 }
 
-
 void SDetailsView::RegisterInstancedCustomPropertyLayout( UClass* Class, FOnGetDetailCustomizationInstance DetailLayoutDelegate )
 {
 	check( Class );
@@ -1468,74 +1361,6 @@ void SDetailsView::UnregisterInstancedCustomPropertyLayout( UClass* Class )
 	InstancedClassToDetailLayoutMap.Remove( Class );
 }
 
-void SDetailsView::SetIsPropertyVisibleDelegate( FIsPropertyVisible InIsPropertyVisible )
-{
-	IsPropertyVisibleDelegate = InIsPropertyVisible;
-}
-
-void SDetailsView::SetIsPropertyEditingEnabledDelegate( FIsPropertyEditingEnabled IsPropertyEditingEnabled )
-{
-	IsPropertyEditingEnabledDelegate = IsPropertyEditingEnabled;
-}
-
-bool SDetailsView::IsPropertyEditingEnabled() const
-{
-	// If the delegate is not bound assume property editing is enabled, otherwise ask the delegate
-	return !IsPropertyEditingEnabledDelegate.IsBound() || IsPropertyEditingEnabledDelegate.Execute();
-}
-
-void SDetailsView::SetGenericLayoutDetailsDelegate( FOnGetDetailCustomizationInstance OnGetGenericDetails )
-{
-	GenericLayoutDelegate = OnGetGenericDetails;
-}
-
-TSharedPtr<FAssetThumbnailPool> SDetailsView::GetThumbnailPool() const
-{
-	if( !ThumbnailPool.IsValid() )
-	{
-		// Create a thumbnail pool for the view if it doesnt exist.  This does not use resources of no thumbnails are used
-		ThumbnailPool = MakeShareable( new FAssetThumbnailPool( 50, TAttribute<bool>::Create( TAttribute<bool>::FGetter::CreateSP(this, &SDetailsView::IsHovered) ) ) );
-	}
-
-	return ThumbnailPool;
-}
-
-void SDetailsView::NotifyFinishedChangingProperties(const FPropertyChangedEvent& PropertyChangedEvent) 
-{
-	OnFinishedChangingPropertiesDelegate.Broadcast(PropertyChangedEvent);
-}
-
-void SDetailsView::RequestItemExpanded( TSharedRef<IDetailTreeNode> TreeNode, bool bExpand )
-{
-	// Don't change expansion state if its already in that state
-	if( DetailTree->IsItemExpanded(TreeNode) != bExpand )
-	{
-		FilteredNodesRequestingExpansionState.Add( TreeNode, bExpand );
-	}
-}
-
-void SDetailsView::RefreshTree()
-{
-	DetailTree->RequestTreeRefresh();
-}
-
-void SDetailsView::SaveCustomExpansionState( const FString& NodePath, bool bIsExpanded )
-{
-	if( bIsExpanded )
-	{
-		ExpandedDetailNodes.Add( NodePath );
-	}
-	else
-	{
-		ExpandedDetailNodes.Remove( NodePath );
-	}
-}
-
-bool SDetailsView::GetCustomSavedExpansionState( const FString& NodePath ) const
-{
-	return ExpandedDetailNodes.Contains( NodePath );
-}
-
 bool SDetailsView::SupportsKeyboardFocus() const
 {
 	return DetailsViewArgs.bSearchInitialKeyFocus && SearchBox->SupportsKeyboardFocus() && GetFilterBoxVisibility() == EVisibility::Visible;
@@ -1563,9 +1388,4 @@ bool SDetailsView::IsCategoryHiddenByClass( FName CategoryName ) const
 	return RootPropertyNode->GetHiddenCategories().Contains( CategoryName );
 }
 
-bool SDetailsView::IsPropertyVisible( const UProperty* Property ) const
-{
-	return IsPropertyVisibleDelegate.IsBound() ? IsPropertyVisibleDelegate.Execute( Property ) : true;
-}
-
 #undef LOCTEXT_NAMESPACE
diff --git a/Engine/Source/Editor/PropertyEditor/Private/SDetailsView.h b/Engine/Source/Editor/PropertyEditor/Private/SDetailsView.h
--- a/Engine/Source/Editor/PropertyEditor/Private/SDetailsView.h
+++ b/Engine/Source/Editor/PropertyEditor/Private/SDetailsView.h
@@ -2,80 +2,9 @@
 
 #pragma once
 
-#include "AssetSelection.h"
-#include "IPropertyUtilities.h"
-#include "IDetailsViewPrivate.h"
+#include "SDetailsViewBase.h"
 
-class FPropertyNode;
-class FObjectPropertyNode;
-class FDetailCategoryImpl;
-class FDetailLayoutBuilderImpl;
-
-
-struct FPropertyNodeMap
-{
-	FPropertyNodeMap()
-		: ParentObjectProperty( NULL )
-	{}
-
-	/** Object property node which contains the properties in the node map */
-	FObjectPropertyNode* ParentObjectProperty;
-
-	/** Property name to property node map */
-	TMap<FName, TSharedPtr<FPropertyNode> > PropertyNameToNode;
-
-	bool Contains( FName PropertyName ) const
-	{
-		return PropertyNameToNode.Contains( PropertyName );
-	}
-
-	void Add( FName PropertyName, TSharedPtr<FPropertyNode>& PropertyNode )
-	{
-		PropertyNameToNode.Add( PropertyName, PropertyNode );
-	}
-};
-
-typedef TArray< TSharedRef<class IDetailTreeNode> > FDetailNodeList;
-
-/** Mapping of categories to all top level item property nodes in that category */
-typedef TMap<FName, TSharedPtr<FDetailCategoryImpl> > FCategoryMap;
-
-/** Class to properties in that class */
-typedef TMap<FName, FPropertyNodeMap> FClassInstanceToPropertyMap;
-
-/** Class to properties in that class */
-typedef TMap<FName, FClassInstanceToPropertyMap> FClassToPropertyMap;
-
-typedef STreeView< TSharedRef<class IDetailTreeNode> > SDetailTree;
-
-/** Represents a filter which controls the visibility of items in the details view */
-struct FDetailFilter
-{
-	FDetailFilter()
-		: bShowOnlyModifiedProperties( false )
-		, bShowAllAdvanced( false )
-	{}
-
-	bool IsEmptyFilter() const { return FilterStrings.Num() == 0 && bShowOnlyModifiedProperties == false && bShowAllAdvanced == false; }
-
-	/** Any user search terms that items must match */
-	TArray<FString> FilterStrings;
-	/** If we should only show modified properties */
-	bool bShowOnlyModifiedProperties;
-	/** If we should show all advanced properties */
-	bool bShowAllAdvanced;
-};
-
-struct FDetailColumnSizeData
-{
-	TAttribute<float> LeftColumnWidth;
-	TAttribute<float> RightColumnWidth;
-	SSplitter::FOnSlotResized OnWidthChanged;
-
-	void SetColumnWidth( float InWidth ) { OnWidthChanged.ExecuteIfBound(InWidth); }
-};
-
-class SDetailsView : public IDetailsViewPrivate
+class SDetailsView : public SDetailsViewBase
 {
 	friend class FPropertyDetailsUtilities;
 public:
@@ -87,10 +16,10 @@ public:
 		SLATE_ARGUMENT( FDetailsViewArgs, DetailsViewArgs )
 	SLATE_END_ARGS()
 
-	~SDetailsView();
+	virtual ~SDetailsView();
 
 	/** Causes the details view to be refreshed (new widgets generated) with the current set of objects */
-	void ForceRefresh();
+	void ForceRefresh() OVERRIDE;
 
 	/**
 	 * Constructs the property view widgets                   
@@ -102,8 +31,6 @@ public:
 	virtual void SetObjects( const TArray< TWeakObjectPtr< UObject > >& InObjects, bool bForceRefresh = false ) OVERRIDE;
 	virtual void SetObject( UObject* InObject, bool bForceRefresh = false ) OVERRIDE;
 
-	virtual FOnFinishedChangingProperties& OnFinishedChangingProperties() OVERRIDE { return OnFinishedChangingPropertiesDelegate; }
-
 	/**
 	 * Replaces objects being observed by the view with new objects
 	 *
@@ -125,11 +52,6 @@ public:
 	 */
 	void RemoveInvalidActors( const TSet<AActor*>& ValidActors );
 
-	/**
-	 * Called when the open color picker window associated with this details view is closed
-	 */
-	void OnColorPickerWindowClosed( const TSharedRef<SWindow>& Window );
-
 	/**
 	 * Creates the color picker window for this property view.
 	 *
@@ -138,29 +60,8 @@ public:
 	 */
 	virtual void CreateColorPickerWindow(const TSharedRef< class FPropertyEditor >& PropertyEditor, bool bUseAlpha);
 
-	/**
-	 * Returns true if the details view is locked and cant have its observed objects changed 
-	 */
-	virtual bool IsLocked() const OVERRIDE
-	{
-		return bIsLocked;
-	}
-
-	/**
-	 * @return true of the details view can be updated from editor selection
-	 */
-	virtual bool IsUpdatable() const OVERRIDE
-	{
-		return DetailsViewArgs.bUpdatesFromSelection;
-	}
-
 	/** Returns the notify hook to use when properties change */
-	virtual FNotifyHook* GetNotifyHook() const { return DetailsViewArgs.NotifyHook; }
-
-	/**
-	 * Returns the property utilities for this view
-	 */
-	TSharedPtr<IPropertyUtilities> GetPropertyUtilities();
+	virtual FNotifyHook* GetNotifyHook() const OVERRIDE { return DetailsViewArgs.NotifyHook; }
 
 	/** Sets the callback for when the property view changes */
 	virtual void SetOnObjectArrayChanged( FOnObjectArrayChanged OnObjectArrayChangedDelegate);
@@ -189,90 +90,8 @@ public:
 	}
 
 	/** Gets the base class being viewed */
-	const UClass* GetBaseClass() const;
-	UClass* GetBaseClass();
-
-	/** @return The identifier for this details view, or NAME_None is this view is anonymous */
-	virtual FName GetIdentifier() const OVERRIDE
-	{
-		return DetailsViewArgs.ViewIdentifier;
-	}
-
-	/**
-	 * Sets a delegate which is called to determine whether a specific property should be visible
-	 */
-	virtual void SetIsPropertyVisibleDelegate( FIsPropertyVisible InIsPropertyVisible ) OVERRIDE;
-
-	/**
-	 * Sets a delegate which is regardless of the objects being viewed to lay out generic details not specific to any object
-	 */
-	virtual void SetGenericLayoutDetailsDelegate( FOnGetDetailCustomizationInstance OnGetGenericDetails ) OVERRIDE;
-	
-    /**
-	 * Sets a delegate to call to determine if the properties  editing is enabled
-	 */ 
-	virtual void SetIsPropertyEditingEnabledDelegate( FIsPropertyEditingEnabled IsPropertyEditingEnabled ) OVERRIDE;
-
-	/** 
-	 * Sets the visible state of the filter box/property grid area
-	 */
-	virtual void HideFilterArea(bool bHide) OVERRIDE;
-
-	/**
-	 * @return true if property editing is enabled (based on the FIsPropertyEditingEnabled delegate)
-	 */
-	virtual bool IsPropertyEditingEnabled() const OVERRIDE;
-
-    /**
-	 * @return The thumbnail pool that should be used for thumbnails being rendered in this view
-	 */
-	TSharedPtr<class FAssetThumbnailPool> GetThumbnailPool() const;
-
-	/**
-	 * Requests that an item in the tree be expanded or collapsed
-	 *
-	 * @param TreeNode	The tree node to expand
-	 * @param bExpand	true if the item should be expanded, false otherwise
-	 */
-	void RequestItemExpanded( TSharedRef<IDetailTreeNode> TreeNode, bool bExpand );
-
-	/**
-	 * Sets the expansion state for a node and optionally all of its children
-	 *
-	 * @param InTreeNode		The node to change expansion state on
-	 * @param bIsItemExpanded	The new expansion state
-	 * @param bRecursive		Whether or not to apply the expansion change to any children
-	 */
-	void SetNodeExpansionState( TSharedRef<IDetailTreeNode> InTreeNode, bool bIsItemExpanded, bool bRecursive );
-
-	/**
-	 * Sets the expansion state all root nodes and optionally all of their children
-	 *
-	 * @param bExpand			The new expansion state
-	 * @param bRecurse			Whether or not to apply the expansion change to any children
-	 */
-	void SetRootExpansionStates( const bool bExpand, const bool bRecurse );
-
-	/**
-	 * Refreshes the detail's treeview
-	 */
-	void RefreshTree();
-
-	/**
-	 * Saves the expansion state of a tree node
-	 *
-	 * @param NodePath	The path to the detail node to save
-	 * @param bIsExpanded	true if the node is expanded, false otherwise
-	 */
-	void SaveCustomExpansionState( const FString& NodePath, bool bIsExpanded );
-
-	/**
-	 * Gets the saved expansion state of a tree node in this category	
-	 *
-	 * @param NodePath	The path to the detail node to get
-	 * @return true if the node should be expanded, false otherwise
-	 */	
-	bool GetCustomSavedExpansionState( const FString& NodePath ) const;
+	const UClass* GetBaseClass() const OVERRIDE;
+	UClass* GetBaseClass() OVERRIDE;
 
 	// SWidget interface
 	virtual bool SupportsKeyboardFocus() const OVERRIDE;
@@ -290,12 +109,7 @@ public:
 	/**
 	 * @return True if a category is hidden by any of the uobject classes currently in view by this details panel
 	 */
-	bool IsCategoryHiddenByClass( FName CategoryName ) const;
-
-	/**
-	 * @return True if the property is visible
-	 */
-	bool IsPropertyVisible( const UProperty* Property ) const;
+	bool IsCategoryHiddenByClass(FName CategoryName) const OVERRIDE;
 
 private:
 	void RegisterInstancedCustomPropertyLayout( UClass* Class, FOnGetDetailCustomizationInstance DetailLayoutDelegate );
@@ -304,35 +118,6 @@ private:
 
 	TSharedRef<SDetailTree> ConstructTreeView( TSharedRef<SScrollBar>& ScrollBar );
 
-	/** 
-	 * Function called through a delegate on the TreeView to request children of a tree node 
-	 * 
-	 * @param InTreeNode		The tree node to get children from
-	 * @param OutChildren		The list of children of InTreeNode that should be visible 
-	 */
-	void OnGetChildrenForDetailTree( TSharedRef<IDetailTreeNode> InTreeNode, TArray< TSharedRef<IDetailTreeNode> >& OutChildren );
-
-	/**
-	 * Returns an SWidget used as the visual representation of a node in the treeview.                     
-	 */
-	TSharedRef<ITableRow> OnGenerateRowForDetailTree( TSharedRef<IDetailTreeNode> InTreeNode, const TSharedRef<STableViewBase>& OwnerTable );
-
-	/**
-	 * Called to recursively expand/collapse the children of the given item
-	 *
-	 * @param InTreeNode		The node that was expanded or collapsed
-	 * @param bIsItemExpanded	True if the item is expanded, false if it is collapsed
-	 */
-	void SetNodeExpansionStateRecursive( TSharedRef<IDetailTreeNode> InTreeNode, bool bIsItemExpanded );
-
-	/**
-	 * Called when an item is expanded or collapsed in the detail tree
-	 *
-	 * @param InTreeNode		The node that was expanded or collapsed
-	 * @param bIsItemExpanded	True if the item is expanded, false if it is collapsed
-	 */
-	void OnItemExpansionChanged( TSharedRef<IDetailTreeNode> InTreeNode, bool bIsItemExpanded );
-
 	/**
 	 * Returns whether or not new objects need to be set. If the new objects being set are identical to the objects 
 	 * already in the details panel, nothing needs to be set
@@ -342,11 +127,6 @@ private:
 	 */
 	bool ShouldSetNewObjects( const TArray< TWeakObjectPtr< UObject > >& InObjects ) const;
 
-	/**
-	 * Adds an action to execute next tick
-	 */
-	void EnqueueDeferredAction( FSimpleDelegate& DeferredAction );
-
 	/** Updates the property map for access when customizing the details view.  Generates default layout for properties */
 	void UpdatePropertyMap();
 
@@ -390,9 +170,6 @@ private:
 	/** Called when the filter button is clicked */
 	void OnFilterButtonClicked();
 
-	/** Called when the locked button is clicked */
-	FReply OnLockButtonClicked();
-
 	/** Called when show only modified is clicked */
 	void OnShowOnlyModifiedClicked();
 	
@@ -414,12 +191,6 @@ private:
 	/** Called to get the visibility of the filter box */
 	EVisibility GetFilterBoxVisibility() const;
 
-	/** Called to get the visibility of the tree view */
-	EVisibility GetTreeVisibility() const;
-
-	/** Returns the name of the image used for the icon on the filter button */ 
-	const FSlateBrush* OnGetFilterButtonImageResource() const;
-
 	/** Returns the name of the image used for the icon on the locked button */
 	const FSlateBrush* OnGetLockButtonImageResource() const;
 	/**
@@ -442,42 +213,11 @@ private:
 	 */
 	void RestoreExpandedItems( TSharedPtr<FPropertyNode> InitialStartNode = NULL );
 
-	/**
-	 * Called when properties have finished changing (after PostEditChange is called)
-	 */
-	void NotifyFinishedChangingProperties(const FPropertyChangedEvent& PropertyChangedEvent);
-
-	/** Column width accessibility */
-	float OnGetLeftColumnWidth() const { return 1.0f - ColumnWidth; }
-	float OnGetRightColumnWidth() const { return ColumnWidth; }
-	void OnSetColumnWidth( float InWidth ) { ColumnWidth = InWidth; }
 private:
-	/** Map of nodes that are requesting an automatic expansion/collapse due to being filtered */
-	TMap< TSharedRef<IDetailTreeNode>, bool > FilteredNodesRequestingExpansionState;
-	/** A mapping unique classes being viewed to their variable names (for multiple of the same type being viewed)*/
-	TSet< TWeakObjectPtr<UStruct> > ClassesWithProperties;
-	/** Current set of expanded detail nodes (by path) that should be saved when the details panel closes */
-	TSet<FString> ExpandedDetailNodes;
-	/** A mapping of class (or struct) to properties in that struct (only top level, non-deeply nested properties appear in this map) */
-	FClassToPropertyMap ClassToPropertyMap;
-	/** A mapping of classes to detail layout delegates, called when querying for custom detail layouts in this instance of the details view only*/
-	FCustomDetailLayoutMap InstancedClassToDetailLayoutMap;
 	/** Information about the current set of selected actors */
 	FSelectedActorInfo SelectedActorInfo;
-	/** The user defined args for the details view */
-	FDetailsViewArgs DetailsViewArgs;
 	/** Selected objects for this detail view.  */
 	TArray< TWeakObjectPtr<UObject> > SelectedObjects;
-	/** Tree view */
-	TSharedPtr<SDetailTree> DetailTree;
-	/** Search box */
-	TSharedPtr<SWidget> SearchBox;
-	/** Customization class instances currently active in this view */
-	TArray< TSharedPtr<IDetailCustomization> > CustomizationClassInstances;
-	/** Customization instances that need to be destroyed when safe to do so */
-	TArray< TSharedPtr<IDetailCustomization> > CustomizationClassInstancesPendingDelete;
-	/** Root tree nodes visible in the tree */
-	TArray< TSharedRef<IDetailTreeNode> > RootTreeNodes;
 	/** Root tree node that needs to be destroyed when safe */
 	TSharedPtr<FObjectPropertyNode> RootNodePendingKill;
 	/** 
@@ -485,44 +225,12 @@ private:
 	 * It will only contain actors from when it was locked 
 	 */
 	TArray< TWeakObjectPtr<AActor> > SelectedActors;
-	/** The current detail layout based on selection */
-	TSharedPtr<class FDetailLayoutBuilderImpl> DetailLayout;
-	/** Settings for this view */
-	TSharedPtr<class FPropertyDetailsUtilities> PropertyUtilities;
 	/** The root property node of the property tree for a specific set of UObjects */
 	TSharedPtr<FObjectPropertyNode> RootPropertyNode;
-	/** The name area which is not recreated when selection changes */
-	TSharedPtr<class SDetailNameArea> NameArea;
-	/** Asset pool for rendering and managing asset thumbnails visible in this view*/
-	mutable TSharedPtr<class FAssetThumbnailPool> ThumbnailPool;
-	/** The current filter */
-	FDetailFilter CurrentFilter;
-	/** Delegate called to get generic details not specific to an object being viewed */
-	FOnGetDetailCustomizationInstance GenericLayoutDelegate;
-	/** Actions that should be executed next tick */
-	TArray<FSimpleDelegate> DeferredActions;
 	/** External property nodes which need to validated each tick */
 	TArray< TWeakPtr<FObjectPropertyNode> > ExternalRootPropertyNodes;
-	/** The property node that the color picker is currently editing. */
-	TWeakPtr< FPropertyNode > ColorPropertyNode;
 	/** Callback to send when the property view changes */
 	FOnObjectArrayChanged OnObjectArrayChanged;
-	/** Delegate executed to determine if a property should be visible */
-	FIsPropertyVisible IsPropertyVisibleDelegate;
-	/** Delegate called to see if a property editing is enabled */
-	FIsPropertyEditingEnabled IsPropertyEditingEnabledDelegate;
-	/** Delegate called when the details panel finishes editing a property (after post edit change is called) */
-	FOnFinishedChangingProperties OnFinishedChangingPropertiesDelegate;
-	/** Container for passing around column size data to rows in the tree (each row has a splitter which can affect the column size)*/
-	FDetailColumnSizeData ColumnSizeData;
-	/** The actual width of the right column.  The left column is 1-ColumnWidth */
-	float ColumnWidth;
-	/** True if there is an active filter (text in the filter box) */
-	bool bHasActiveFilter;
-	/** True if this property view is currently locked (I.E The objects being observed are not changed automatically due to user selection)*/
-	bool bIsLocked;
 	/** True if at least one viewed object is a CDO (blueprint editing) */
 	bool bViewingClassDefaultObject;
-	/** Whether or not this instance of the details view opened a color picker and it is not closed yet */
-	bool bHasOpenColorPicker;
 };
diff --git a/Engine/Source/Editor/PropertyEditor/Private/SDetailsViewBase.cpp b/Engine/Source/Editor/PropertyEditor/Private/SDetailsViewBase.cpp
new file mode 100644
--- /dev/null
+++ b/Engine/Source/Editor/PropertyEditor/Private/SDetailsViewBase.cpp
@@ -0,0 +1,196 @@
+
+#include "PropertyEditorPrivatePCH.h"
+#include "SDetailsViewBase.h"
+#include "AssetSelection.h"
+#include "PropertyNode.h"
+#include "ItemPropertyNode.h"
+#include "CategoryPropertyNode.h"
+#include "ObjectPropertyNode.h"
+#include "ScopedTransaction.h"
+#include "AssetThumbnail.h"
+#include "SDetailNameArea.h"
+#include "IPropertyUtilities.h"
+#include "PropertyEditorHelpers.h"
+#include "PropertyEditor.h"
+#include "PropertyDetailsUtilities.h"
+#include "SPropertyEditorEditInline.h"
+#include "ObjectEditorUtils.h"
+
+SDetailsViewBase::~SDetailsViewBase()
+{
+	if (ThumbnailPool.IsValid())
+	{
+		ThumbnailPool->ReleaseResources();
+	}
+}
+
+
+void SDetailsViewBase::OnGetChildrenForDetailTree(TSharedRef<IDetailTreeNode> InTreeNode, TArray< TSharedRef<IDetailTreeNode> >& OutChildren)
+{
+	InTreeNode->GetChildren(OutChildren);
+}
+
+TSharedRef<ITableRow> SDetailsViewBase::OnGenerateRowForDetailTree(TSharedRef<IDetailTreeNode> InTreeNode, const TSharedRef<STableViewBase>& OwnerTable)
+{
+	return InTreeNode->GenerateNodeWidget(OwnerTable, ColumnSizeData, PropertyUtilities.ToSharedRef());
+}
+
+void SDetailsViewBase::SetRootExpansionStates(const bool bExpand, const bool bRecurse)
+{
+	for (auto Iter = RootTreeNodes.CreateIterator(); Iter; ++Iter)
+	{
+		SetNodeExpansionState(*Iter, bExpand, bRecurse);
+	}
+}
+
+void SDetailsViewBase::SetNodeExpansionState(TSharedRef<IDetailTreeNode> InTreeNode, bool bIsItemExpanded, bool bRecursive)
+{
+	TArray< TSharedRef<IDetailTreeNode> > Children;
+	InTreeNode->GetChildren(Children);
+
+	if (Children.Num())
+	{
+		RequestItemExpanded(InTreeNode, bIsItemExpanded);
+		InTreeNode->OnItemExpansionChanged(bIsItemExpanded);
+
+		if (bRecursive)
+		{
+			for (int32 ChildIndex = 0; ChildIndex < Children.Num(); ++ChildIndex)
+			{
+				TSharedRef<IDetailTreeNode> Child = Children[ChildIndex];
+
+				SetNodeExpansionState(Child, bIsItemExpanded, bRecursive);
+			}
+		}
+	}
+}
+
+void SDetailsViewBase::SetNodeExpansionStateRecursive(TSharedRef<IDetailTreeNode> InTreeNode, bool bIsItemExpanded)
+{
+	SetNodeExpansionState(InTreeNode, bIsItemExpanded, true);
+}
+
+void SDetailsViewBase::OnItemExpansionChanged(TSharedRef<IDetailTreeNode> InTreeNode, bool bIsItemExpanded)
+{
+	SetNodeExpansionState(InTreeNode, bIsItemExpanded, false);
+}
+
+FReply SDetailsViewBase::OnLockButtonClicked()
+{
+	bIsLocked = !bIsLocked;
+	return FReply::Handled();
+}
+
+void SDetailsViewBase::HideFilterArea(bool bHide)
+{
+	DetailsViewArgs.bAllowSearch = !bHide;
+}
+
+EVisibility SDetailsViewBase::GetTreeVisibility() const
+{
+	return DetailLayout.IsValid() && DetailLayout->HasDetails() ? EVisibility::Visible : EVisibility::Collapsed;
+}
+
+/** Returns the image used for the icon on the filter button */
+const FSlateBrush* SDetailsViewBase::OnGetFilterButtonImageResource() const
+{
+	if (bHasActiveFilter)
+	{
+		return FEditorStyle::GetBrush(TEXT("PropertyWindow.FilterCancel"));
+	}
+	else
+	{
+		return FEditorStyle::GetBrush(TEXT("PropertyWindow.FilterSearch"));
+	}
+}
+
+void SDetailsViewBase::EnqueueDeferredAction(FSimpleDelegate& DeferredAction)
+{
+	DeferredActions.AddUnique(DeferredAction);
+}
+
+void SDetailsViewBase::OnColorPickerWindowClosed(const TSharedRef<SWindow>& Window)
+{
+	// A color picker window is no longer open
+	bHasOpenColorPicker = false;
+	ColorPropertyNode.Reset();
+}
+
+
+void SDetailsViewBase::SetIsPropertyVisibleDelegate(FIsPropertyVisible InIsPropertyVisible)
+{
+	IsPropertyVisibleDelegate = InIsPropertyVisible;
+}
+
+void SDetailsViewBase::SetIsPropertyEditingEnabledDelegate(FIsPropertyEditingEnabled IsPropertyEditingEnabled)
+{
+	IsPropertyEditingEnabledDelegate = IsPropertyEditingEnabled;
+}
+
+bool SDetailsViewBase::IsPropertyEditingEnabled() const
+{
+	// If the delegate is not bound assume property editing is enabled, otherwise ask the delegate
+	return !IsPropertyEditingEnabledDelegate.IsBound() || IsPropertyEditingEnabledDelegate.Execute();
+}
+
+void SDetailsViewBase::SetGenericLayoutDetailsDelegate(FOnGetDetailCustomizationInstance OnGetGenericDetails)
+{
+	GenericLayoutDelegate = OnGetGenericDetails;
+}
+
+TSharedPtr<FAssetThumbnailPool> SDetailsViewBase::GetThumbnailPool() const
+{
+	if (!ThumbnailPool.IsValid())
+	{
+		// Create a thumbnail pool for the view if it doesnt exist.  This does not use resources of no thumbnails are used
+		ThumbnailPool = MakeShareable(new FAssetThumbnailPool(50, TAttribute<bool>::Create(TAttribute<bool>::FGetter::CreateSP(this, &SDetailsView::IsHovered))));
+	}
+
+	return ThumbnailPool;
+}
+
+void SDetailsViewBase::NotifyFinishedChangingProperties(const FPropertyChangedEvent& PropertyChangedEvent)
+{
+	OnFinishedChangingPropertiesDelegate.Broadcast(PropertyChangedEvent);
+}
+
+void SDetailsViewBase::RequestItemExpanded(TSharedRef<IDetailTreeNode> TreeNode, bool bExpand)
+{
+	// Don't change expansion state if its already in that state
+	if (DetailTree->IsItemExpanded(TreeNode) != bExpand)
+	{
+		FilteredNodesRequestingExpansionState.Add(TreeNode, bExpand);
+	}
+}
+
+void SDetailsViewBase::RefreshTree()
+{
+	DetailTree->RequestTreeRefresh();
+}
+
+void SDetailsViewBase::SaveCustomExpansionState(const FString& NodePath, bool bIsExpanded)
+{
+	if (bIsExpanded)
+	{
+		ExpandedDetailNodes.Add(NodePath);
+	}
+	else
+	{
+		ExpandedDetailNodes.Remove(NodePath);
+	}
+}
+
+bool SDetailsViewBase::GetCustomSavedExpansionState(const FString& NodePath) const
+{
+	return ExpandedDetailNodes.Contains(NodePath);
+}
+
+bool SDetailsViewBase::IsPropertyVisible(const UProperty* Property) const
+{
+	return IsPropertyVisibleDelegate.IsBound() ? IsPropertyVisibleDelegate.Execute(Property) : true;
+}
+
+TSharedPtr<IPropertyUtilities> SDetailsViewBase::GetPropertyUtilities()
+{
+	return PropertyUtilities;
+}
diff --git a/Engine/Source/Editor/PropertyEditor/Private/SDetailsViewBase.h b/Engine/Source/Editor/PropertyEditor/Private/SDetailsViewBase.h
new file mode 100644
--- /dev/null
+++ b/Engine/Source/Editor/PropertyEditor/Private/SDetailsViewBase.h
@@ -0,0 +1,310 @@
+// Copyright 1998-2014 Epic Games, Inc. All Rights Reserved.
+
+#pragma once
+
+#include "IDetailsViewPrivate.h"
+#include "AssetSelection.h"
+#include "IPropertyUtilities.h"
+
+
+class FDetailCategoryImpl;
+class FDetailLayoutBuilderImpl;
+
+
+struct FPropertyNodeMap
+{
+	FPropertyNodeMap()
+		: ParentObjectProperty(NULL)
+	{}
+
+	/** Object property node which contains the properties in the node map */
+	FObjectPropertyNode* ParentObjectProperty;
+
+	/** Property name to property node map */
+	TMap<FName, TSharedPtr<FPropertyNode> > PropertyNameToNode;
+
+	bool Contains(FName PropertyName) const
+	{
+		return PropertyNameToNode.Contains(PropertyName);
+	}
+
+	void Add(FName PropertyName, TSharedPtr<FPropertyNode>& PropertyNode)
+	{
+		PropertyNameToNode.Add(PropertyName, PropertyNode);
+	}
+};
+
+typedef TArray< TSharedRef<class IDetailTreeNode> > FDetailNodeList;
+
+/** Mapping of categories to all top level item property nodes in that category */
+typedef TMap<FName, TSharedPtr<FDetailCategoryImpl> > FCategoryMap;
+
+/** Class to properties in that class */
+typedef TMap<FName, FPropertyNodeMap> FClassInstanceToPropertyMap;
+
+/** Class to properties in that class */
+typedef TMap<FName, FClassInstanceToPropertyMap> FClassToPropertyMap;
+
+typedef STreeView< TSharedRef<class IDetailTreeNode> > SDetailTree;
+
+/** Represents a filter which controls the visibility of items in the details view */
+struct FDetailFilter
+{
+	FDetailFilter()
+		: bShowOnlyModifiedProperties(false)
+		, bShowAllAdvanced(false)
+	{}
+
+	bool IsEmptyFilter() const { return FilterStrings.Num() == 0 && bShowOnlyModifiedProperties == false && bShowAllAdvanced == false; }
+
+	/** Any user search terms that items must match */
+	TArray<FString> FilterStrings;
+	/** If we should only show modified properties */
+	bool bShowOnlyModifiedProperties;
+	/** If we should show all advanced properties */
+	bool bShowAllAdvanced;
+};
+
+struct FDetailColumnSizeData
+{
+	TAttribute<float> LeftColumnWidth;
+	TAttribute<float> RightColumnWidth;
+	SSplitter::FOnSlotResized OnWidthChanged;
+
+	void SetColumnWidth(float InWidth) { OnWidthChanged.ExecuteIfBound(InWidth); }
+};
+
+class SDetailsViewBase : public IDetailsViewPrivate
+{
+public:
+	/**
+	* @return true of the details view can be updated from editor selection
+	*/
+	virtual bool IsUpdatable() const OVERRIDE
+	{
+		return DetailsViewArgs.bUpdatesFromSelection;
+	}
+
+	/** @return The identifier for this details view, or NAME_None is this view is anonymous */
+	virtual FName GetIdentifier() const OVERRIDE
+	{
+		return DetailsViewArgs.ViewIdentifier;
+	}
+
+	/**
+	 * Sets the visible state of the filter box/property grid area
+	 */
+	virtual void HideFilterArea(bool bHide) OVERRIDE;
+
+	virtual FOnFinishedChangingProperties& OnFinishedChangingProperties() OVERRIDE
+	{ 
+		return OnFinishedChangingPropertiesDelegate; 
+	}
+
+	/**
+	 * Called when the open color picker window associated with this details view is closed
+	 */
+	void OnColorPickerWindowClosed(const TSharedRef<SWindow>& Window);
+
+	/**
+	 * Returns true if the details view is locked and cant have its observed objects changed 
+	 */
+	virtual bool IsLocked() const OVERRIDE
+	{
+		return bIsLocked;
+	}
+
+	/**
+	 * Sets a delegate which is called to determine whether a specific property should be visible
+	 */
+	virtual void SetIsPropertyVisibleDelegate(FIsPropertyVisible InIsPropertyVisible) OVERRIDE;
+
+	/**
+	 * Sets a delegate to call to determine if the properties  editing is enabled
+	 */ 
+	virtual void SetIsPropertyEditingEnabledDelegate(FIsPropertyEditingEnabled IsPropertyEditingEnabled) OVERRIDE;
+
+	/**
+	 * @return true if property editing is enabled (based on the FIsPropertyEditingEnabled delegate)
+	 */
+	virtual bool IsPropertyEditingEnabled() const OVERRIDE;
+
+	/**
+	 * Requests that an item in the tree be expanded or collapsed
+	 *
+	 * @param TreeNode	The tree node to expand
+	 * @param bExpand	true if the item should be expanded, false otherwise
+	 */
+	void RequestItemExpanded(TSharedRef<IDetailTreeNode> TreeNode, bool bExpand) OVERRIDE;
+
+	/**
+	 * Sets the expansion state for a node and optionally all of its children
+	 *
+	 * @param InTreeNode		The node to change expansion state on
+	 * @param bIsItemExpanded	The new expansion state
+	 * @param bRecursive		Whether or not to apply the expansion change to any children
+	 */
+	void SetNodeExpansionState(TSharedRef<IDetailTreeNode> InTreeNode, bool bIsItemExpanded, bool bRecursive) OVERRIDE;
+
+	/**
+	 * Sets the expansion state all root nodes and optionally all of their children
+	 *
+	 * @param bExpand			The new expansion state
+	 * @param bRecurse			Whether or not to apply the expansion change to any children
+	 */
+	void SetRootExpansionStates(const bool bExpand, const bool bRecurse);
+
+	/**
+	 * Refreshes the detail's treeview
+	 */
+	void RefreshTree() OVERRIDE;
+
+	/**
+	 * Saves the expansion state of a tree node
+	 *
+	 * @param NodePath	The path to the detail node to save
+	 * @param bIsExpanded	true if the node is expanded, false otherwise
+	 */
+	void SaveCustomExpansionState(const FString& NodePath, bool bIsExpanded) OVERRIDE;
+
+	/**
+	 * Gets the saved expansion state of a tree node in this category	
+	 *
+	 * @param NodePath	The path to the detail node to get
+	 * @return true if the node should be expanded, false otherwise
+	 */	
+	bool GetCustomSavedExpansionState(const FString& NodePath) const OVERRIDE;
+
+	/**
+	 * Called when properties have finished changing (after PostEditChange is called)
+	 */
+	void NotifyFinishedChangingProperties(const FPropertyChangedEvent& PropertyChangedEvent);
+
+	/** Column width accessibility */
+	float OnGetLeftColumnWidth() const { return 1.0f - ColumnWidth; }
+	float OnGetRightColumnWidth() const { return ColumnWidth; }
+	void OnSetColumnWidth(float InWidth) { ColumnWidth = InWidth; }
+
+	/**
+	 * @return True if the property is visible
+	 */
+	bool IsPropertyVisible(const UProperty* Property) const OVERRIDE;
+
+	/**
+	* Sets a delegate which is regardless of the objects being viewed to lay out generic details not specific to any object
+	*/
+	virtual void SetGenericLayoutDetailsDelegate(FOnGetDetailCustomizationInstance OnGetGenericDetails) OVERRIDE;
+
+	/**
+	 * Adds an action to execute next tick
+	 */
+	void EnqueueDeferredAction( FSimpleDelegate& DeferredAction );
+
+	/**
+	 * @return The thumbnail pool that should be used for thumbnails being rendered in this view
+	 */
+	TSharedPtr<class FAssetThumbnailPool> GetThumbnailPool() const;
+
+	/**
+	 * Returns the property utilities for this view
+	 */
+	TSharedPtr<IPropertyUtilities> GetPropertyUtilities() OVERRIDE;
+
+	virtual ~SDetailsViewBase();
+
+protected:
+
+	/** Called to get the visibility of the tree view */
+	EVisibility GetTreeVisibility() const;
+
+	/** Returns the name of the image used for the icon on the filter button */
+	const FSlateBrush* OnGetFilterButtonImageResource() const;
+
+	/** Called when the locked button is clicked */
+	FReply OnLockButtonClicked();
+
+	/**
+	 * Called to recursively expand/collapse the children of the given item
+	 *
+	 * @param InTreeNode		The node that was expanded or collapsed
+	 * @param bIsItemExpanded	True if the item is expanded, false if it is collapsed
+	 */
+	void SetNodeExpansionStateRecursive( TSharedRef<IDetailTreeNode> InTreeNode, bool bIsItemExpanded );
+
+	/**
+	 * Called when an item is expanded or collapsed in the detail tree
+	 *
+	 * @param InTreeNode		The node that was expanded or collapsed
+	 * @param bIsItemExpanded	True if the item is expanded, false if it is collapsed
+	 */
+	void OnItemExpansionChanged( TSharedRef<IDetailTreeNode> InTreeNode, bool bIsItemExpanded );
+
+	/** 
+	 * Function called through a delegate on the TreeView to request children of a tree node 
+	 * 
+	 * @param InTreeNode		The tree node to get children from
+	 * @param OutChildren		The list of children of InTreeNode that should be visible 
+	 */
+	void OnGetChildrenForDetailTree( TSharedRef<IDetailTreeNode> InTreeNode, TArray< TSharedRef<IDetailTreeNode> >& OutChildren );
+
+	/**
+	 * Returns an SWidget used as the visual representation of a node in the treeview.                     
+	 */
+	TSharedRef<ITableRow> OnGenerateRowForDetailTree( TSharedRef<IDetailTreeNode> InTreeNode, const TSharedRef<STableViewBase>& OwnerTable );
+
+protected:
+	/** The user defined args for the details view */
+	FDetailsViewArgs DetailsViewArgs;
+	/** A mapping of class (or struct) to properties in that struct (only top level, non-deeply nested properties appear in this map) */
+	FClassToPropertyMap ClassToPropertyMap;
+	/** A mapping unique classes being viewed to their variable names (for multiple of the same type being viewed)*/
+	TSet< TWeakObjectPtr<UStruct> > ClassesWithProperties;
+	/** A mapping of classes to detail layout delegates, called when querying for custom detail layouts in this instance of the details view only*/
+	FCustomDetailLayoutMap InstancedClassToDetailLayoutMap;
+	/** The current detail layout based on selection */
+	TSharedPtr<class FDetailLayoutBuilderImpl> DetailLayout;
+	/** Search box */
+	TSharedPtr<SWidget> SearchBox;
+	/** Customization class instances currently active in this view */
+	TArray< TSharedPtr<IDetailCustomization> > CustomizationClassInstances;
+	/** Customization instances that need to be destroyed when safe to do so */
+	TArray< TSharedPtr<IDetailCustomization> > CustomizationClassInstancesPendingDelete;
+	/** Map of nodes that are requesting an automatic expansion/collapse due to being filtered */
+	TMap< TSharedRef<IDetailTreeNode>, bool > FilteredNodesRequestingExpansionState;
+	/** Current set of expanded detail nodes (by path) that should be saved when the details panel closes */
+	TSet<FString> ExpandedDetailNodes;
+	/** Tree view */
+	TSharedPtr<SDetailTree> DetailTree;
+	/** Root tree nodes visible in the tree */
+	TArray< TSharedRef<IDetailTreeNode> > RootTreeNodes;
+	/** Delegate executed to determine if a property should be visible */
+	FIsPropertyVisible IsPropertyVisibleDelegate;
+	/** Delegate called to see if a property editing is enabled */
+	FIsPropertyEditingEnabled IsPropertyEditingEnabledDelegate;
+	/** Delegate called when the details panel finishes editing a property (after post edit change is called) */
+	FOnFinishedChangingProperties OnFinishedChangingPropertiesDelegate;
+	/** Container for passing around column size data to rows in the tree (each row has a splitter which can affect the column size)*/
+	FDetailColumnSizeData ColumnSizeData;
+	/** The actual width of the right column.  The left column is 1-ColumnWidth */
+	float ColumnWidth;
+	/** True if there is an active filter (text in the filter box) */
+	bool bHasActiveFilter;
+	/** True if this property view is currently locked (I.E The objects being observed are not changed automatically due to user selection)*/
+	bool bIsLocked;
+	/** The property node that the color picker is currently editing. */
+	TWeakPtr< FPropertyNode > ColorPropertyNode;
+	/** Whether or not this instance of the details view opened a color picker and it is not closed yet */
+	bool bHasOpenColorPicker;
+	/** Settings for this view */
+	TSharedPtr<class IPropertyUtilities> PropertyUtilities;
+	/** The name area which is not recreated when selection changes */
+	TSharedPtr<class SDetailNameArea> NameArea;
+	/** Asset pool for rendering and managing asset thumbnails visible in this view*/
+	mutable TSharedPtr<class FAssetThumbnailPool> ThumbnailPool;
+	/** The current filter */
+	FDetailFilter CurrentFilter;
+	/** Delegate called to get generic details not specific to an object being viewed */
+	FOnGetDetailCustomizationInstance GenericLayoutDelegate;
+	/** Actions that should be executed next tick */
+	TArray<FSimpleDelegate> DeferredActions;
+};
\ No newline at end of file
