PropertyEditor refactored. It's ready for StructureDefaultView.

diff --git a/Engine/Source/Editor/PropertyEditor/Private/IDetailsViewPrivate.h b/Engine/Source/Editor/PropertyEditor/Private/IDetailsViewPrivate.h
--- a/Engine/Source/Editor/PropertyEditor/Private/IDetailsViewPrivate.h
+++ b/Engine/Source/Editor/PropertyEditor/Private/IDetailsViewPrivate.h
@@ -89,4 +89,21 @@ public:
 	 */
 	virtual void AddExternalRootPropertyNode(TSharedRef<FObjectPropertyNode> ExternalRootNode) = 0;
 
+	/**
+	 * Creates the color picker window for this property view.
+	 *
+	 * @param PropertyEditor				The slate property node to edit.
+	 * @param bUseAlpha			Whether or not alpha is supported
+	 */
+	virtual void CreateColorPickerWindow(const TSharedRef< class FPropertyEditor >& PropertyEditor, bool bUseAlpha) = 0;
+
+	/**
+	 * Adds an action to execute next tick
+	 */
+	virtual void EnqueueDeferredAction(FSimpleDelegate& DeferredAction) = 0;
+
+	/**
+	 * Called when properties have finished changing (after PostEditChange is called)
+	 */
+	virtual void NotifyFinishedChangingProperties(const FPropertyChangedEvent& PropertyChangedEvent) = 0;
 };
diff --git a/Engine/Source/Editor/PropertyEditor/Private/SDetailsView.cpp b/Engine/Source/Editor/PropertyEditor/Private/SDetailsView.cpp
--- a/Engine/Source/Editor/PropertyEditor/Private/SDetailsView.cpp
+++ b/Engine/Source/Editor/PropertyEditor/Private/SDetailsView.cpp
@@ -237,33 +237,6 @@ FReply SDetailsView::OnOpenRawPropertyEditorClicked()
 	return FReply::Handled();
 }
 
-void SDetailsView::OnShowOnlyModifiedClicked()
-{
-	CurrentFilter.bShowOnlyModifiedProperties = !CurrentFilter.bShowOnlyModifiedProperties;
-
-	UpdateFilteredDetails();
-}
-
-void SDetailsView::OnShowAllAdvancedClicked()
-{
-	CurrentFilter.bShowAllAdvanced = !CurrentFilter.bShowAllAdvanced;
-
-	UpdateFilteredDetails();
-}
-
-/** Called when the filter text changes.  This filters specific property nodes out of view */
-void SDetailsView::OnFilterTextChanged( const FText& InFilterText )
-{
-	FString InFilterString = InFilterText.ToString();
-	InFilterString.Trim().TrimTrailing();
-
-	// Was the filter just cleared
-	bool bFilterCleared = InFilterString.Len() == 0 && CurrentFilter.FilterStrings.Num() > 0;
-
-	FilterView( InFilterString );
-
-}
-
 EVisibility SDetailsView::GetActorNameAreaVisibility() const
 {
 	const bool bVisible = !DetailsViewArgs.bHideActorNameArea && !bViewingClassDefaultObject;
@@ -1093,28 +1066,6 @@ void SDetailsView::QueryCustomDetailLayout( FDetailLayoutBuilderImpl& CustomDeta
 	}
 }
 
-/** 
- * Hides or shows properties based on the passed in filter text
- * 
- * @param InFilterText	The filter text
- */
-void SDetailsView::FilterView( const FString& InFilterText )
-{
-	TArray<FString> CurrentFilterStrings;
-
-	FString ParseString = InFilterText;
-	// Remove whitespace from the front and back of the string
-	ParseString.Trim();
-	ParseString.TrimTrailing();
-	ParseString.ParseIntoArray(&CurrentFilterStrings, TEXT(" "), true);
-
-	bHasActiveFilter = CurrentFilterStrings.Num() > 0;
-
-	CurrentFilter.FilterStrings = CurrentFilterStrings;
-
-	UpdateFilteredDetails();
-}
-
 void SDetailsView::UpdateFilteredDetails()
 {
 	RootPropertyNode->FilterNodes( CurrentFilter.FilterStrings );
diff --git a/Engine/Source/Editor/PropertyEditor/Private/SDetailsView.h b/Engine/Source/Editor/PropertyEditor/Private/SDetailsView.h
--- a/Engine/Source/Editor/PropertyEditor/Private/SDetailsView.h
+++ b/Engine/Source/Editor/PropertyEditor/Private/SDetailsView.h
@@ -58,10 +58,7 @@ public:
 	 * @param PropertyEditor				The slate property node to edit.
 	 * @param bUseAlpha			Whether or not alpha is supported
 	 */
-	virtual void CreateColorPickerWindow(const TSharedRef< class FPropertyEditor >& PropertyEditor, bool bUseAlpha);
-
-	/** Returns the notify hook to use when properties change */
-	virtual FNotifyHook* GetNotifyHook() const OVERRIDE { return DetailsViewArgs.NotifyHook; }
+	virtual void CreateColorPickerWindow(const TSharedRef< class FPropertyEditor >& PropertyEditor, bool bUseAlpha) OVERRIDE;
 
 	/** Sets the callback for when the property view changes */
 	virtual void SetOnObjectArrayChanged( FOnObjectArrayChanged OnObjectArrayChangedDelegate);
@@ -155,36 +152,14 @@ private:
 	 */
 	void QueryCustomDetailLayout( class FDetailLayoutBuilderImpl& CustomDetailLayout );
 
-	/** 
-	 * Hides or shows properties based on the passed in filter text
-	 * 
-	 * @param InFilterText	The filter text
-	 */
-	void FilterView( const FString& InFilterText );
-
 	/**
 	 * Updates the details with the passed in filter                                                              
 	 */
-	void UpdateFilteredDetails();
+	virtual void UpdateFilteredDetails() OVERRIDE;
 
 	/** Called when the filter button is clicked */
 	void OnFilterButtonClicked();
 
-	/** Called when show only modified is clicked */
-	void OnShowOnlyModifiedClicked();
-	
-	/** Called when show all advanced is clicked */
-	void OnShowAllAdvancedClicked();
-
-	/** @return true if show only modified is checked */
-	bool IsShowOnlyModifiedChecked() const { return CurrentFilter.bShowOnlyModifiedProperties; }
-
-	/** @return true if show all advanced is checked */
-	bool IsShowAllAdvancedChecked() const { return CurrentFilter.bShowAllAdvanced; }
-
-	/** Called when the filter text changes.  This filters specific property nodes out of view */
-	void OnFilterTextChanged( const FText& InFilterText );
-
 	/** Called to get the visibility of the actor name area */
 	EVisibility GetActorNameAreaVisibility() const;
 
diff --git a/Engine/Source/Editor/PropertyEditor/Private/SDetailsViewBase.cpp b/Engine/Source/Editor/PropertyEditor/Private/SDetailsViewBase.cpp
--- a/Engine/Source/Editor/PropertyEditor/Private/SDetailsViewBase.cpp
+++ b/Engine/Source/Editor/PropertyEditor/Private/SDetailsViewBase.cpp
@@ -194,3 +194,52 @@ TSharedPtr<IPropertyUtilities> SDetailsViewBase::GetPropertyUtilities()
 {
 	return PropertyUtilities;
 }
+
+void SDetailsViewBase::OnShowOnlyModifiedClicked()
+{
+	CurrentFilter.bShowOnlyModifiedProperties = !CurrentFilter.bShowOnlyModifiedProperties;
+
+	UpdateFilteredDetails();
+}
+
+void SDetailsViewBase::OnShowAllAdvancedClicked()
+{
+	CurrentFilter.bShowAllAdvanced = !CurrentFilter.bShowAllAdvanced;
+
+	UpdateFilteredDetails();
+}
+
+/** Called when the filter text changes.  This filters specific property nodes out of view */
+void SDetailsViewBase::OnFilterTextChanged(const FText& InFilterText)
+{
+	FString InFilterString = InFilterText.ToString();
+	InFilterString.Trim().TrimTrailing();
+
+	// Was the filter just cleared
+	bool bFilterCleared = InFilterString.Len() == 0 && CurrentFilter.FilterStrings.Num() > 0;
+
+	FilterView(InFilterString);
+
+}
+
+/** 
+ * Hides or shows properties based on the passed in filter text
+ * 
+ * @param InFilterText	The filter text
+ */
+void SDetailsViewBase::FilterView(const FString& InFilterText)
+{
+	TArray<FString> CurrentFilterStrings;
+
+	FString ParseString = InFilterText;
+	// Remove whitespace from the front and back of the string
+	ParseString.Trim();
+	ParseString.TrimTrailing();
+	ParseString.ParseIntoArray(&CurrentFilterStrings, TEXT(" "), true);
+
+	bHasActiveFilter = CurrentFilterStrings.Num() > 0;
+
+	CurrentFilter.FilterStrings = CurrentFilterStrings;
+
+	UpdateFilteredDetails();
+}
\ No newline at end of file
diff --git a/Engine/Source/Editor/PropertyEditor/Private/SDetailsViewBase.h b/Engine/Source/Editor/PropertyEditor/Private/SDetailsViewBase.h
--- a/Engine/Source/Editor/PropertyEditor/Private/SDetailsViewBase.h
+++ b/Engine/Source/Editor/PropertyEditor/Private/SDetailsViewBase.h
@@ -178,7 +178,7 @@ public:
 	/**
 	 * Called when properties have finished changing (after PostEditChange is called)
 	 */
-	void NotifyFinishedChangingProperties(const FPropertyChangedEvent& PropertyChangedEvent);
+	virtual void NotifyFinishedChangingProperties(const FPropertyChangedEvent& PropertyChangedEvent) OVERRIDE;
 
 	/** Column width accessibility */
 	float OnGetLeftColumnWidth() const { return 1.0f - ColumnWidth; }
@@ -198,7 +198,7 @@ public:
 	/**
 	 * Adds an action to execute next tick
 	 */
-	void EnqueueDeferredAction( FSimpleDelegate& DeferredAction );
+	virtual void EnqueueDeferredAction(FSimpleDelegate& DeferredAction) OVERRIDE;
 
 	/**
 	 * @return The thumbnail pool that should be used for thumbnails being rendered in this view
@@ -210,6 +210,12 @@ public:
 	 */
 	TSharedPtr<IPropertyUtilities> GetPropertyUtilities() OVERRIDE;
 
+	/** Returns the notify hook to use when properties change */
+	virtual FNotifyHook* GetNotifyHook() const OVERRIDE
+	{ 
+		return DetailsViewArgs.NotifyHook; 
+	}
+
 	virtual ~SDetailsViewBase();
 
 protected:
@@ -252,6 +258,30 @@ protected:
 	 */
 	TSharedRef<ITableRow> OnGenerateRowForDetailTree( TSharedRef<IDetailTreeNode> InTreeNode, const TSharedRef<STableViewBase>& OwnerTable );
 
+	/** @return true if show only modified is checked */
+	bool IsShowOnlyModifiedChecked() const { return CurrentFilter.bShowOnlyModifiedProperties; }
+
+	/** @return true if show all advanced is checked */
+	bool IsShowAllAdvancedChecked() const { return CurrentFilter.bShowAllAdvanced; }
+
+	/** Called when show only modified is clicked */
+	void OnShowOnlyModifiedClicked();
+
+	/** Called when show all advanced is clicked */
+	void OnShowAllAdvancedClicked();
+
+	virtual void UpdateFilteredDetails() = 0;
+
+	/** Called when the filter text changes.  This filters specific property nodes out of view */
+	void OnFilterTextChanged(const FText& InFilterText);
+
+	/** 
+	 * Hides or shows properties based on the passed in filter text
+	 * 
+	 * @param InFilterText	The filter text
+	 */
+	void FilterView( const FString& InFilterText );
+
 protected:
 	/** The user defined args for the details view */
 	FDetailsViewArgs DetailsViewArgs;
diff --git a/Engine/Source/Editor/PropertyEditor/Private/UserInterface/PropertyDetails/PropertyDetailsUtilities.cpp b/Engine/Source/Editor/PropertyEditor/Private/UserInterface/PropertyDetails/PropertyDetailsUtilities.cpp
--- a/Engine/Source/Editor/PropertyEditor/Private/UserInterface/PropertyDetails/PropertyDetailsUtilities.cpp
+++ b/Engine/Source/Editor/PropertyEditor/Private/UserInterface/PropertyDetails/PropertyDetailsUtilities.cpp
@@ -3,7 +3,7 @@
 #include "PropertyEditorPrivatePCH.h"
 #include "PropertyDetailsUtilities.h"
 
-FPropertyDetailsUtilities::FPropertyDetailsUtilities( SDetailsView& InDetailsView )
+FPropertyDetailsUtilities::FPropertyDetailsUtilities(IDetailsViewPrivate& InDetailsView)
 	: DetailsView( InDetailsView )
 {
 }
diff --git a/Engine/Source/Editor/PropertyEditor/Private/UserInterface/PropertyDetails/PropertyDetailsUtilities.h b/Engine/Source/Editor/PropertyEditor/Private/UserInterface/PropertyDetails/PropertyDetailsUtilities.h
--- a/Engine/Source/Editor/PropertyEditor/Private/UserInterface/PropertyDetails/PropertyDetailsUtilities.h
+++ b/Engine/Source/Editor/PropertyEditor/Private/UserInterface/PropertyDetails/PropertyDetailsUtilities.h
@@ -6,7 +6,7 @@
 class FPropertyDetailsUtilities : public IPropertyUtilities
 {
 public:
-	FPropertyDetailsUtilities( SDetailsView& InDetailsView );
+	FPropertyDetailsUtilities(IDetailsViewPrivate& InDetailsView);
 	/** IPropertyUtilities interface */
 	virtual class FNotifyHook* GetNotifyHook() const OVERRIDE;
 	virtual bool AreFavoritesEnabled() const OVERRIDE;
@@ -18,5 +18,5 @@ public:
 	virtual TSharedPtr<class FAssetThumbnailPool> GetThumbnailPool() const OVERRIDE;
 	virtual void NotifyFinishedChangingProperties(const FPropertyChangedEvent& PropertyChangedEvent) OVERRIDE;
 private:
-	SDetailsView& DetailsView;
+	IDetailsViewPrivate& DetailsView;
 };
