Check ParentCategory is valid before attempting to Pin() and call methods on it.
#ttp 336990 EDITOR: CRASH: Project crashing when placing bsp into level
reviewer Thomas.Sarkanen

diff --git a/Engine/Source/Editor/PropertyEditor/Private/DetailItemNode.cpp b/Engine/Source/Editor/PropertyEditor/Private/DetailItemNode.cpp
--- a/Engine/Source/Editor/PropertyEditor/Private/DetailItemNode.cpp
+++ b/Engine/Source/Editor/PropertyEditor/Private/DetailItemNode.cpp
@@ -183,6 +183,11 @@ void FDetailItemNode::GenerateChildren( bool bUpdateFilteredNodes )
 {
 	Children.Empty();
 
+	if (!ParentCategory.IsValid())
+	{
+		return;
+	}
+
 	if( Customization.HasPropertyNode() )
 	{
 		Customization.PropertyRow->OnGenerateChildren( Children );
