Making the property type customization work for UObjects.

diff --git a/Engine/Source/Editor/PropertyEditor/Private/PropertyEditorModule.cpp b/Engine/Source/Editor/PropertyEditor/Private/PropertyEditorModule.cpp
--- a/Engine/Source/Editor/PropertyEditor/Private/PropertyEditorModule.cpp
+++ b/Engine/Source/Editor/PropertyEditor/Private/PropertyEditorModule.cpp
@@ -583,6 +583,9 @@ FPropertyTypeLayoutCallback FPropertyEditorModule::GetPropertyTypeCustomization(
 		const UByteProperty* ByteProperty = Cast<UByteProperty>(Property);
 		const bool bEnumProperty = ByteProperty && ByteProperty->Enum;
 
+		const UObjectProperty* ObjectProperty = Cast<UObjectProperty>(Property);
+		const bool bObjectProperty = ObjectProperty != NULL && ObjectProperty->PropertyClass != NULL;
+
 		FName PropertyTypeName;
 		if( bStructProperty )
 		{
@@ -592,9 +595,13 @@ FPropertyTypeLayoutCallback FPropertyEditorModule::GetPropertyTypeCustomization(
 		{
 			PropertyTypeName = ByteProperty->Enum->GetFName();
 		}
+		else if ( bObjectProperty )
+		{
+			PropertyTypeName = ObjectProperty->PropertyClass->GetFName();
+		}
 		else
 		{
-			PropertyTypeName = Property->GetFName();
+			PropertyTypeName = Property->GetClass()->GetFName();
 		}
 
 
