Fix abiguous ? operator usage on OS X.

diff --git a/Engine/Source/Editor/PropertyEditor/Public/PropertyEditorModule.h b/Engine/Source/Editor/PropertyEditor/Public/PropertyEditorModule.h
--- a/Engine/Source/Editor/PropertyEditor/Public/PropertyEditorModule.h
+++ b/Engine/Source/Editor/PropertyEditor/Public/PropertyEditorModule.h
@@ -240,7 +240,7 @@ struct FPropertyTypeLayoutCallback
 
 	TSharedRef<IPropertyTypeCustomization> GetCustomizationInstance() const
 	{
-		return PropertyTypeLayoutDelegate.IsBound() ? PropertyTypeLayoutDelegate.Execute() : DeprecatedLayoutDelegate.Execute();
+		return PropertyTypeLayoutDelegate.IsBound() ? PropertyTypeLayoutDelegate.Execute() : static_cast<TSharedRef<IPropertyTypeCustomization>>(DeprecatedLayoutDelegate.Execute());
 	}
 };
 
