More fixup of deprecated IStructCustomization interface

diff --git a/Engine/Source/Editor/BehaviorTreeEditor/Private/BehaviorTreeEditorModule.cpp b/Engine/Source/Editor/BehaviorTreeEditor/Private/BehaviorTreeEditorModule.cpp
--- a/Engine/Source/Editor/BehaviorTreeEditor/Private/BehaviorTreeEditorModule.cpp
+++ b/Engine/Source/Editor/BehaviorTreeEditor/Private/BehaviorTreeEditorModule.cpp
@@ -88,7 +88,7 @@ void FBehaviorTreeEditorModule::ShutdownModule()
 	if (FModuleManager::Get().IsModuleLoaded("PropertyEditor"))
 	{
 		FPropertyEditorModule& PropertyModule = FModuleManager::LoadModuleChecked<FPropertyEditorModule>("PropertyEditor");
-		PropertyModule.UnregisterStructPropertyLayout( "BlackboardKeySelector" );
+		PropertyModule.UnregisterCustomPropertyTypeLayout( "BlackboardKeySelector" );
 		PropertyModule.UnregisterCustomPropertyLayout( "BTDecorator_Blackboard" );
 		PropertyModule.UnregisterCustomPropertyLayout( "BTDecorator" );
 		PropertyModule.NotifyCustomizationModuleChanged();
diff --git a/Engine/Source/Editor/EnvironmentQueryEditor/Private/EnvironmentQueryEditorModule.cpp b/Engine/Source/Editor/EnvironmentQueryEditor/Private/EnvironmentQueryEditorModule.cpp
--- a/Engine/Source/Editor/EnvironmentQueryEditor/Private/EnvironmentQueryEditorModule.cpp
+++ b/Engine/Source/Editor/EnvironmentQueryEditor/Private/EnvironmentQueryEditorModule.cpp
@@ -86,11 +86,12 @@ void FEnvironmentQueryEditorModule::ShutdownModule()
 	if (FModuleManager::Get().IsModuleLoaded("PropertyEditor"))
 	{
 		FPropertyEditorModule& PropertyModule = FModuleManager::LoadModuleChecked<FPropertyEditorModule>("PropertyEditor");
-		PropertyModule.UnregisterStructPropertyLayout( "EnvFloatParam" );
-		PropertyModule.UnregisterStructPropertyLayout( "EnvIntParam" );
-		PropertyModule.UnregisterStructPropertyLayout( "EnvBoolParam" );
-		PropertyModule.UnregisterStructPropertyLayout( "EnvDirection" );
-		PropertyModule.UnregisterStructPropertyLayout( "EnvTraceData" );
+		PropertyModule.UnregisterCustomPropertyTypeLayout( "EnvFloatParam" );
+		PropertyModule.UnregisterCustomPropertyTypeLayout( "EnvIntParam" );
+		PropertyModule.UnregisterCustomPropertyTypeLayout( "EnvBoolParam" );
+		PropertyModule.UnregisterCustomPropertyTypeLayout( "EnvDirection" );
+		PropertyModule.UnregisterCustomPropertyTypeLayout( "EnvTraceData" );
+
 		PropertyModule.UnregisterCustomPropertyLayout( "EnvQueryTest" );
 		PropertyModule.NotifyCustomizationModuleChanged();
 	}
diff --git a/Engine/Source/Editor/Persona/Private/PersonaModule.cpp b/Engine/Source/Editor/Persona/Private/PersonaModule.cpp
--- a/Engine/Source/Editor/Persona/Private/PersonaModule.cpp
+++ b/Engine/Source/Editor/Persona/Private/PersonaModule.cpp
@@ -78,8 +78,8 @@ void FPersonaModule::ShutdownModule()
 		PropertyModule.UnregisterCustomPropertyLayout("AnimNotifyState");
 		PropertyModule.UnregisterCustomPropertyLayout("AnimGraphNode_Base");
 
-		PropertyModule.UnregisterStructPropertyLayout("InputScaleBias");
-		PropertyModule.UnregisterStructPropertyLayout("BoneReference");
+		PropertyModule.UnregisterCustomPropertyTypeLayout("InputScaleBias");
+		PropertyModule.UnregisterCustomPropertyTypeLayout("BoneReference");
 	}
 }
 
diff --git a/Engine/Source/Editor/PropertyEditor/Private/DetailCustomBuilderRow.h b/Engine/Source/Editor/PropertyEditor/Private/DetailCustomBuilderRow.h
--- a/Engine/Source/Editor/PropertyEditor/Private/DetailCustomBuilderRow.h
+++ b/Engine/Source/Editor/PropertyEditor/Private/DetailCustomBuilderRow.h
@@ -2,6 +2,9 @@
 
 #pragma once
 
+class IDetailCustomNodeBuilder;
+class FDetailItemNode;
+class FDetailCategoryImpl;
 
 class FDetailCustomBuilderRow : public TSharedFromThis<FDetailCustomBuilderRow>
 {
@@ -12,7 +15,7 @@ public:
 	bool RequiresTick() const;
 	bool HasColumns() const;
 	bool ShowOnlyChildren() const;
-	void OnItemNodeInitialized( TSharedRef<class FDetailItemNode> InTreeNode, TSharedRef<class FDetailCategoryImpl> InParentCategory, const TAttribute<bool>& InIsParentEnabled );
+	void OnItemNodeInitialized( TSharedRef<FDetailItemNode> InTreeNode, TSharedRef<FDetailCategoryImpl> InParentCategory, const TAttribute<bool>& InIsParentEnabled );
 	FName GetCustomBuilderName() const;
 	void OnGenerateChildren( FDetailNodeList& OutChildren );
 	bool IsInitiallyCollapsed() const;
diff --git a/Engine/Source/Editor/PropertyEditor/Private/PropertyEditorModule.cpp b/Engine/Source/Editor/PropertyEditor/Private/PropertyEditorModule.cpp
--- a/Engine/Source/Editor/PropertyEditor/Private/PropertyEditorModule.cpp
+++ b/Engine/Source/Editor/PropertyEditor/Private/PropertyEditorModule.cpp
@@ -24,6 +24,10 @@
 
 IMPLEMENT_MODULE( FPropertyEditorModule, PropertyEditor );
 
+TSharedRef<IPropertyTypeCustomization> FPropertyTypeLayoutCallback::GetCustomizationInstance() const
+{
+	return PropertyTypeLayoutDelegate.IsBound() ? PropertyTypeLayoutDelegate.Execute() : StaticCastSharedRef<IPropertyTypeCustomization>(DeprecatedLayoutDelegate.Execute());
+}
 
 void FPropertyTypeLayoutCallbackList::Add( const FPropertyTypeLayoutCallback& NewCallback )
 {
diff --git a/Engine/Source/Editor/PropertyEditor/Public/IDetailChildrenBuilder.h b/Engine/Source/Editor/PropertyEditor/Public/IDetailChildrenBuilder.h
new file mode 100644
--- /dev/null
+++ b/Engine/Source/Editor/PropertyEditor/Public/IDetailChildrenBuilder.h
@@ -0,0 +1,56 @@
+// Copyright 1998-2014 Epic Games, Inc. All Rights Reserved.
+
+#pragma once
+
+class IDetailCustomNodeBuilder;
+class IPropertyHandle;
+
+/**
+ * Builder for adding children to a detail customization
+ */
+class IDetailChildrenBuilder
+{
+public:
+	virtual ~IDetailChildrenBuilder(){}
+
+	/**
+	 * Adds a custom builder as a child
+	 *
+	 * @param InCustomBuilder		The custom builder to add
+	 */
+	virtual IDetailChildrenBuilder& AddChildCustomBuilder( TSharedRef<IDetailCustomNodeBuilder> InCustomBuilder ) = 0;
+
+	/**
+	 * Adds a group to the category
+	 *
+	 * @param GroupName	The name of the group 
+	 * @param LocalizedDisplayName	The display name of the group
+	 * @param true if the group should appear in the advanced section of the category
+	 */
+	virtual class IDetailGroup& AddChildGroup( FName GroupName, const FString& LocalizedDisplayName ) = 0;
+
+	/**
+	 * Adds new custom content as a child to the struct
+	 *
+	 * @param SearchString	Search string that will be matched when users search in the details panel.  If the search string doesnt match what the user types, this row will be hidden
+	 * @return A row that accepts widgets
+	 */
+	virtual class FDetailWidgetRow& AddChildContent( const FString& SearchString ) = 0;
+
+	/**
+	 * Adds a property to the struct
+	 * 
+	 * @param PropertyHandle	The handle to the property to add
+	 * @return An interface to the property row that can be used to customize the appearance of the property
+	 */
+	virtual class IDetailPropertyRow& AddChildProperty( TSharedRef<IPropertyHandle> PropertyHandle ) = 0;
+
+	/**
+	 * Generates a value widget from a customized struct
+	 * If the customized struct has no value widget an empty widget will be returned
+	 *
+	 * @param StructPropertyHandle	The handle to the struct property to generate the value widget from
+	 */
+	virtual TSharedRef<SWidget> GenerateStructValueWidget(TSharedRef<IPropertyHandle> StructPropertyHandle) = 0;
+
+};
\ No newline at end of file
diff --git a/Engine/Source/Editor/PropertyEditor/Public/IDetailCustomNodeBuilder.h b/Engine/Source/Editor/PropertyEditor/Public/IDetailCustomNodeBuilder.h
--- a/Engine/Source/Editor/PropertyEditor/Public/IDetailCustomNodeBuilder.h
+++ b/Engine/Source/Editor/PropertyEditor/Public/IDetailCustomNodeBuilder.h
@@ -2,6 +2,9 @@
 
 #pragma once
 
+class IDetailChildrenBuilder;
+class FDetailWidgetRow;
+
 /**
  * A custom node that can be given to a details panel category to customize widgets
  */
diff --git a/Engine/Source/Editor/PropertyEditor/Public/IDetailCustomization.h b/Engine/Source/Editor/PropertyEditor/Public/IDetailCustomization.h
new file mode 100644
--- /dev/null
+++ b/Engine/Source/Editor/PropertyEditor/Public/IDetailCustomization.h
@@ -0,0 +1,16 @@
+// Copyright 1998-2014 Epic Games, Inc. All Rights Reserved.
+#pragma once
+
+class IDetailLayoutBuilder;
+
+/** 
+ * Interface for any class that lays out details for a specific class
+ */
+class IDetailCustomization : public TSharedFromThis<IDetailCustomization>
+{
+public:
+	virtual ~IDetailCustomization() {}
+
+	/** Called when details should be customized */
+	virtual void CustomizeDetails( IDetailLayoutBuilder& DetailBuilder ) = 0;
+};
diff --git a/Engine/Source/Editor/PropertyEditor/Public/IDetailGroup.h b/Engine/Source/Editor/PropertyEditor/Public/IDetailGroup.h
--- a/Engine/Source/Editor/PropertyEditor/Public/IDetailGroup.h
+++ b/Engine/Source/Editor/PropertyEditor/Public/IDetailGroup.h
@@ -2,6 +2,10 @@
 
 #pragma once
 
+class IDetailPropertyRow;
+class FDetailWidgetRow;
+class IPropertyHandle;
+
 /**
  * A group in the details panel that can have children                                                              
  */
diff --git a/Engine/Source/Editor/PropertyEditor/Public/IPropertyTypeCustomization.h b/Engine/Source/Editor/PropertyEditor/Public/IPropertyTypeCustomization.h
new file mode 100644
--- /dev/null
+++ b/Engine/Source/Editor/PropertyEditor/Public/IPropertyTypeCustomization.h
@@ -0,0 +1,106 @@
+// Copyright 1998-2014 Epic Games, Inc. All Rights Reserved.
+
+#pragma once
+
+class IPropertyHandle;
+class FDetailWidgetRow;
+class IPropertyTypeCustomizationUtils;
+class IDetailChildrenBuilder;
+class FAssetThumbnailPool;
+class IPropertyUtilities;
+
+/**
+ * Base class for property type customizations
+ */
+class IPropertyTypeCustomization : public TSharedFromThis<IPropertyTypeCustomization>
+{
+public:
+	virtual ~IPropertyTypeCustomization() {}
+
+	/**
+	 * Called when the header of the property (the row in the details panel where the property is shown)
+	 * If nothing is added to the row, the header is not displayed
+	 *
+	 * @param PropertyHandle			Handle to the property being customized
+	 * @param HeaderRow					A row that widgets can be added to
+	 * @param StructCustomizationUtils	Utilities for customization
+	 */
+	virtual void CustomizeHeader( TSharedRef<IPropertyHandle> PropertyHandle, FDetailWidgetRow& HeaderRow, IPropertyTypeCustomizationUtils& CustomizationUtils ) = 0;
+
+	/**
+	 * Called when the children of the property should be customized or extra rows added
+	 *
+	 * @param PropertyHandle			Handle to the property being customized
+	 * @param StructBuilder				A builder for adding children
+	 * @param StructCustomizationUtils	Utilities for customization
+	 */
+	virtual void CustomizeChildren( TSharedRef<IPropertyHandle> PropertyHandle, IDetailChildrenBuilder& ChildBuilder, IPropertyTypeCustomizationUtils& CustomizationUtils ) = 0;
+};
+
+
+/**
+ * Utilities for property type customization
+ */
+class IPropertyTypeCustomizationUtils
+{
+public:
+	virtual ~IPropertyTypeCustomizationUtils(){};
+
+	/**
+	 * @return the font used for properties and details
+	 */ 
+	static FSlateFontInfo GetRegularFont() { return FEditorStyle::GetFontStyle( TEXT("PropertyWindow.NormalFont") ); }
+
+	/**
+	 * @return the bold font used for properties and details
+	 */ 
+	static FSlateFontInfo GetBoldFont() { return FEditorStyle::GetFontStyle( TEXT("PropertyWindow.BoldFont") ); }
+	
+	/**
+	 * Gets the thumbnail pool that should be used for rendering thumbnails in the struct                  
+	 */
+	virtual TSharedPtr<FAssetThumbnailPool> GetThumbnailPool() const = 0;
+
+	/**
+	*	@return the utilities various widgets need access to certain features of PropertyDetails
+	*/
+	virtual TSharedPtr<IPropertyUtilities> GetPropertyUtilities() const { return NULL; }
+};
+
+
+/** Deprecated IStructCustomizationUtils interface */
+typedef IPropertyTypeCustomizationUtils IStructCustomizationUtils;
+
+/** Deprecated IStructCustomization interface */
+class IStructCustomization : public IPropertyTypeCustomization
+{
+public:
+	virtual void CustomizeHeader( TSharedRef<IPropertyHandle> PropertyHandle, FDetailWidgetRow& HeaderRow, IPropertyTypeCustomizationUtils& CustomizationUtils )
+	{
+		CustomizeStructHeader( PropertyHandle, HeaderRow, CustomizationUtils );
+	}
+
+	virtual void CustomizeChildren( TSharedRef<IPropertyHandle> PropertyHandle, IDetailChildrenBuilder& ChildBuilder, IPropertyTypeCustomizationUtils& CustomizationUtils ) 
+	{
+		CustomizeStructChildren( PropertyHandle, ChildBuilder, CustomizationUtils );
+	}
+
+	/**
+	 * Called when the header of the struct (usually where the name of the struct and information about the struct as a whole is added)
+	 * If nothing is added to the row, the header is not displayed
+	 *
+	 * @param StructPropertyHandle		Handle to the struct property
+	 * @param HeaderRow					A row that widgets can be added to
+	 * @param StructCustomizationUtils	Utilities for struct customization
+	 */
+	virtual void CustomizeStructHeader( TSharedRef<IPropertyHandle> StructPropertyHandle, class FDetailWidgetRow& HeaderRow, IStructCustomizationUtils& StructCustomizationUtils ) = 0;
+
+	/**
+	 * Called when the children of the struct should be customized
+	 *
+	 * @param StructPropertyHandle		Handle to the struct property
+	 * @param StructBuilder				A builder for customizing the struct children
+	 * @param StructCustomizationUtils	Utilities for struct customization
+	 */
+	virtual void CustomizeStructChildren( TSharedRef<IPropertyHandle> StructPropertyHandle, class IDetailChildrenBuilder& ChildBuilder, IStructCustomizationUtils& StructCustomizationUtils ) = 0;
+};
diff --git a/Engine/Source/Editor/PropertyEditor/Public/PropertyEditing.h b/Engine/Source/Editor/PropertyEditor/Public/PropertyEditing.h
--- a/Engine/Source/Editor/PropertyEditor/Public/PropertyEditing.h
+++ b/Engine/Source/Editor/PropertyEditor/Public/PropertyEditing.h
@@ -10,4 +10,7 @@
 #include "IDetailPropertyRow.h"
 #include "DetailLayoutBuilder.h"
 #include "DetailCategoryBuilder.h"
+#include "IPropertyTypeCustomization.h"
+#include "IDetailChildrenBuilder.h"
+#include "IDetailCustomization.h"
 #include "IDetailsView.h"
diff --git a/Engine/Source/Editor/PropertyEditor/Public/PropertyEditorModule.h b/Engine/Source/Editor/PropertyEditor/Public/PropertyEditorModule.h
--- a/Engine/Source/Editor/PropertyEditor/Public/PropertyEditorModule.h
+++ b/Engine/Source/Editor/PropertyEditor/Public/PropertyEditorModule.h
@@ -5,6 +5,9 @@
 #include "Toolkits/IToolkit.h"
 
 #include "PropertyEditorDelegates.h"
+#include "IDetailCustomization.h"
+#include "IPropertyTypeCustomization.h"
+#include "IDetailChildrenBuilder.h"
 
 /**
  * The location of a property name relative to its editor widget                   
@@ -42,167 +45,12 @@ namespace FPropertyAccess
 }
 
 
-/** 
- * Interface for any class that lays out details for a specific class
- */
-class IDetailCustomization : public TSharedFromThis<IDetailCustomization>
-{
-public:
-	virtual ~IDetailCustomization() {}
-
-	/** Called when details should be customized */
-	virtual void CustomizeDetails( class IDetailLayoutBuilder& DetailBuilder ) = 0;
-};
 
 class IPropertyHandle;
 class SPropertyTreeViewImpl;
 class SWindow;
 class IPropertyTableCellPresenter;
-
-
-/**
- * Utilities for property type customization
- */
-class IPropertyTypeCustomizationUtils
-{
-public:
-	virtual ~IPropertyTypeCustomizationUtils(){};
-
-	/**
-	 * @return the font used for properties and details
-	 */ 
-	static FSlateFontInfo GetRegularFont() { return FEditorStyle::GetFontStyle( TEXT("PropertyWindow.NormalFont") ); }
-
-	/**
-	 * @return the bold font used for properties and details
-	 */ 
-	static FSlateFontInfo GetBoldFont() { return FEditorStyle::GetFontStyle( TEXT("PropertyWindow.BoldFont") ); }
-	
-	/**
-	 * Gets the thumbnail pool that should be used for rendering thumbnails in the struct                  
-	 */
-	virtual TSharedPtr<class FAssetThumbnailPool> GetThumbnailPool() const = 0;
-
-	/**
-	*	@return the utilities various widgets need access to certain features of PropertyDetails
-	*/
-	virtual TSharedPtr<class IPropertyUtilities> GetPropertyUtilities() const { return NULL; }
-};
-
-/** Deprecated IStructCustomizationUtils interface */
-typedef IPropertyTypeCustomizationUtils IStructCustomizationUtils;
-
-/**
- * Builder for adding children to a detail customization
- */
-class IDetailChildrenBuilder
-{
-public:
-	virtual ~IDetailChildrenBuilder(){}
-
-	/**
-	 * Adds a custom builder as a child
-	 *
-	 * @param InCustomBuilder		The custom builder to add
-	 */
-	virtual IDetailChildrenBuilder& AddChildCustomBuilder( TSharedRef<class IDetailCustomNodeBuilder> InCustomBuilder ) = 0;
-
-	/**
-	 * Adds a group to the category
-	 *
-	 * @param GroupName	The name of the group 
-	 * @param LocalizedDisplayName	The display name of the group
-	 * @param true if the group should appear in the advanced section of the category
-	 */
-	virtual class IDetailGroup& AddChildGroup( FName GroupName, const FString& LocalizedDisplayName ) = 0;
-
-	/**
-	 * Adds new custom content as a child to the struct
-	 *
-	 * @param SearchString	Search string that will be matched when users search in the details panel.  If the search string doesnt match what the user types, this row will be hidden
-	 * @return A row that accepts widgets
-	 */
-	virtual class FDetailWidgetRow& AddChildContent( const FString& SearchString ) = 0;
-
-	/**
-	 * Adds a property to the struct
-	 * 
-	 * @param PropertyHandle	The handle to the property to add
-	 * @return An interface to the property row that can be used to customize the appearance of the property
-	 */
-	virtual class IDetailPropertyRow& AddChildProperty( TSharedRef<IPropertyHandle> PropertyHandle ) = 0;
-
-	/**
-	 * Generates a value widget from a customized struct
-	 * If the customized struct has no value widget an empty widget will be returned
-	 *
-	 * @param StructPropertyHandle	The handle to the struct property to generate the value widget from
-	 */
-	virtual TSharedRef<SWidget> GenerateStructValueWidget(TSharedRef<IPropertyHandle> StructPropertyHandle) = 0;
-
-};
-
-/**
- * Base class for property type customizations
- */
-class IPropertyTypeCustomization : public TSharedFromThis<IPropertyTypeCustomization>
-{
-public:
-	virtual ~IPropertyTypeCustomization() {}
-
-	/**
-	 * Called when the header of the property (the row in the details panel where the property is shown)
-	 * If nothing is added to the row, the header is not displayed
-	 *
-	 * @param PropertyHandle			Handle to the property being customized
-	 * @param HeaderRow					A row that widgets can be added to
-	 * @param StructCustomizationUtils	Utilities for customization
-	 */
-	virtual void CustomizeHeader( TSharedRef<IPropertyHandle> PropertyHandle, class FDetailWidgetRow& HeaderRow, IPropertyTypeCustomizationUtils& CustomizationUtils ) = 0;
-
-	/**
-	 * Called when the children of the property should be customized or extra rows added
-	 *
-	 * @param PropertyHandle			Handle to the property being customized
-	 * @param StructBuilder				A builder for adding children
-	 * @param StructCustomizationUtils	Utilities for customization
-	 */
-	virtual void CustomizeChildren( TSharedRef<IPropertyHandle> PropertyHandle, class IDetailChildrenBuilder& ChildBuilder, IPropertyTypeCustomizationUtils& CustomizationUtils ) = 0;
-};
-
-/** Deprecated IStructCustomization interface */
-class IStructCustomization : public IPropertyTypeCustomization
-{
-public:
-	virtual void CustomizeHeader( TSharedRef<IPropertyHandle> PropertyHandle, class FDetailWidgetRow& HeaderRow, IPropertyTypeCustomizationUtils& CustomizationUtils )
-	{
-		CustomizeStructHeader( PropertyHandle, HeaderRow, CustomizationUtils );
-	}
-
-	virtual void CustomizeChildren( TSharedRef<IPropertyHandle> PropertyHandle, class IDetailChildrenBuilder& ChildBuilder, IPropertyTypeCustomizationUtils& CustomizationUtils ) 
-	{
-		CustomizeStructChildren( PropertyHandle, ChildBuilder, CustomizationUtils );
-	}
-
-	/**
-	 * Called when the header of the struct (usually where the name of the struct and information about the struct as a whole is added)
-	 * If nothing is added to the row, the header is not displayed
-	 *
-	 * @param StructPropertyHandle		Handle to the struct property
-	 * @param HeaderRow					A row that widgets can be added to
-	 * @param StructCustomizationUtils	Utilities for struct customization
-	 */
-	virtual void CustomizeStructHeader( TSharedRef<IPropertyHandle> StructPropertyHandle, class FDetailWidgetRow& HeaderRow, IStructCustomizationUtils& StructCustomizationUtils ) = 0;
-
-	/**
-	 * Called when the children of the struct should be customized
-	 *
-	 * @param StructPropertyHandle		Handle to the struct property
-	 * @param StructBuilder				A builder for customizing the struct children
-	 * @param StructCustomizationUtils	Utilities for struct customization
-	 */
-	virtual void CustomizeStructChildren( TSharedRef<IPropertyHandle> StructPropertyHandle, class IDetailChildrenBuilder& ChildBuilder, IStructCustomizationUtils& StructCustomizationUtils ) = 0;
-};
+class IPropertyTypeCustomization;
 
 /** 
  * Base class for adding an extra data to identify a custom property type
@@ -241,10 +89,7 @@ struct FPropertyTypeLayoutCallback
 
 	bool IsValid() { return PropertyTypeLayoutDelegate.IsBound() || DeprecatedLayoutDelegate.IsBound(); }
 
-	TSharedRef<IPropertyTypeCustomization> GetCustomizationInstance() const
-	{
-		return PropertyTypeLayoutDelegate.IsBound() ? PropertyTypeLayoutDelegate.Execute() : StaticCastSharedRef<IPropertyTypeCustomization>(DeprecatedLayoutDelegate.Execute());
-	}
+	TSharedRef<IPropertyTypeCustomization> GetCustomizationInstance() const;
 };
 
 
diff --git a/Engine/Source/Editor/WorldBrowser/Private/Tiles/WorldTileCollectionModel.cpp b/Engine/Source/Editor/WorldBrowser/Private/Tiles/WorldTileCollectionModel.cpp
--- a/Engine/Source/Editor/WorldBrowser/Private/Tiles/WorldTileCollectionModel.cpp
+++ b/Engine/Source/Editor/WorldBrowser/Private/Tiles/WorldTileCollectionModel.cpp
@@ -575,8 +575,8 @@ void FWorldTileCollectionModel::RegisterDetailsCustomization(FPropertyEditorModu
 void FWorldTileCollectionModel::UnregisterDetailsCustomization(FPropertyEditorModule& InPropertyModule,
 															   TSharedPtr<IDetailsView> InDetailsView)
 {
-	InPropertyModule.UnregisterStructPropertyLayout("TileStreamingLevelDetails");
-	InPropertyModule.UnregisterStructPropertyLayout("TileLODEntryDetails");
+	InPropertyModule.UnregisterCustomPropertyTypeLayout("TileStreamingLevelDetails");
+	InPropertyModule.UnregisterCustomPropertyTypeLayout("TileLODEntryDetails");
 	InDetailsView->UnregisterInstancedCustomPropertyLayout(UWorldTileDetails::StaticClass());
 }
 
