Added the concept of a primary asset editor to the asset toolkit manager

The property editor is not considered 'primary' for the purposes of spawning editors. This means that the property editor can be opened, for example, for a static mesh while the static mesh editor is open. Giving focus to both editors when the revlevant editor tries to repon now works correctly. Previously the static mesh editor would give focus to the property editor rather than open up a new mesh editor.

This addresses TTP#334576 - EDITOR: Content Browser: Right clicking on an asset and selecting 'details' will spawn a new details tab even if it is already opened

diff --git a/Engine/Source/Editor/PropertyEditor/Private/PropertyEditorToolkit.cpp b/Engine/Source/Editor/PropertyEditor/Private/PropertyEditorToolkit.cpp
--- a/Engine/Source/Editor/PropertyEditor/Private/PropertyEditorToolkit.cpp
+++ b/Engine/Source/Editor/PropertyEditor/Private/PropertyEditorToolkit.cpp
@@ -13,6 +13,7 @@
 
 #define LOCTEXT_NAMESPACE "PropertyEditorToolkit"
 
+const FName FPropertyEditorToolkit::ToolkitFName( TEXT( "PropertyEditor" ) );
 const FName FPropertyEditorToolkit::ApplicationId( TEXT( "PropertyEditorToolkitApp" ) );
 const FName FPropertyEditorToolkit::TreeTabId( TEXT( "PropertyEditorToolkit_PropertyTree" ) );
 const FName FPropertyEditorToolkit::GridTabId( TEXT( "PropertyEditorToolkit_PropertyTable" ) );
@@ -46,8 +47,33 @@ FPropertyEditorToolkit::FPropertyEditorToolkit()
 }
 
 
+TSharedPtr<FPropertyEditorToolkit> FPropertyEditorToolkit::FindExistingEditor( UObject* Object )
+{
+	// Find any existing property editor instances for this asset
+	const TArray<IAssetEditorInstance*> Editors = FAssetEditorManager::Get().FindEditorsForAsset( Object );
+
+	IAssetEditorInstance* const * ExistingInstance = Editors.FindByPredicate( [&]( IAssetEditorInstance* Editor ){
+		return Editor->GetEditorName() == ToolkitFName;
+	} );
+
+	if( ExistingInstance )
+	{
+		auto* PropertyEditor = static_cast<FPropertyEditorToolkit*>( *ExistingInstance );
+		return PropertyEditor->SharedThis( PropertyEditor );
+	}
+
+	return nullptr;
+}
+
 TSharedRef<FPropertyEditorToolkit> FPropertyEditorToolkit::CreateEditor( const EToolkitMode::Type Mode, const TSharedPtr< IToolkitHost >& InitToolkitHost, UObject* ObjectToEdit )
 {
+	auto ExistingEditor = FindExistingEditor( ObjectToEdit );
+	if( ExistingEditor.IsValid() )
+	{
+		ExistingEditor->FocusWindow();
+		return ExistingEditor.ToSharedRef();
+	}
+
 	TSharedRef< FPropertyEditorToolkit > NewEditor( new FPropertyEditorToolkit() );
 
 	TArray<UObject*> ObjectsToEdit;
@@ -60,6 +86,16 @@ TSharedRef<FPropertyEditorToolkit> FPropertyEditorToolkit::CreateEditor( const E
 
 TSharedRef<FPropertyEditorToolkit> FPropertyEditorToolkit::CreateEditor( const EToolkitMode::Type Mode, const TSharedPtr< IToolkitHost >& InitToolkitHost, const TArray<UObject*>& ObjectsToEdit )
 {
+	if( ObjectsToEdit.Num() == 1 )
+	{
+		auto ExistingEditor = FindExistingEditor( ObjectsToEdit[0] );
+		if( ExistingEditor.IsValid() )
+		{
+			ExistingEditor->FocusWindow();
+			return ExistingEditor.ToSharedRef();
+		}
+	}
+
 	TSharedRef< FPropertyEditorToolkit > NewEditor( new FPropertyEditorToolkit() );
 	NewEditor->Initialize( Mode, InitToolkitHost, ObjectsToEdit );
 
@@ -444,7 +480,7 @@ void FPropertyEditorToolkit::GridRootPathChanged()
 
 FName FPropertyEditorToolkit::GetToolkitFName() const
 {
-	return FName("PropertyEditor");
+	return ToolkitFName;
 }
 
 FText FPropertyEditorToolkit::GetBaseToolkitName() const
diff --git a/Engine/Source/Editor/PropertyEditor/Private/PropertyEditorToolkit.h b/Engine/Source/Editor/PropertyEditor/Private/PropertyEditorToolkit.h
--- a/Engine/Source/Editor/PropertyEditor/Private/PropertyEditorToolkit.h
+++ b/Engine/Source/Editor/PropertyEditor/Private/PropertyEditorToolkit.h
@@ -31,6 +31,8 @@ public:
 
 	virtual bool CloseWindow() OVERRIDE;
 
+	virtual bool IsPrimaryEditor() const OVERRIDE{ return false; };
+
 public:
 
 	static TSharedRef<FPropertyEditorToolkit> CreateEditor( const EToolkitMode::Type Mode, const TSharedPtr< IToolkitHost >& InitToolkitHost, UObject* ObjectToEdit );
@@ -39,6 +41,7 @@ public:
 
 
 private:
+	static TSharedPtr<FPropertyEditorToolkit> FindExistingEditor( UObject* Object );
 
 	void Initialize( const EToolkitMode::Type Mode, const TSharedPtr< class IToolkitHost >& InitToolkitHost, const TArray<UObject*>& ObjectsToEdit );
 
@@ -90,6 +93,8 @@ private:
 	FSlateColor PinColor;
 	TArray< TWeakPtr<IPropertyTreeRow> > PinRows;
 
+	static const FName ToolkitFName;
+
 	static const FName ApplicationId;
 	static const FName TreeTabId;
 	static const FName GridTabId;
diff --git a/Engine/Source/Editor/UnrealEd/Private/Toolkits/AssetEditorManager.cpp b/Engine/Source/Editor/UnrealEd/Private/Toolkits/AssetEditorManager.cpp
--- a/Engine/Source/Editor/UnrealEd/Private/Toolkits/AssetEditorManager.cpp
+++ b/Engine/Source/Editor/UnrealEd/Private/Toolkits/AssetEditorManager.cpp
@@ -89,19 +89,28 @@ void FAssetEditorManager::AddReferencedObjects( FReferenceCollector& Collector )
 
 IAssetEditorInstance* FAssetEditorManager::FindEditorForAsset(UObject* Asset, bool bFocusIfOpen)
 {
-	IAssetEditorInstance** pAssetEditor = OpenedAssets.Find(Asset);
-	const bool bEditorOpen = pAssetEditor != NULL;
+	const TArray<IAssetEditorInstance*> AssetEditors = FindEditorsForAsset(Asset);
 
+	IAssetEditorInstance* const * PrimaryEditor = AssetEditors.FindByPredicate( [](IAssetEditorInstance* Editor){ return Editor->IsPrimaryEditor(); } );
+	
+	const bool bEditorOpen = PrimaryEditor != NULL;
 	if (bEditorOpen && bFocusIfOpen)
 	{
 		// @todo toolkit minor: We may need to handle this differently for world-centric vs standalone.  (multiple level editors, etc)
-		(*pAssetEditor)->FocusWindow(Asset);
+		(*PrimaryEditor)->FocusWindow(Asset);
 	}
 
-	return bEditorOpen ? *pAssetEditor : NULL;
+	return bEditorOpen ? *PrimaryEditor : NULL;
 }
 
 
+TArray<IAssetEditorInstance*> FAssetEditorManager::FindEditorsForAsset(UObject* Asset)
+{
+	TArray<IAssetEditorInstance*> AssetEditors;
+	OpenedAssets.MultiFind(Asset, AssetEditors);
+	return AssetEditors;
+}
+
 void FAssetEditorManager::CloseOtherEditors( UObject* Asset, IAssetEditorInstance* OnlyEditor)
 {
 	TArray<UObject*> AllAssets;
@@ -294,6 +303,22 @@ bool FAssetEditorManager::OpenEditorForAssets( const TArray< UObject* >& Assets,
 		FAssetEditorToolkit::SetPreviousWorldCentricToolkitHostForNewAssetEditor( OpenedFromLevelEditor.ToSharedRef() );
 	}
 
+	if( Assets.Num() == 1 )
+	{
+		// Find any existing simple editors for this asset
+		const TArray<IAssetEditorInstance*> Editors = FAssetEditorManager::Get().FindEditorsForAsset( Assets[0] );
+
+		IAssetEditorInstance* const * ExistingInstance = Editors.FindByPredicate( [&](IAssetEditorInstance* Editor){
+			return Editor->GetEditorName() == FSimpleAssetEditor::ToolkitFName;
+		} );
+
+		if( ExistingInstance )
+		{
+			(*ExistingInstance)->FocusWindow();
+			return true;
+		}
+	}
+
 	// No asset type actions for this asset. Just use a properties editor.
 	FSimpleAssetEditor::CreateEditor(ToolkitMode, ToolkitMode == EToolkitMode::WorldCentric ? OpenedFromLevelEditor : TSharedPtr<IToolkitHost>(), Assets );
 
diff --git a/Engine/Source/Editor/UnrealEd/Private/Toolkits/SimpleAssetEditor.cpp b/Engine/Source/Editor/UnrealEd/Private/Toolkits/SimpleAssetEditor.cpp
--- a/Engine/Source/Editor/UnrealEd/Private/Toolkits/SimpleAssetEditor.cpp
+++ b/Engine/Source/Editor/UnrealEd/Private/Toolkits/SimpleAssetEditor.cpp
@@ -10,6 +10,7 @@
 
 #define LOCTEXT_NAMESPACE "GenericEditor"
 
+const FName FSimpleAssetEditor::ToolkitFName( TEXT( "GenericAssetEditor" ) );
 const FName FSimpleAssetEditor::PropertiesTabId( TEXT( "GenericEditor_Properties" ) );
 
 void FSimpleAssetEditor::RegisterTabSpawners(const TSharedRef<class FTabManager>& TabManager)
@@ -101,7 +102,7 @@ void FSimpleAssetEditor::InitEditor( const EToolkitMode::Type Mode, const TShare
 
 FName FSimpleAssetEditor::GetToolkitFName() const
 {
-	return FName("GenericAssetEditor");
+	return ToolkitFName;
 }
 
 FText FSimpleAssetEditor::GetBaseToolkitName() const
diff --git a/Engine/Source/Editor/UnrealEd/Public/MiniCurveEditor.h b/Engine/Source/Editor/UnrealEd/Public/MiniCurveEditor.h
--- a/Engine/Source/Editor/UnrealEd/Public/MiniCurveEditor.h
+++ b/Engine/Source/Editor/UnrealEd/Public/MiniCurveEditor.h
@@ -23,6 +23,7 @@ public:
 	virtual FName GetEditorName() const OVERRIDE;
 	virtual void FocusWindow(UObject* ObjectToFocusOn) OVERRIDE;
 	virtual bool CloseWindow() OVERRIDE;
+	virtual bool IsPrimaryEditor() const OVERRIDE { return true; };
 private:
 
 	float ViewMinInput;
diff --git a/Engine/Source/Editor/UnrealEd/Public/Toolkits/AssetEditorManager.h b/Engine/Source/Editor/UnrealEd/Public/Toolkits/AssetEditorManager.h
--- a/Engine/Source/Editor/UnrealEd/Public/Toolkits/AssetEditorManager.h
+++ b/Engine/Source/Editor/UnrealEd/Public/Toolkits/AssetEditorManager.h
@@ -19,6 +19,7 @@ public:
 	virtual FName GetEditorName() const = 0;
 	virtual void FocusWindow(UObject* ObjectToFocusOn = NULL) = 0;
 	virtual bool CloseWindow() = 0;
+	virtual bool IsPrimaryEditor() const = 0;
 };
 
 
@@ -43,10 +44,13 @@ public:
 	bool OpenEditorForAsset( UObject* Asset, const EToolkitMode::Type ToolkitMode = EToolkitMode::Standalone, TSharedPtr< class IToolkitHost > OpenedFromLevelEditor = TSharedPtr< IToolkitHost >() );
 	bool OpenEditorForAssets( const TArray< UObject* >& Asset, const EToolkitMode::Type ToolkitMode = EToolkitMode::Standalone, TSharedPtr< class IToolkitHost > OpenedFromLevelEditor = TSharedPtr< IToolkitHost >() );
 
-	// Returns the editor if there is already an editor open for the specified asset.
+	// Returns the primary editor if one is already open for the specified asset.
 	// If there is one open and bFocusIfOpen is true, that editor will be brought to the foreground and focused if possible.
 	IAssetEditorInstance* FindEditorForAsset(UObject* Asset, bool bFocusIfOpen);
 
+	// Returns all editors currently opened for the specified asset
+	TArray<IAssetEditorInstance*> FindEditorsForAsset(UObject* Asset);
+
 	// Close any editor which is not this one
 	void CloseOtherEditors(UObject* Asset, IAssetEditorInstance* OnlyEditor);
 
diff --git a/Engine/Source/Editor/UnrealEd/Public/Toolkits/AssetEditorToolkit.h b/Engine/Source/Editor/UnrealEd/Public/Toolkits/AssetEditorToolkit.h
--- a/Engine/Source/Editor/UnrealEd/Public/Toolkits/AssetEditorToolkit.h
+++ b/Engine/Source/Editor/UnrealEd/Public/Toolkits/AssetEditorToolkit.h
@@ -54,6 +54,7 @@ public:
 	virtual FName GetEditorName() const OVERRIDE;
 	virtual void FocusWindow(UObject* ObjectToFocusOn = NULL) OVERRIDE;
 	virtual bool CloseWindow() OVERRIDE;
+	virtual bool IsPrimaryEditor() const OVERRIDE { return true; };
 
 	/**
 	 * Fills in the supplied menu with commands for working with this asset file
diff --git a/Engine/Source/Editor/UnrealEd/Public/Toolkits/SimpleAssetEditor.h b/Engine/Source/Editor/UnrealEd/Public/Toolkits/SimpleAssetEditor.h
--- a/Engine/Source/Editor/UnrealEd/Public/Toolkits/SimpleAssetEditor.h
+++ b/Engine/Source/Editor/UnrealEd/Public/Toolkits/SimpleAssetEditor.h
@@ -35,6 +35,7 @@ public:
 	virtual FText GetToolkitName() const OVERRIDE;
 	virtual FString GetWorldCentricTabPrefix() const OVERRIDE;
 	virtual FLinearColor GetWorldCentricTabColorScale() const OVERRIDE;
+	virtual bool IsPrimaryEditor() const OVERRIDE { return false; }
 	
 	/** Used to show or hide certain properties */
 	void SetPropertyVisibilityDelegate(FIsPropertyVisible InVisibilityDelegate);
@@ -56,6 +57,9 @@ private:
 	static const FName PropertiesTabId;
 
 public:
+	/** The name given to all instances of this type of editor */
+	static const FName ToolkitFName;
+
 	static TSharedRef<FSimpleAssetEditor> CreateEditor( const EToolkitMode::Type Mode, const TSharedPtr< IToolkitHost >& InitToolkitHost, UObject* ObjectToEdit, FGetDetailsViewObjects GetDetailsViewObjects = FGetDetailsViewObjects() );
 
 	static TSharedRef<FSimpleAssetEditor> CreateEditor( const EToolkitMode::Type Mode, const TSharedPtr< IToolkitHost >& InitToolkitHost, const TArray<UObject*>& ObjectsToEdit, FGetDetailsViewObjects GetDetailsViewObjects = FGetDetailsViewObjects() );
