UE4 Refactoring. Changed OVERRIDE and FINAL macros to keywords override and final.

diff --git a/Engine/Source/Editor/Persona/Private/STrack.h b/Engine/Source/Editor/Persona/Private/STrack.h
--- a/Engine/Source/Editor/Persona/Private/STrack.h
+++ b/Engine/Source/Editor/Persona/Private/STrack.h
@@ -31,9 +31,9 @@ class FTrackNodeDragDropOp : public FDragDropOperation
 public:
 	DRAG_DROP_OPERATOR_TYPE(FTrackNodeDragDropOp, FDragDropOperation)
 
-	virtual void OnDrop( bool bDropWasHandled, const FPointerEvent& MouseEvent ) OVERRIDE;
+	virtual void OnDrop( bool bDropWasHandled, const FPointerEvent& MouseEvent ) override;
 
-	virtual void OnDragged( const class FDragDropEvent& DragDropEvent ) OVERRIDE;
+	virtual void OnDragged( const class FDragDropEvent& DragDropEvent ) override;
 
 	static TSharedRef<FTrackNodeDragDropOp> New(TSharedRef<STrackNode> TrackNode, const FVector2D &CursorPosition, const FVector2D &ScreenPositionOfNode);
 
@@ -107,10 +107,10 @@ public:
 	void Construct(const FArguments& Declaration);
 
 	// mouse interface for tooltip/selection
-	virtual FReply OnMouseButtonUp( const FGeometry& MyGeometry, const FPointerEvent& MouseEvent ) OVERRIDE;
-	virtual FReply OnMouseMove( const FGeometry& MyGeometry, const FPointerEvent& MouseEvent ) OVERRIDE;
-	virtual void OnMouseEnter( const FGeometry& MyGeometry, const FPointerEvent& MouseEvent ) OVERRIDE;
-	virtual void OnMouseLeave( const FPointerEvent& MouseEvent ) OVERRIDE;
+	virtual FReply OnMouseButtonUp( const FGeometry& MyGeometry, const FPointerEvent& MouseEvent ) override;
+	virtual FReply OnMouseMove( const FGeometry& MyGeometry, const FPointerEvent& MouseEvent ) override;
+	virtual void OnMouseEnter( const FGeometry& MyGeometry, const FPointerEvent& MouseEvent ) override;
+	virtual void OnMouseLeave( const FPointerEvent& MouseEvent ) override;
 	virtual void OnDragged(const class FDragDropEvent& DragDropEvent );
 
 	// virtual draw related functions
@@ -120,8 +120,8 @@ public:
 	virtual int32 OnPaint(const FGeometry& AllottedGeometry, const FSlateRect& MyClippingRect, FSlateWindowElementList& OutDrawElements, int32 LayerId, const FWidgetStyle& InWidgetStyle, bool bParentEnabled) const;
 
 	// drag drop relationship
-	virtual FReply OnDragDetected( const FGeometry& MyGeometry, const FPointerEvent& MouseEvent ) OVERRIDE;
-	virtual FReply OnMouseButtonDown( const FGeometry& MyGeometry, const FPointerEvent& MouseEvent ) OVERRIDE;
+	virtual FReply OnDragDetected( const FGeometry& MyGeometry, const FPointerEvent& MouseEvent ) override;
+	virtual FReply OnMouseButtonDown( const FGeometry& MyGeometry, const FPointerEvent& MouseEvent ) override;
 
 	void OnDropCancelled(const FPointerEvent& MouseEvent);
 
@@ -238,22 +238,22 @@ public:
 	virtual void	OnArrangeChildren( const FGeometry& AllottedGeometry, FArrangedChildren& ArrangedChildren ) const;
 	virtual int32	OnPaint( const FGeometry& AllottedGeometry, const FSlateRect& MyClippingRect, FSlateWindowElementList& OutDrawElements, int32 LayerId, const FWidgetStyle& InWidgetStyle, bool bParentEnabled ) const;
 
-	virtual FReply	OnDrop( const FGeometry& MyGeometry, const FDragDropEvent& DragDropEvent ) OVERRIDE;	
-	virtual FReply	OnDragOver( const FGeometry& MyGeometry, const FDragDropEvent& DragDropEvent ) OVERRIDE;
+	virtual FReply	OnDrop( const FGeometry& MyGeometry, const FDragDropEvent& DragDropEvent ) override;	
+	virtual FReply	OnDragOver( const FGeometry& MyGeometry, const FDragDropEvent& DragDropEvent ) override;
 
 	int32			GetHitNode(const FGeometry& MyGeometry, const FVector2D& CursorPosition);
 
-	FReply			OnMouseMove( const FGeometry& MyGeometry, const FPointerEvent& MouseEvent ) OVERRIDE;
-	FReply			OnMouseButtonDown( const FGeometry& MyGeometry, const FPointerEvent& MouseEvent ) OVERRIDE;
-	FReply			OnMouseButtonUp( const FGeometry& MyGeometry, const FPointerEvent& MouseEvent ) OVERRIDE;
+	FReply			OnMouseMove( const FGeometry& MyGeometry, const FPointerEvent& MouseEvent ) override;
+	FReply			OnMouseButtonDown( const FGeometry& MyGeometry, const FPointerEvent& MouseEvent ) override;
+	FReply			OnMouseButtonUp( const FGeometry& MyGeometry, const FPointerEvent& MouseEvent ) override;
 	
-	virtual void OnMouseLeave( const FPointerEvent& MouseEvent ) OVERRIDE;
-	virtual void OnMouseEnter( const FGeometry& MyGeometry, const FPointerEvent& MouseEvent ) OVERRIDE;
-	virtual FCursorReply OnCursorQuery( const FGeometry& MyGeometry, const FPointerEvent& CursorEvent ) const OVERRIDE;
+	virtual void OnMouseLeave( const FPointerEvent& MouseEvent ) override;
+	virtual void OnMouseEnter( const FGeometry& MyGeometry, const FPointerEvent& MouseEvent ) override;
+	virtual FCursorReply OnCursorQuery( const FGeometry& MyGeometry, const FPointerEvent& CursorEvent ) const override;
 
 	void				AddTrackNode( TSharedRef<STrackNode> Node );
-	FVector2D			ComputeDesiredSize() const OVERRIDE;
-	virtual FChildren*	GetChildren() OVERRIDE;
+	FVector2D			ComputeDesiredSize() const override;
+	virtual FChildren*	GetChildren() override;
 
 protected:
 	
diff --git a/Engine/Source/Editor/Persona/Private/Shared/PersonaMode.h b/Engine/Source/Editor/Persona/Private/Shared/PersonaMode.h
--- a/Engine/Source/Editor/Persona/Private/Shared/PersonaMode.h
+++ b/Engine/Source/Editor/Persona/Private/Shared/PersonaMode.h
@@ -129,8 +129,8 @@ protected:
 
 public:
 	// FApplicationMode interface
-	virtual void RegisterTabFactories(TSharedPtr<FTabManager> InTabManager) OVERRIDE;
-	virtual void PostActivateMode() OVERRIDE;
+	virtual void RegisterTabFactories(TSharedPtr<FTabManager> InTabManager) override;
+	virtual void PostActivateMode() override;
 	// End of FApplicationMode interface
 
 protected:
@@ -148,10 +148,10 @@ struct FSkeletonTreeSummoner : public FWorkflowTabFactory
 public:
 	FSkeletonTreeSummoner(TSharedPtr<class FAssetEditorToolkit> InHostingApp);
 	
-	virtual TSharedRef<SWidget> CreateTabBody(const FWorkflowTabSpawnInfo& Info) const OVERRIDE;
+	virtual TSharedRef<SWidget> CreateTabBody(const FWorkflowTabSpawnInfo& Info) const override;
 
 	// Create a tooltip widget for the tab
-	virtual TSharedPtr<SToolTip> CreateTabToolTipWidget(const FWorkflowTabSpawnInfo& Info) const OVERRIDE
+	virtual TSharedPtr<SToolTip> CreateTabToolTipWidget(const FWorkflowTabSpawnInfo& Info) const override
 	{
 		return  IDocumentation::Get()->CreateToolTip(LOCTEXT("SkeletonTreeTooltip", "The Skeleton Tree tab lets you see and select bones (and sockets) in the skeleton hierarchy."), NULL, TEXT("Shared/Editors/Persona"), TEXT("SkeletonTree_Window"));
 	}
@@ -165,10 +165,10 @@ struct FMorphTargetTabSummoner : public FWorkflowTabFactory
 public:
 	FMorphTargetTabSummoner(TSharedPtr<class FAssetEditorToolkit> InHostingApp);
 
-	virtual TSharedRef<SWidget> CreateTabBody(const FWorkflowTabSpawnInfo& Info) const OVERRIDE;
+	virtual TSharedRef<SWidget> CreateTabBody(const FWorkflowTabSpawnInfo& Info) const override;
 
 	// Create a tooltip widget for the tab
-	virtual TSharedPtr<SToolTip> CreateTabToolTipWidget(const FWorkflowTabSpawnInfo& Info) const OVERRIDE
+	virtual TSharedPtr<SToolTip> CreateTabToolTipWidget(const FWorkflowTabSpawnInfo& Info) const override
 	{
 		return  IDocumentation::Get()->CreateToolTip(LOCTEXT("MorphTargetTooltip", "The Morph Target tab lets you preview any morph targets (aka blend shapes) available for the current mesh."), NULL, TEXT("Shared/Editors/Persona"), TEXT("MorphTarget_Window"));
 	}
@@ -182,10 +182,10 @@ struct FAnimationAssetBrowserSummoner : public FWorkflowTabFactory
 {
 	FAnimationAssetBrowserSummoner(TSharedPtr<class FAssetEditorToolkit> InHostingApp);
 
-	virtual TSharedRef<SWidget> CreateTabBody(const FWorkflowTabSpawnInfo& Info) const OVERRIDE;
+	virtual TSharedRef<SWidget> CreateTabBody(const FWorkflowTabSpawnInfo& Info) const override;
 
 	// Create a tooltip widget for the tab
-	virtual TSharedPtr<SToolTip> CreateTabToolTipWidget(const FWorkflowTabSpawnInfo& Info) const OVERRIDE
+	virtual TSharedPtr<SToolTip> CreateTabToolTipWidget(const FWorkflowTabSpawnInfo& Info) const override
 	{
 		return  IDocumentation::Get()->CreateToolTip(LOCTEXT("AnimAssetBrowserTooltip", "The Asset Browser lets you browse all animation-related assets (animations, blend spaces etc)."), NULL, TEXT("Shared/Editors/Persona"), TEXT("AssetBrowser_Window"));
 	}
@@ -198,7 +198,7 @@ struct FPreviewViewportSummoner : public FWorkflowTabFactory
 {
 	FPreviewViewportSummoner(TSharedPtr<class FAssetEditorToolkit> InHostingApp);
 
-	virtual TSharedRef<SWidget> CreateTabBody(const FWorkflowTabSpawnInfo& Info) const OVERRIDE;
+	virtual TSharedRef<SWidget> CreateTabBody(const FWorkflowTabSpawnInfo& Info) const override;
 };
 
 /////////////////////////////////////////////////////
@@ -209,10 +209,10 @@ struct FRetargetSourceManagerTabSummoner : public FWorkflowTabFactory
 public:
 	FRetargetSourceManagerTabSummoner(TSharedPtr<class FAssetEditorToolkit> InHostingApp);
 
-	virtual TSharedRef<SWidget> CreateTabBody(const FWorkflowTabSpawnInfo& Info) const OVERRIDE;
+	virtual TSharedRef<SWidget> CreateTabBody(const FWorkflowTabSpawnInfo& Info) const override;
 
 	// Create a tooltip widget for the tab
-	virtual TSharedPtr<SToolTip> CreateTabToolTipWidget(const FWorkflowTabSpawnInfo& Info) const OVERRIDE
+	virtual TSharedPtr<SToolTip> CreateTabToolTipWidget(const FWorkflowTabSpawnInfo& Info) const override
 	{
 		return  IDocumentation::Get()->CreateToolTip(LOCTEXT("RetargetSourceTooltip", "In this panel, you can manage retarget sources for different body types"), NULL, TEXT("Shared/Editors/Persona"), TEXT("RetargetSource_Window"));
 	}
@@ -235,9 +235,9 @@ struct FAnimBlueprintDefaultsEditorSummoner : public FWorkflowTabFactory
 public:
 	FAnimBlueprintDefaultsEditorSummoner(TSharedPtr<class FAssetEditorToolkit> InHostingApp);
 
-	virtual TSharedRef<SWidget> CreateTabBody(const FWorkflowTabSpawnInfo& Info) const OVERRIDE;
+	virtual TSharedRef<SWidget> CreateTabBody(const FWorkflowTabSpawnInfo& Info) const override;
 
-	virtual FText GetTabToolTipText(const FWorkflowTabSpawnInfo& Info) const OVERRIDE;
+	virtual FText GetTabToolTipText(const FWorkflowTabSpawnInfo& Info) const override;
 
 private:
 	/** Delegates to customize tab look based on selected mode */
@@ -257,8 +257,8 @@ class FAnimBlueprintParentPlayerEditorSummoner : public FWorkflowTabFactory
 public:
 	FAnimBlueprintParentPlayerEditorSummoner(TSharedPtr<class FAssetEditorToolkit> InHostingApp);
 
-	virtual TSharedRef<SWidget> CreateTabBody(const FWorkflowTabSpawnInfo& Info) const OVERRIDE;
-	virtual FText GetTabToolTipText(const FWorkflowTabSpawnInfo& Info) const OVERRIDE;
+	virtual TSharedRef<SWidget> CreateTabBody(const FWorkflowTabSpawnInfo& Info) const override;
+	virtual FText GetTabToolTipText(const FWorkflowTabSpawnInfo& Info) const override;
 
 private:
 
diff --git a/Engine/Source/Editor/PhAT/Classes/PhATEdSkeletalMeshComponent.h b/Engine/Source/Editor/PhAT/Classes/PhATEdSkeletalMeshComponent.h
--- a/Engine/Source/Editor/PhAT/Classes/PhATEdSkeletalMeshComponent.h
+++ b/Engine/Source/Editor/PhAT/Classes/PhATEdSkeletalMeshComponent.h
@@ -47,7 +47,7 @@ class UPhATEdSkeletalMeshComponent : public USkeletalMeshComponent
 
 
 	/** UPrimitiveComponent interface */
-	virtual FPrimitiveSceneProxy* CreateSceneProxy() OVERRIDE;
+	virtual FPrimitiveSceneProxy* CreateSceneProxy() override;
 	
 	/** Renders non-hitproxy elements for the viewport, this function is called in the Game Thread */
 	virtual void Render(const FSceneView* View, class FPrimitiveDrawInterface* PDI);
diff --git a/Engine/Source/Editor/PhAT/Private/PhAT.cpp b/Engine/Source/Editor/PhAT/Private/PhAT.cpp
--- a/Engine/Source/Editor/PhAT/Private/PhAT.cpp
+++ b/Engine/Source/Editor/PhAT/Private/PhAT.cpp
@@ -1348,12 +1348,12 @@ TSharedRef<ITableRow> FPhAT::OnGenerateRowForTree(FTreeElemPtr Item, const TShar
 			Phat = phat;
 		}
 
-		virtual void OnMouseEnter( const FGeometry& MyGeometry, const FPointerEvent& MouseEvent ) OVERRIDE
+		virtual void OnMouseEnter( const FGeometry& MyGeometry, const FPointerEvent& MouseEvent ) override
 		{
 			Phat->OnTreeHighlightChanged();
 		}
 	
-		virtual void OnMouseLeave( const FPointerEvent& MouseEvent ) OVERRIDE
+		virtual void OnMouseLeave( const FPointerEvent& MouseEvent ) override
 		{
 			Phat->OnTreeHighlightChanged();
 		}
diff --git a/Engine/Source/Editor/PhAT/Private/PhAT.h b/Engine/Source/Editor/PhAT/Private/PhAT.h
--- a/Engine/Source/Editor/PhAT/Private/PhAT.h
+++ b/Engine/Source/Editor/PhAT/Private/PhAT.h
@@ -31,8 +31,8 @@ enum EPhATSimulationMode
 class FPhAT : public IPhAT, public FGCObject, public FEditorUndoClient
 {
 public:
-	virtual void RegisterTabSpawners(const TSharedRef<class FTabManager>& TabManager) OVERRIDE;
-	virtual void UnregisterTabSpawners(const TSharedRef<class FTabManager>& TabManager) OVERRIDE;
+	virtual void RegisterTabSpawners(const TSharedRef<class FTabManager>& TabManager) override;
+	virtual void UnregisterTabSpawners(const TSharedRef<class FTabManager>& TabManager) override;
 	TSharedRef<SDockTab> SpawnTab( const FSpawnTabArgs& TabSpawnArgs, FName TabIdentifier );
 
 	/** Destructor */
@@ -69,19 +69,19 @@ public:
 	FString GetHierarchyFilter() const;
 
 	/** IToolkit interface */
-	virtual FName GetToolkitFName() const OVERRIDE;
-	virtual FText GetBaseToolkitName() const OVERRIDE;
-	virtual FString GetWorldCentricTabPrefix() const OVERRIDE;
-	virtual FLinearColor GetWorldCentricTabColorScale() const OVERRIDE;
+	virtual FName GetToolkitFName() const override;
+	virtual FText GetBaseToolkitName() const override;
+	virtual FString GetWorldCentricTabPrefix() const override;
+	virtual FLinearColor GetWorldCentricTabColorScale() const override;
 
 	/** @return the documentation location for this editor */
-	virtual FString GetDocumentationLink() const OVERRIDE
+	virtual FString GetDocumentationLink() const override
 	{
 		return FString(TEXT("Engine/Physics/PhAT"));
 	}
 
 	/** FGCObject interface */
-	virtual void AddReferencedObjects(FReferenceCollector& Collector) OVERRIDE;
+	virtual void AddReferencedObjects(FReferenceCollector& Collector) override;
 	
 	/** Cycles through transform modes. */
 	void CycleTransformMode();
@@ -98,8 +98,8 @@ private:
 	};
 
 	// Begin FEditorUndoClient Interface
-	virtual void PostUndo(bool bSuccess) OVERRIDE;
-	virtual void PostRedo(bool bSuccess) OVERRIDE;
+	virtual void PostUndo(bool bSuccess) override;
+	virtual void PostRedo(bool bSuccess) override;
 	// End of FEditorUndoClient
 
 	/** Creates all internal widgets for the tabs to point at */
diff --git a/Engine/Source/Editor/PhAT/Private/PhATModule.cpp b/Engine/Source/Editor/PhAT/Private/PhATModule.cpp
--- a/Engine/Source/Editor/PhAT/Private/PhATModule.cpp
+++ b/Engine/Source/Editor/PhAT/Private/PhATModule.cpp
@@ -22,27 +22,27 @@ public:
 	}
 
 	/** Called right after the module DLL has been loaded and the module object has been created */
-	virtual void StartupModule() OVERRIDE
+	virtual void StartupModule() override
 	{
 		MenuExtensibilityManager = MakeShareable(new FExtensibilityManager);
 		ToolBarExtensibilityManager = MakeShareable(new FExtensibilityManager);
 	}
 
 	/** Called before the module is unloaded, right before the module object is destroyed. */
-	virtual void ShutdownModule() OVERRIDE
+	virtual void ShutdownModule() override
 	{
 		MenuExtensibilityManager.Reset();
 		ToolBarExtensibilityManager.Reset();
 	}
 
-	virtual TSharedRef<IPhAT> CreatePhAT(const EToolkitMode::Type Mode, const TSharedPtr< IToolkitHost >& InitToolkitHost, UPhysicsAsset* PhysicsAsset) OVERRIDE
+	virtual TSharedRef<IPhAT> CreatePhAT(const EToolkitMode::Type Mode, const TSharedPtr< IToolkitHost >& InitToolkitHost, UPhysicsAsset* PhysicsAsset) override
 	{
 		TSharedRef<FPhAT> NewPhAT(new FPhAT());
 		NewPhAT->InitPhAT(Mode, InitToolkitHost, PhysicsAsset);
 		return NewPhAT;
 	}
 
-	virtual void OpenNewBodyDlg(FPhysAssetCreateParams* NewBodyData, EAppReturnType::Type* NewBodyResponse) OVERRIDE
+	virtual void OpenNewBodyDlg(FPhysAssetCreateParams* NewBodyData, EAppReturnType::Type* NewBodyResponse) override
 	{
 		FPhATSharedData::OpenNewBodyDlg(NewBodyData, NewBodyResponse);
 	}
diff --git a/Engine/Source/Editor/PhAT/Private/PhATPreviewViewportClient.h b/Engine/Source/Editor/PhAT/Private/PhATPreviewViewportClient.h
--- a/Engine/Source/Editor/PhAT/Private/PhATPreviewViewportClient.h
+++ b/Engine/Source/Editor/PhAT/Private/PhATPreviewViewportClient.h
@@ -22,20 +22,20 @@ public:
 
 	/** FEditorViewportClient interface */
 	virtual void DrawCanvas( FViewport& InViewport, FSceneView& View, FCanvas& Canvas );
-	virtual void Draw(const FSceneView* View, FPrimitiveDrawInterface* PDI) OVERRIDE;
-	virtual bool InputKey(FViewport* Viewport, int32 ControllerId, FKey Key, EInputEvent Event, float AmountDepressed = 1.0f, bool bGamepad = false) OVERRIDE;
-	virtual bool InputAxis(FViewport* Viewport, int32 ControllerId, FKey Key, float Delta, float DeltaTime, int32 NumSamples = 1, bool bGamepad = false) OVERRIDE;
-	virtual void ProcessClick(class FSceneView& View, class HHitProxy* HitProxy, FKey Key, EInputEvent Event, uint32 HitX, uint32 HitY) OVERRIDE;
-	virtual bool InputWidgetDelta( FViewport* Viewport, EAxisList::Type CurrentAxis, FVector& Drag, FRotator& Rot, FVector& Scale ) OVERRIDE;
-	virtual void TrackingStarted( const struct FInputEventState& InInputState, bool bIsDragging, bool bNudge ) OVERRIDE;
-	virtual void TrackingStopped() OVERRIDE;
-	virtual FWidget::EWidgetMode GetWidgetMode() const OVERRIDE;
+	virtual void Draw(const FSceneView* View, FPrimitiveDrawInterface* PDI) override;
+	virtual bool InputKey(FViewport* Viewport, int32 ControllerId, FKey Key, EInputEvent Event, float AmountDepressed = 1.0f, bool bGamepad = false) override;
+	virtual bool InputAxis(FViewport* Viewport, int32 ControllerId, FKey Key, float Delta, float DeltaTime, int32 NumSamples = 1, bool bGamepad = false) override;
+	virtual void ProcessClick(class FSceneView& View, class HHitProxy* HitProxy, FKey Key, EInputEvent Event, uint32 HitX, uint32 HitY) override;
+	virtual bool InputWidgetDelta( FViewport* Viewport, EAxisList::Type CurrentAxis, FVector& Drag, FRotator& Rot, FVector& Scale ) override;
+	virtual void TrackingStarted( const struct FInputEventState& InInputState, bool bIsDragging, bool bNudge ) override;
+	virtual void TrackingStopped() override;
+	virtual FWidget::EWidgetMode GetWidgetMode() const override;
 	virtual FVector GetWidgetLocation() const;
 	virtual FMatrix GetWidgetCoordSystem() const;
 	virtual ECoordSystem GetWidgetCoordSystemSpace() const;
-	virtual void Tick(float DeltaSeconds) OVERRIDE;
-	virtual FSceneInterface* GetScene() const OVERRIDE;
-	virtual FLinearColor GetBackgroundColor() const OVERRIDE;
+	virtual void Tick(float DeltaSeconds) override;
+	virtual FSceneInterface* GetScene() const override;
+	virtual FLinearColor GetBackgroundColor() const override;
 
 private:
 	/** Methods for building the various context menus */
diff --git a/Engine/Source/Editor/PhAT/Private/SPhATNewAssetDlg.h b/Engine/Source/Editor/PhAT/Private/SPhATNewAssetDlg.h
--- a/Engine/Source/Editor/PhAT/Private/SPhATNewAssetDlg.h
+++ b/Engine/Source/Editor/PhAT/Private/SPhATNewAssetDlg.h
@@ -23,8 +23,8 @@ public:
 	void Construct(const FArguments& InArgs);
 
 	/** SWidget interface */
-	virtual FReply OnKeyDown(const FGeometry& MyGeometry, const FKeyboardEvent& InKeyboardEvent) OVERRIDE;
-	virtual bool SupportsKeyboardFocus() const OVERRIDE;
+	virtual FReply OnKeyDown(const FGeometry& MyGeometry, const FKeyboardEvent& InKeyboardEvent) override;
+	virtual bool SupportsKeyboardFocus() const override;
 
 private:
 	/** Makes option widgets for combobox controls */
diff --git a/Engine/Source/Editor/PlacementMode/Private/PlacementMode.h b/Engine/Source/Editor/PlacementMode/Private/PlacementMode.h
--- a/Engine/Source/Editor/PlacementMode/Private/PlacementMode.h
+++ b/Engine/Source/Editor/PlacementMode/Private/PlacementMode.h
@@ -11,48 +11,48 @@ public:
 	~FPlacementMode();
 
 	// FEdMode interface
-	virtual bool UsesToolkits() const OVERRIDE;
-	void Enter() OVERRIDE;
-	void Exit() OVERRIDE;
-	virtual void Tick(FLevelEditorViewportClient* ViewportClient,float DeltaTime) OVERRIDE;
-	virtual bool MouseEnter( FLevelEditorViewportClient* ViewportClient,FViewport* Viewport,int32 x, int32 y ) OVERRIDE;
-	virtual bool MouseLeave( FLevelEditorViewportClient* ViewportClient,FViewport* Viewport ) OVERRIDE;
-	virtual bool MouseMove( FLevelEditorViewportClient* ViewportClient, FViewport* Viewport, int32 x, int32 y ) OVERRIDE;
-	virtual bool InputKey( FLevelEditorViewportClient* InViewportClient, FViewport* InViewport, FKey InKey, EInputEvent InEvent ) OVERRIDE;
-	virtual bool EndTracking( FLevelEditorViewportClient* InViewportClient, FViewport* InViewport ) OVERRIDE;
-	virtual bool StartTracking( FLevelEditorViewportClient* InViewportClient, FViewport* InViewport ) OVERRIDE;
-
-	virtual bool HandleClick(FLevelEditorViewportClient* InViewportClient, HHitProxy *HitProxy, const FViewportClick &Click ) OVERRIDE;
-	virtual bool InputDelta( FLevelEditorViewportClient* InViewportClient, FViewport* InViewport, FVector& InDrag, FRotator& InRot, FVector& InScale ) OVERRIDE;
-	virtual bool ShouldDrawWidget() const OVERRIDE;
-	virtual bool UsesPropertyWidgets() const OVERRIDE;
+	virtual bool UsesToolkits() const override;
+	void Enter() override;
+	void Exit() override;
+	virtual void Tick(FLevelEditorViewportClient* ViewportClient,float DeltaTime) override;
+	virtual bool MouseEnter( FLevelEditorViewportClient* ViewportClient,FViewport* Viewport,int32 x, int32 y ) override;
+	virtual bool MouseLeave( FLevelEditorViewportClient* ViewportClient,FViewport* Viewport ) override;
+	virtual bool MouseMove( FLevelEditorViewportClient* ViewportClient, FViewport* Viewport, int32 x, int32 y ) override;
+	virtual bool InputKey( FLevelEditorViewportClient* InViewportClient, FViewport* InViewport, FKey InKey, EInputEvent InEvent ) override;
+	virtual bool EndTracking( FLevelEditorViewportClient* InViewportClient, FViewport* InViewport ) override;
+	virtual bool StartTracking( FLevelEditorViewportClient* InViewportClient, FViewport* InViewport ) override;
+
+	virtual bool HandleClick(FLevelEditorViewportClient* InViewportClient, HHitProxy *HitProxy, const FViewportClick &Click ) override;
+	virtual bool InputDelta( FLevelEditorViewportClient* InViewportClient, FViewport* InViewport, FVector& InDrag, FRotator& InRot, FVector& InScale ) override;
+	virtual bool ShouldDrawWidget() const override;
+	virtual bool UsesPropertyWidgets() const override;
 	// End of FEdMode interface
 
 public: //IPlacementMode
 
 	DECLARE_DERIVED_EVENT( FPlacementMode, IPlacementMode::FOnStoppedPlacingEvent, FOnStoppedPlacingEvent );
-	virtual FOnStoppedPlacingEvent& OnStoppedPlacing() OVERRIDE { return StoppedPlacingEvent; }
-	virtual void StopPlacing() OVERRIDE;
+	virtual FOnStoppedPlacingEvent& OnStoppedPlacing() override { return StoppedPlacingEvent; }
+	virtual void StopPlacing() override;
 
-	virtual bool IsCurrentlyPlacing() const OVERRIDE { return AssetsToPlace.Num() > 0; };
+	virtual bool IsCurrentlyPlacing() const override { return AssetsToPlace.Num() > 0; };
 
 	DECLARE_DERIVED_EVENT( FPlacementMode, IPlacementMode::FOnStartedPlacingEvent, FOnStartedPlacingEvent );
-	virtual FOnStartedPlacingEvent& OnStartedPlacing() OVERRIDE { return StartedPlacingEvent; }
-	virtual void StartPlacing( const TArray< UObject* >& Assets, UActorFactory* Factory = NULL ) OVERRIDE;
+	virtual FOnStartedPlacingEvent& OnStartedPlacing() override { return StartedPlacingEvent; }
+	virtual void StartPlacing( const TArray< UObject* >& Assets, UActorFactory* Factory = NULL ) override;
 
-	virtual UActorFactory* GetPlacingFactory() const OVERRIDE { return PlacementFactory.Get(); }
-	virtual void SetPlacingFactory( UActorFactory* Factory ) OVERRIDE;
-	virtual UActorFactory* FindLastUsedFactoryForAssetType( UObject* Asset ) const OVERRIDE;
+	virtual UActorFactory* GetPlacingFactory() const override { return PlacementFactory.Get(); }
+	virtual void SetPlacingFactory( UActorFactory* Factory ) override;
+	virtual UActorFactory* FindLastUsedFactoryForAssetType( UObject* Asset ) const override;
 
-	virtual void AddValidFocusTargetForPlacement( const TWeakPtr< SWidget >& Widget ) OVERRIDE { ValidFocusTargetsForPlacement.Add( Widget ); }
-	virtual void RemoveValidFocusTargetForPlacement( const TWeakPtr< SWidget >& Widget ) OVERRIDE { ValidFocusTargetsForPlacement.Remove( Widget ); }
+	virtual void AddValidFocusTargetForPlacement( const TWeakPtr< SWidget >& Widget ) override { ValidFocusTargetsForPlacement.Add( Widget ); }
+	virtual void RemoveValidFocusTargetForPlacement( const TWeakPtr< SWidget >& Widget ) override { ValidFocusTargetsForPlacement.Remove( Widget ); }
 
 	DECLARE_DERIVED_EVENT( FPlacementMode, IPlacementMode::FOnRecentlyPlacedChanged, FOnRecentlyPlacedChanged );
-	virtual FOnRecentlyPlacedChanged& OnRecentlyPlacedChanged() OVERRIDE { return RecentlyPlacedChanged; }
-	virtual void AddToRecentlyPlaced( const TArray< UObject* >& PlacedObjects, UActorFactory* FactoryUsed = NULL ) OVERRIDE;
-	virtual const TArray< FActorPlacementInfo >& GetRecentlyPlaced() const OVERRIDE { return RecentlyPlaced; }
+	virtual FOnRecentlyPlacedChanged& OnRecentlyPlacedChanged() override { return RecentlyPlacedChanged; }
+	virtual void AddToRecentlyPlaced( const TArray< UObject* >& PlacedObjects, UActorFactory* FactoryUsed = NULL ) override;
+	virtual const TArray< FActorPlacementInfo >& GetRecentlyPlaced() const override { return RecentlyPlaced; }
 
-	virtual const TArray< TWeakObjectPtr<UObject> >& GetCurrentlyPlacingObjects() const OVERRIDE;
+	virtual const TArray< TWeakObjectPtr<UObject> >& GetCurrentlyPlacingObjects() const override;
 
 private:
 
diff --git a/Engine/Source/Editor/PlacementMode/Private/PlacementModeModule.cpp b/Engine/Source/Editor/PlacementMode/Private/PlacementModeModule.cpp
--- a/Engine/Source/Editor/PlacementMode/Private/PlacementModeModule.cpp
+++ b/Engine/Source/Editor/PlacementMode/Private/PlacementModeModule.cpp
@@ -11,7 +11,7 @@ public:
 	/**
 	 * Called right after the module's DLL has been loaded and the module object has been created
 	 */
-	virtual void StartupModule() OVERRIDE
+	virtual void StartupModule() override
 	{
 		PlacementMode = FPlacementMode::Create();
 		GEditorModeTools().RegisterMode( PlacementMode.ToSharedRef() );
@@ -29,7 +29,7 @@ public:
 	/**
 	 * Called before the module is unloaded, right before the module object is destroyed.
 	 */
-	virtual void ShutdownModule() OVERRIDE
+	virtual void ShutdownModule() override
 	{
 		if ( PlacementMode.IsValid() )
 		{
@@ -47,29 +47,29 @@ public:
 		}
 	}
 
-	virtual void AddToRecentlyPlaced( const TArray< UObject* >& Assets, UActorFactory* FactoryUsed = NULL ) OVERRIDE
+	virtual void AddToRecentlyPlaced( const TArray< UObject* >& Assets, UActorFactory* FactoryUsed = NULL ) override
 	{
 		PlacementMode->AddToRecentlyPlaced( Assets, FactoryUsed );
 	}
 
-	virtual void AddToRecentlyPlaced( UObject* Asset, UActorFactory* FactoryUsed = NULL ) OVERRIDE
+	virtual void AddToRecentlyPlaced( UObject* Asset, UActorFactory* FactoryUsed = NULL ) override
 	{
 		TArray< UObject* > Assets;
 		Assets.Add( Asset );
 		PlacementMode->AddToRecentlyPlaced( Assets, FactoryUsed );
 	}
 
-	virtual const TArray< FActorPlacementInfo >& GetRecentlyPlaced() const OVERRIDE
+	virtual const TArray< FActorPlacementInfo >& GetRecentlyPlaced() const override
 	{
 		return PlacementMode->GetRecentlyPlaced();
 	}
 
-	virtual bool IsPlacementModeAvailable() const OVERRIDE
+	virtual bool IsPlacementModeAvailable() const override
 	{
 		return PlacementMode.IsValid();
 	}
 
-	virtual TSharedRef< IPlacementMode > GetPlacementMode() const OVERRIDE
+	virtual TSharedRef< IPlacementMode > GetPlacementMode() const override
 	{
 		return PlacementMode.ToSharedRef();
 	}
diff --git a/Engine/Source/Editor/PlacementMode/Private/PlacementModeToolkit.h b/Engine/Source/Editor/PlacementMode/Private/PlacementModeToolkit.h
--- a/Engine/Source/Editor/PlacementMode/Private/PlacementModeToolkit.h
+++ b/Engine/Source/Editor/PlacementMode/Private/PlacementModeToolkit.h
@@ -14,10 +14,10 @@ public:
 	}
 
 	/** IToolkit interface */
-	virtual FName GetToolkitFName() const OVERRIDE { return FName("PlacementMode"); }
-	virtual FText GetBaseToolkitName() const OVERRIDE { return NSLOCTEXT("BuilderModeToolkit", "DisplayName", "Builder"); }
-	virtual class FEdMode* GetEditorMode() const OVERRIDE { return GEditorModeTools().GetActiveMode( FBuiltinEditorModes::EM_Default ); }
-	virtual TSharedPtr<class SWidget> GetInlineContent() const OVERRIDE { return PlacementModeTools; }
+	virtual FName GetToolkitFName() const override { return FName("PlacementMode"); }
+	virtual FText GetBaseToolkitName() const override { return NSLOCTEXT("BuilderModeToolkit", "DisplayName", "Builder"); }
+	virtual class FEdMode* GetEditorMode() const override { return GEditorModeTools().GetActiveMode( FBuiltinEditorModes::EM_Default ); }
+	virtual TSharedPtr<class SWidget> GetInlineContent() const override { return PlacementModeTools; }
 
 private:
 
diff --git a/Engine/Source/Editor/PlacementMode/Private/SPlacementModeTools.h b/Engine/Source/Editor/PlacementMode/Private/SPlacementModeTools.h
--- a/Engine/Source/Editor/PlacementMode/Private/SPlacementModeTools.h
+++ b/Engine/Source/Editor/PlacementMode/Private/SPlacementModeTools.h
@@ -26,9 +26,9 @@ public:
 
 	void Construct(const FArguments& InArgs, UActorFactory* InFactory, const FAssetData& InAsset);
 
-	virtual FReply OnMouseButtonDown(const FGeometry& MyGeometry, const FPointerEvent& MouseEvent) OVERRIDE;
-	virtual FReply OnMouseButtonUp(const FGeometry& MyGeometry, const FPointerEvent& MouseEvent) OVERRIDE;
-	virtual FReply OnDragDetected(const FGeometry& MyGeometry, const FPointerEvent& MouseEvent) OVERRIDE;
+	virtual FReply OnMouseButtonDown(const FGeometry& MyGeometry, const FPointerEvent& MouseEvent) override;
+	virtual FReply OnMouseButtonUp(const FGeometry& MyGeometry, const FPointerEvent& MouseEvent) override;
+	virtual FReply OnDragDetected(const FGeometry& MyGeometry, const FPointerEvent& MouseEvent) override;
 
 	bool IsPressed() const;
 
@@ -99,8 +99,8 @@ private:
 	void RefreshVolumes();
 
 	// Begin SWidget
-	void Tick( const FGeometry& AllottedGeometry, const double InCurrentTime, const float InDeltaTime ) OVERRIDE;
-	virtual FReply OnKeyDown( const FGeometry& MyGeometry, const FKeyboardEvent& InKeyboardEvent ) OVERRIDE;
+	void Tick( const FGeometry& AllottedGeometry, const double InCurrentTime, const float InDeltaTime ) override;
+	virtual FReply OnKeyDown( const FGeometry& MyGeometry, const FKeyboardEvent& InKeyboardEvent ) override;
 	// End SWidget
 
 private:
diff --git a/Engine/Source/Editor/ProjectTargetPlatformEditor/Private/ProjectTargetPlatformEditorModule.cpp b/Engine/Source/Editor/ProjectTargetPlatformEditor/Private/ProjectTargetPlatformEditorModule.cpp
--- a/Engine/Source/Editor/ProjectTargetPlatformEditor/Private/ProjectTargetPlatformEditorModule.cpp
+++ b/Engine/Source/Editor/ProjectTargetPlatformEditor/Private/ProjectTargetPlatformEditorModule.cpp
@@ -14,7 +14,7 @@ public:
 
 	// Begin IProjectTargetPlatformEditorModule interface
 
-	virtual TWeakPtr<SWidget> CreateProjectTargetPlatformEditorPanel() OVERRIDE
+	virtual TWeakPtr<SWidget> CreateProjectTargetPlatformEditorPanel() override
 	{
 		TSharedPtr<SWidget> Panel = SNew(SProjectTargetPlatformSettings);
 		EditorPanels.Add(Panel);
@@ -22,12 +22,12 @@ public:
 		return Panel;
 	}
 
-	virtual void DestroyProjectTargetPlatformEditorPanel(const TWeakPtr<SWidget>& Panel) OVERRIDE
+	virtual void DestroyProjectTargetPlatformEditorPanel(const TWeakPtr<SWidget>& Panel) override
 	{
 		EditorPanels.Remove(Panel.Pin());
 	}
 
-	virtual void AddOpenProjectTargetPlatformEditorMenuItem(FMenuBuilder& MenuBuilder) const OVERRIDE
+	virtual void AddOpenProjectTargetPlatformEditorMenuItem(FMenuBuilder& MenuBuilder) const override
 	{
 		struct Local
 		{
@@ -45,7 +45,7 @@ public:
 			);
 	}
 
-	virtual TSharedRef<SWidget> MakePlatformMenuItemWidget(const PlatformInfo::FPlatformInfo& PlatformInfo, const bool bForCheckBox = false, const FText& DisplayNameOverride = FText()) const OVERRIDE
+	virtual TSharedRef<SWidget> MakePlatformMenuItemWidget(const PlatformInfo::FPlatformInfo& PlatformInfo, const bool bForCheckBox = false, const FText& DisplayNameOverride = FText()) const override
 	{
 		struct Local
 		{
@@ -103,7 +103,7 @@ public:
 			];
 	}
 
-	virtual bool ShowUnsupportedTargetWarning(const FName PlatformName) const OVERRIDE
+	virtual bool ShowUnsupportedTargetWarning(const FName PlatformName) const override
 	{
 		const PlatformInfo::FPlatformInfo* const PlatformInfo = PlatformInfo::FindPlatformInfo(PlatformName);
 		check(PlatformInfo);
diff --git a/Engine/Source/Editor/PropertyEditor/Private/CategoryPropertyNode.h b/Engine/Source/Editor/PropertyEditor/Private/CategoryPropertyNode.h
--- a/Engine/Source/Editor/PropertyEditor/Private/CategoryPropertyNode.h
+++ b/Engine/Source/Editor/PropertyEditor/Private/CategoryPropertyNode.h
@@ -11,7 +11,7 @@ public:
 	/**
 	 * Overridden function to get the derived object node
 	 */
-	virtual FCategoryPropertyNode* AsCategoryNode() OVERRIDE { return this; }
+	virtual FCategoryPropertyNode* AsCategoryNode() override { return this; }
 
 	/**
 	 * Accessor functions for getting a category name
@@ -26,22 +26,22 @@ public:
 	 */
 	bool IsSubcategory() const;
 
-	virtual FString GetDisplayName() const OVERRIDE;
+	virtual FString GetDisplayName() const override;
 
 protected:
 	/**
 	 * Overridden function for special setup
 	 */
-	virtual void InitBeforeNodeFlags() OVERRIDE;
+	virtual void InitBeforeNodeFlags() override;
 	/**
 	 * Overridden function for Creating Child Nodes
 	 */
-	virtual void InitChildNodes() OVERRIDE;
+	virtual void InitChildNodes() override;
 
 	/**
 	 * Appends my path, including an array index (where appropriate)
 	 */
-	virtual void GetQualifiedName( FString& PathPlusIndex, const bool bWithArrayIndex ) const OVERRIDE;
+	virtual void GetQualifiedName( FString& PathPlusIndex, const bool bWithArrayIndex ) const override;
 
 	FString GetSubcategoryName() const;
 	/**
diff --git a/Engine/Source/Editor/PropertyEditor/Private/CustomChildBuilder.cpp b/Engine/Source/Editor/PropertyEditor/Private/CustomChildBuilder.cpp
--- a/Engine/Source/Editor/PropertyEditor/Private/CustomChildBuilder.cpp
+++ b/Engine/Source/Editor/PropertyEditor/Private/CustomChildBuilder.cpp
@@ -72,7 +72,7 @@ public:
 		];
 	}
 
-	virtual TSharedPtr<FAssetThumbnailPool> GetThumbnailPool() const OVERRIDE
+	virtual TSharedPtr<FAssetThumbnailPool> GetThumbnailPool() const override
 	{
 		TSharedPtr<FDetailCategoryImpl> ParentCategoryPinned = ParentCategory.Pin();
 		return ParentCategoryPinned.IsValid() ? ParentCategoryPinned->GetParentLayout().GetThumbnailPool() : NULL;
diff --git a/Engine/Source/Editor/PropertyEditor/Private/CustomChildBuilder.h b/Engine/Source/Editor/PropertyEditor/Private/CustomChildBuilder.h
--- a/Engine/Source/Editor/PropertyEditor/Private/CustomChildBuilder.h
+++ b/Engine/Source/Editor/PropertyEditor/Private/CustomChildBuilder.h
@@ -9,10 +9,10 @@ public:
 		: ParentCategory( InParentCategory )
 	{}
 	
-	virtual IDetailChildrenBuilder& AddChildCustomBuilder( TSharedRef<class IDetailCustomNodeBuilder> InCustomBuilder ) OVERRIDE;
-	virtual IDetailGroup& AddChildGroup( FName GroupName, const FString& LocalizedDisplayName ) OVERRIDE;
-	virtual FDetailWidgetRow& AddChildContent( const FString& SearchString ) OVERRIDE;
-	virtual IDetailPropertyRow& AddChildProperty( TSharedRef<IPropertyHandle> PropertyHandle ) OVERRIDE;
+	virtual IDetailChildrenBuilder& AddChildCustomBuilder( TSharedRef<class IDetailCustomNodeBuilder> InCustomBuilder ) override;
+	virtual IDetailGroup& AddChildGroup( FName GroupName, const FString& LocalizedDisplayName ) override;
+	virtual FDetailWidgetRow& AddChildContent( const FString& SearchString ) override;
+	virtual IDetailPropertyRow& AddChildProperty( TSharedRef<IPropertyHandle> PropertyHandle ) override;
 	virtual TSharedRef<SWidget> GenerateStructValueWidget( TSharedRef<IPropertyHandle> StructPropertyHandle );
 
 	const TArray< FDetailLayoutCustomization >& GetChildCustomizations() const { return ChildCustomizations; }
diff --git a/Engine/Source/Editor/PropertyEditor/Private/DetailAdvancedDropdownNode.h b/Engine/Source/Editor/PropertyEditor/Private/DetailAdvancedDropdownNode.h
--- a/Engine/Source/Editor/PropertyEditor/Private/DetailAdvancedDropdownNode.h
+++ b/Engine/Source/Editor/PropertyEditor/Private/DetailAdvancedDropdownNode.h
@@ -25,16 +25,16 @@ public:
 	{}
 private:
 	/** IDetailTreeNode Interface */
-	virtual IDetailsViewPrivate& GetDetailsView() const OVERRIDE{ return ParentCategory.GetDetailsView(); }
-	virtual TSharedRef< ITableRow > GenerateNodeWidget( const TSharedRef<STableViewBase>& OwnerTable, const FDetailColumnSizeData& ColumnSizeData, const TSharedRef<IPropertyUtilities>& PropertyUtilities ) OVERRIDE;
-	virtual void GetChildren( TArray< TSharedRef<IDetailTreeNode> >& OutChildren )  OVERRIDE {}
-	virtual void OnItemExpansionChanged( bool bIsExpanded ) OVERRIDE {}
-	virtual bool ShouldBeExpanded() const OVERRIDE { return false; }
-	virtual ENodeVisibility::Type GetVisibility() const OVERRIDE { return ENodeVisibility::Visible; }
-	virtual void FilterNode( const FDetailFilter& InFilter ) OVERRIDE {}
-	virtual void Tick( float DeltaTime ) OVERRIDE {}
-	virtual bool ShouldShowOnlyChildren() const OVERRIDE { return false; }
-	virtual FName GetNodeName() const OVERRIDE { return NAME_None; }
+	virtual IDetailsViewPrivate& GetDetailsView() const override{ return ParentCategory.GetDetailsView(); }
+	virtual TSharedRef< ITableRow > GenerateNodeWidget( const TSharedRef<STableViewBase>& OwnerTable, const FDetailColumnSizeData& ColumnSizeData, const TSharedRef<IPropertyUtilities>& PropertyUtilities ) override;
+	virtual void GetChildren( TArray< TSharedRef<IDetailTreeNode> >& OutChildren )  override {}
+	virtual void OnItemExpansionChanged( bool bIsExpanded ) override {}
+	virtual bool ShouldBeExpanded() const override { return false; }
+	virtual ENodeVisibility::Type GetVisibility() const override { return ENodeVisibility::Visible; }
+	virtual void FilterNode( const FDetailFilter& InFilter ) override {}
+	virtual void Tick( float DeltaTime ) override {}
+	virtual bool ShouldShowOnlyChildren() const override { return false; }
+	virtual FName GetNodeName() const override { return NAME_None; }
 
 	/** Called when the advanced drop down arrow is clicked */
 	FReply OnAdvancedDropDownClicked();
diff --git a/Engine/Source/Editor/PropertyEditor/Private/DetailCategoryBuilderImpl.h b/Engine/Source/Editor/PropertyEditor/Private/DetailCategoryBuilderImpl.h
--- a/Engine/Source/Editor/PropertyEditor/Private/DetailCategoryBuilderImpl.h
+++ b/Engine/Source/Editor/PropertyEditor/Private/DetailCategoryBuilderImpl.h
@@ -130,30 +130,30 @@ public:
 	~FDetailCategoryImpl();
 
 	/** IDetailCategoryBuilder interface */
-	virtual IDetailCategoryBuilder& InitiallyCollapsed( bool bShouldBeInitiallyCollapsed ) OVERRIDE;
-	virtual IDetailCategoryBuilder& OnExpansionChanged( FOnBooleanValueChanged InOnExpansionChanged ) OVERRIDE;
-	virtual IDetailCategoryBuilder& RestoreExpansionState( bool bRestore ) OVERRIDE;
-	virtual IDetailCategoryBuilder& HeaderContent( TSharedRef<SWidget> InHeaderContent ) OVERRIDE;
-	virtual IDetailPropertyRow& AddProperty( FName PropertyPath, UClass* ClassOuter = NULL, FName InstanceName = NAME_None,  EPropertyLocation::Type Location = EPropertyLocation::Default ) OVERRIDE;
-	virtual IDetailPropertyRow& AddProperty( TSharedPtr<IPropertyHandle> PropertyHandle,  EPropertyLocation::Type Location = EPropertyLocation::Default ) OVERRIDE;
-	virtual IDetailPropertyRow* AddExternalProperty( const TArray<UObject*>& Objects, FName PropertyName, EPropertyLocation::Type Location = EPropertyLocation::Default ) OVERRIDE;
-	virtual IDetailLayoutBuilder& GetParentLayout() const OVERRIDE { return *DetailLayoutBuilder.Pin(); }
-	virtual FDetailWidgetRow& AddCustomRow( const FString& FilterString, bool bForAdvanced = false ) OVERRIDE;
-	virtual void AddCustomBuilder( TSharedRef<IDetailCustomNodeBuilder> InCustomBuilder, bool bForAdvanced = false ) OVERRIDE;
-	virtual IDetailGroup& AddGroup( FName GroupName, const FString& LocalizedDisplayName, bool bForAdvanced = false ) OVERRIDE;
-	virtual void GetDefaultProperties( TArray<TSharedRef<IPropertyHandle> >& OutAllProperties, bool bSimpleProperties = true, bool bAdvancedProperties = true ) OVERRIDE;
-	virtual const FString& GetDisplayName() const OVERRIDE { return DisplayName; }
+	virtual IDetailCategoryBuilder& InitiallyCollapsed( bool bShouldBeInitiallyCollapsed ) override;
+	virtual IDetailCategoryBuilder& OnExpansionChanged( FOnBooleanValueChanged InOnExpansionChanged ) override;
+	virtual IDetailCategoryBuilder& RestoreExpansionState( bool bRestore ) override;
+	virtual IDetailCategoryBuilder& HeaderContent( TSharedRef<SWidget> InHeaderContent ) override;
+	virtual IDetailPropertyRow& AddProperty( FName PropertyPath, UClass* ClassOuter = NULL, FName InstanceName = NAME_None,  EPropertyLocation::Type Location = EPropertyLocation::Default ) override;
+	virtual IDetailPropertyRow& AddProperty( TSharedPtr<IPropertyHandle> PropertyHandle,  EPropertyLocation::Type Location = EPropertyLocation::Default ) override;
+	virtual IDetailPropertyRow* AddExternalProperty( const TArray<UObject*>& Objects, FName PropertyName, EPropertyLocation::Type Location = EPropertyLocation::Default ) override;
+	virtual IDetailLayoutBuilder& GetParentLayout() const override { return *DetailLayoutBuilder.Pin(); }
+	virtual FDetailWidgetRow& AddCustomRow( const FString& FilterString, bool bForAdvanced = false ) override;
+	virtual void AddCustomBuilder( TSharedRef<IDetailCustomNodeBuilder> InCustomBuilder, bool bForAdvanced = false ) override;
+	virtual IDetailGroup& AddGroup( FName GroupName, const FString& LocalizedDisplayName, bool bForAdvanced = false ) override;
+	virtual void GetDefaultProperties( TArray<TSharedRef<IPropertyHandle> >& OutAllProperties, bool bSimpleProperties = true, bool bAdvancedProperties = true ) override;
+	virtual const FString& GetDisplayName() const override { return DisplayName; }
 
 	/** IDetailTreeNode interface */
-	virtual IDetailsViewPrivate& GetDetailsView() const OVERRIDE{ return DetailLayoutBuilder.Pin()->GetDetailsView(); }
-	virtual TSharedRef< ITableRow > GenerateNodeWidget( const TSharedRef<STableViewBase>& OwnerTable, const FDetailColumnSizeData& ColumnSizeData, const TSharedRef<IPropertyUtilities>& PropertyUtilities ) OVERRIDE;
-	virtual void GetChildren( TArray< TSharedRef<IDetailTreeNode> >& OutChildren ) OVERRIDE;
-	virtual bool ShouldBeExpanded() const OVERRIDE;
-	virtual ENodeVisibility::Type GetVisibility() const OVERRIDE;
-	virtual void FilterNode( const FDetailFilter& DetailFilter ) OVERRIDE;
-	virtual void Tick( float DeltaTime ) OVERRIDE {}
-	virtual bool ShouldShowOnlyChildren() const OVERRIDE { return false; }
-	virtual FName GetNodeName() const OVERRIDE { return GetCategoryName(); }
+	virtual IDetailsViewPrivate& GetDetailsView() const override{ return DetailLayoutBuilder.Pin()->GetDetailsView(); }
+	virtual TSharedRef< ITableRow > GenerateNodeWidget( const TSharedRef<STableViewBase>& OwnerTable, const FDetailColumnSizeData& ColumnSizeData, const TSharedRef<IPropertyUtilities>& PropertyUtilities ) override;
+	virtual void GetChildren( TArray< TSharedRef<IDetailTreeNode> >& OutChildren ) override;
+	virtual bool ShouldBeExpanded() const override;
+	virtual ENodeVisibility::Type GetVisibility() const override;
+	virtual void FilterNode( const FDetailFilter& DetailFilter ) override;
+	virtual void Tick( float DeltaTime ) override {}
+	virtual bool ShouldShowOnlyChildren() const override { return false; }
+	virtual FName GetNodeName() const override { return GetCategoryName(); }
 
 	/**
 	 * @return The name of the category
@@ -252,7 +252,7 @@ public:
 	void OnAdvancedDropdownClicked();
 
 private:
-	virtual void OnItemExpansionChanged( bool bIsExpanded ) OVERRIDE;
+	virtual void OnItemExpansionChanged( bool bIsExpanded ) override;
 
 	/**
 	 * Adds a new filter widget to this category (for checking if anything is visible in the category when filtered)
diff --git a/Engine/Source/Editor/PropertyEditor/Private/DetailCategoryGroupNode.h b/Engine/Source/Editor/PropertyEditor/Private/DetailCategoryGroupNode.h
--- a/Engine/Source/Editor/PropertyEditor/Private/DetailCategoryGroupNode.h
+++ b/Engine/Source/Editor/PropertyEditor/Private/DetailCategoryGroupNode.h
@@ -20,8 +20,8 @@ private:
 	EVisibility IsSeparatorVisible() const;
 	const FSlateBrush* GetBackgroundImage() const;
 private:
-	virtual FReply OnMouseButtonDown( const FGeometry& MyGeometry, const FPointerEvent& MouseEvent ) OVERRIDE;
-	virtual FReply OnMouseButtonDoubleClick( const FGeometry& InMyGeometry, const FPointerEvent& InMouseEvent ) OVERRIDE;
+	virtual FReply OnMouseButtonDown( const FGeometry& MyGeometry, const FPointerEvent& MouseEvent ) override;
+	virtual FReply OnMouseButtonDoubleClick( const FGeometry& InMyGeometry, const FPointerEvent& InMouseEvent ) override;
 private:
 	bool bIsInnerCategory;
 };
@@ -33,16 +33,16 @@ public:
 	FDetailCategoryGroupNode( const FDetailNodeList& InChildNodes, FName InGroupName, FDetailCategoryImpl& InParentCategory );
 
 private:
-	virtual IDetailsViewPrivate& GetDetailsView() const OVERRIDE{ return ParentCategory.GetDetailsView(); }
-	virtual void OnItemExpansionChanged( bool bIsExpanded ) OVERRIDE {}
-	virtual bool ShouldBeExpanded() const OVERRIDE { return true; }
-	virtual ENodeVisibility::Type GetVisibility() const OVERRIDE { return bShouldBeVisible ? ENodeVisibility::Visible : ENodeVisibility::HiddenDueToFiltering; }
-	virtual TSharedRef< ITableRow > GenerateNodeWidget( const TSharedRef<STableViewBase>& OwnerTable, const FDetailColumnSizeData& ColumnSizeData, const TSharedRef<IPropertyUtilities>& PropertyUtilities ) OVERRIDE;
-	virtual void GetChildren( TArray< TSharedRef<IDetailTreeNode> >& OutChildren )  OVERRIDE;
-	virtual void FilterNode( const FDetailFilter& InFilter ) OVERRIDE;
-	virtual void Tick( float DeltaTime ) OVERRIDE {}
-	virtual bool ShouldShowOnlyChildren() const OVERRIDE { return false; }
-	virtual FName GetNodeName() const OVERRIDE { return NAME_None; }
+	virtual IDetailsViewPrivate& GetDetailsView() const override{ return ParentCategory.GetDetailsView(); }
+	virtual void OnItemExpansionChanged( bool bIsExpanded ) override {}
+	virtual bool ShouldBeExpanded() const override { return true; }
+	virtual ENodeVisibility::Type GetVisibility() const override { return bShouldBeVisible ? ENodeVisibility::Visible : ENodeVisibility::HiddenDueToFiltering; }
+	virtual TSharedRef< ITableRow > GenerateNodeWidget( const TSharedRef<STableViewBase>& OwnerTable, const FDetailColumnSizeData& ColumnSizeData, const TSharedRef<IPropertyUtilities>& PropertyUtilities ) override;
+	virtual void GetChildren( TArray< TSharedRef<IDetailTreeNode> >& OutChildren )  override;
+	virtual void FilterNode( const FDetailFilter& InFilter ) override;
+	virtual void Tick( float DeltaTime ) override {}
+	virtual bool ShouldShowOnlyChildren() const override { return false; }
+	virtual FName GetNodeName() const override { return NAME_None; }
 private:
 	FDetailNodeList ChildNodes;
 	FDetailCategoryImpl& ParentCategory;
diff --git a/Engine/Source/Editor/PropertyEditor/Private/DetailGroup.h b/Engine/Source/Editor/PropertyEditor/Private/DetailGroup.h
--- a/Engine/Source/Editor/PropertyEditor/Private/DetailGroup.h
+++ b/Engine/Source/Editor/PropertyEditor/Private/DetailGroup.h
@@ -8,11 +8,11 @@ public:
 	FDetailGroup( const FName InGroupName, TSharedRef<FDetailCategoryImpl> InParentCategory, const FString& InLocalizedDisplayName );
 
 	/** IDetailGroup interface */     
-	virtual FDetailWidgetRow& HeaderRow() OVERRIDE;
-	virtual IDetailPropertyRow& HeaderProperty( TSharedRef<IPropertyHandle> PropertyHandle ) OVERRIDE;
-	virtual FDetailWidgetRow& AddWidgetRow() OVERRIDE;
-	virtual IDetailPropertyRow& AddPropertyRow( TSharedRef<IPropertyHandle> PropertyHandle ) OVERRIDE;
-	virtual void ToggleExpansion( bool bExpand ) OVERRIDE;
+	virtual FDetailWidgetRow& HeaderRow() override;
+	virtual IDetailPropertyRow& HeaderProperty( TSharedRef<IPropertyHandle> PropertyHandle ) override;
+	virtual FDetailWidgetRow& AddWidgetRow() override;
+	virtual IDetailPropertyRow& AddPropertyRow( TSharedRef<IPropertyHandle> PropertyHandle ) override;
+	virtual void ToggleExpansion( bool bExpand ) override;
 
 	/** @return The name of the group */
 	FName GetGroupName() const { return GroupName; }
diff --git a/Engine/Source/Editor/PropertyEditor/Private/DetailItemNode.h b/Engine/Source/Editor/PropertyEditor/Private/DetailItemNode.h
--- a/Engine/Source/Editor/PropertyEditor/Private/DetailItemNode.h
+++ b/Engine/Source/Editor/PropertyEditor/Private/DetailItemNode.h
@@ -31,16 +31,16 @@ public:
 	bool HasMultiColumnWidget() const;
 
 	/** IDetailTreeNode interface */
-	virtual IDetailsViewPrivate& GetDetailsView() const OVERRIDE{ return ParentCategory.Pin()->GetDetailsView(); }
-	virtual TSharedRef< ITableRow > GenerateNodeWidget( const TSharedRef<STableViewBase>& OwnerTable, const FDetailColumnSizeData& ColumnSizeData, const TSharedRef<IPropertyUtilities>& PropertyUtilities ) OVERRIDE;
-	virtual void GetChildren( FDetailNodeList& OutChildren )  OVERRIDE;
-	virtual void OnItemExpansionChanged( bool bInIsExpanded ) OVERRIDE;
-	virtual bool ShouldBeExpanded() const OVERRIDE;
-	virtual ENodeVisibility::Type GetVisibility() const OVERRIDE;
-	virtual void FilterNode( const FDetailFilter& InFilter ) OVERRIDE;
-	virtual void Tick( float DeltaTime ) OVERRIDE;
-	virtual bool ShouldShowOnlyChildren() const OVERRIDE;
-	virtual FName GetNodeName() const OVERRIDE;
+	virtual IDetailsViewPrivate& GetDetailsView() const override{ return ParentCategory.Pin()->GetDetailsView(); }
+	virtual TSharedRef< ITableRow > GenerateNodeWidget( const TSharedRef<STableViewBase>& OwnerTable, const FDetailColumnSizeData& ColumnSizeData, const TSharedRef<IPropertyUtilities>& PropertyUtilities ) override;
+	virtual void GetChildren( FDetailNodeList& OutChildren )  override;
+	virtual void OnItemExpansionChanged( bool bInIsExpanded ) override;
+	virtual bool ShouldBeExpanded() const override;
+	virtual ENodeVisibility::Type GetVisibility() const override;
+	virtual void FilterNode( const FDetailFilter& InFilter ) override;
+	virtual void Tick( float DeltaTime ) override;
+	virtual bool ShouldShowOnlyChildren() const override;
+	virtual FName GetNodeName() const override;
 
 private:
 
diff --git a/Engine/Source/Editor/PropertyEditor/Private/DetailLayoutBuilderImpl.h b/Engine/Source/Editor/PropertyEditor/Private/DetailLayoutBuilderImpl.h
--- a/Engine/Source/Editor/PropertyEditor/Private/DetailLayoutBuilderImpl.h
+++ b/Engine/Source/Editor/PropertyEditor/Private/DetailLayoutBuilderImpl.h
@@ -9,18 +9,18 @@ public:
 	FDetailLayoutBuilderImpl(FClassToPropertyMap& InPropertyMap, const TSharedRef< class IPropertyUtilities >& InPropertyUtilities, const TSharedRef< IDetailsViewPrivate >& InDetailsView);
 
 	/** IDetailLayoutBuilder Interface */
-	virtual const IDetailsView& GetDetailsView() const OVERRIDE;
-	virtual void GetObjectsBeingCustomized( TArray< TWeakObjectPtr<UObject> >& OutObjects ) const OVERRIDE;
+	virtual const IDetailsView& GetDetailsView() const override;
+	virtual void GetObjectsBeingCustomized( TArray< TWeakObjectPtr<UObject> >& OutObjects ) const override;
 	virtual IDetailCategoryBuilder& EditCategory( FName CategoryName, const FString& NewLocalizedDisplayName = TEXT(""), ECategoryPriority::Type CategoryType = ECategoryPriority::Default );
-	virtual TSharedRef<IPropertyHandle> GetProperty( const FName PropertyPath, const UClass* ClassOutermost, FName InInstanceName ) OVERRIDE;
-	virtual void HideProperty( const TSharedPtr<IPropertyHandle> Property ) OVERRIDE;
-	virtual void HideProperty( FName PropertyPath, const UClass* ClassOutermost = NULL, FName InstanceName = NAME_None ) OVERRIDE;
-	virtual void ForceRefreshDetails() OVERRIDE;
-	virtual TSharedPtr<FAssetThumbnailPool> GetThumbnailPool() const OVERRIDE;
-	virtual bool IsPropertyVisible( TSharedRef<IPropertyHandle> PropertyHandle ) const OVERRIDE;
-	virtual bool IsPropertyVisible( UProperty* PropertyHandle ) const OVERRIDE;
-
-	virtual const TSharedRef< IPropertyUtilities >& GetPropertyUtilities() const OVERRIDE; 
+	virtual TSharedRef<IPropertyHandle> GetProperty( const FName PropertyPath, const UClass* ClassOutermost, FName InInstanceName ) override;
+	virtual void HideProperty( const TSharedPtr<IPropertyHandle> Property ) override;
+	virtual void HideProperty( FName PropertyPath, const UClass* ClassOutermost = NULL, FName InstanceName = NAME_None ) override;
+	virtual void ForceRefreshDetails() override;
+	virtual TSharedPtr<FAssetThumbnailPool> GetThumbnailPool() const override;
+	virtual bool IsPropertyVisible( TSharedRef<IPropertyHandle> PropertyHandle ) const override;
+	virtual bool IsPropertyVisible( UProperty* PropertyHandle ) const override;
+
+	virtual const TSharedRef< IPropertyUtilities >& GetPropertyUtilities() const override; 
 
 	/**
 	 * Creates a default category. The SDetails view will generate widgets in default categories 
diff --git a/Engine/Source/Editor/PropertyEditor/Private/DetailPropertyRow.h b/Engine/Source/Editor/PropertyEditor/Private/DetailPropertyRow.h
--- a/Engine/Source/Editor/PropertyEditor/Private/DetailPropertyRow.h
+++ b/Engine/Source/Editor/PropertyEditor/Private/DetailPropertyRow.h
@@ -8,20 +8,20 @@ public:
 	FDetailPropertyRow( TSharedPtr<FPropertyNode> InPropertyNode, TSharedRef<FDetailCategoryImpl> InParentCategory, TSharedPtr<FObjectPropertyNode> InExternalRootNode = NULL);
 
 	/** IDetailPropertyRow interface */
-	virtual IDetailPropertyRow& DisplayName( const FString& InDisplayName ) OVERRIDE;
-	virtual IDetailPropertyRow& ToolTip( const FString& InToolTip ) OVERRIDE;
-	virtual IDetailPropertyRow& ShowPropertyButtons( bool bInShowPropertyButtons ) OVERRIDE;
-	virtual IDetailPropertyRow& EditCondition( TAttribute<bool> EditConditionValue, FOnBooleanValueChanged OnEditConditionValueChanged ) OVERRIDE;
-	virtual IDetailPropertyRow& IsEnabled( TAttribute<bool> InIsEnabled ) OVERRIDE;
-	virtual IDetailPropertyRow& Visibility( TAttribute<EVisibility> Visibility ) OVERRIDE;
-	virtual IDetailPropertyRow& OverrideResetToDefault( TAttribute<bool> IsResetToDefaultVisible, FSimpleDelegate OnResetToDefaultClicked ) OVERRIDE;
-	virtual FDetailWidgetRow& CustomWidget( bool bShowChildren = false ) OVERRIDE;
-	virtual void GetDefaultWidgets( TSharedPtr<SWidget>& OutNameWidget, TSharedPtr<SWidget>& OutValueWidget ) OVERRIDE;
-	virtual void GetDefaultWidgets( TSharedPtr<SWidget>& OutNameWidget, TSharedPtr<SWidget>& OutValueWidget, FDetailWidgetRow& Row ) OVERRIDE;
+	virtual IDetailPropertyRow& DisplayName( const FString& InDisplayName ) override;
+	virtual IDetailPropertyRow& ToolTip( const FString& InToolTip ) override;
+	virtual IDetailPropertyRow& ShowPropertyButtons( bool bInShowPropertyButtons ) override;
+	virtual IDetailPropertyRow& EditCondition( TAttribute<bool> EditConditionValue, FOnBooleanValueChanged OnEditConditionValueChanged ) override;
+	virtual IDetailPropertyRow& IsEnabled( TAttribute<bool> InIsEnabled ) override;
+	virtual IDetailPropertyRow& Visibility( TAttribute<EVisibility> Visibility ) override;
+	virtual IDetailPropertyRow& OverrideResetToDefault( TAttribute<bool> IsResetToDefaultVisible, FSimpleDelegate OnResetToDefaultClicked ) override;
+	virtual FDetailWidgetRow& CustomWidget( bool bShowChildren = false ) override;
+	virtual void GetDefaultWidgets( TSharedPtr<SWidget>& OutNameWidget, TSharedPtr<SWidget>& OutValueWidget ) override;
+	virtual void GetDefaultWidgets( TSharedPtr<SWidget>& OutNameWidget, TSharedPtr<SWidget>& OutValueWidget, FDetailWidgetRow& Row ) override;
 
 	/** IPropertyTypeCustomizationUtils interface */
-	virtual TSharedPtr<class FAssetThumbnailPool> GetThumbnailPool() const OVERRIDE;
-	virtual TSharedPtr<class IPropertyUtilities> GetPropertyUtilities() const OVERRIDE;
+	virtual TSharedPtr<class FAssetThumbnailPool> GetThumbnailPool() const override;
+	virtual TSharedPtr<class IPropertyUtilities> GetPropertyUtilities() const override;
 
 	/** @return true if this row has widgets with columns */
 	bool HasColumns() const;
diff --git a/Engine/Source/Editor/PropertyEditor/Private/ItemPropertyNode.h b/Engine/Source/Editor/PropertyEditor/Private/ItemPropertyNode.h
--- a/Engine/Source/Editor/PropertyEditor/Private/ItemPropertyNode.h
+++ b/Engine/Source/Editor/PropertyEditor/Private/ItemPropertyNode.h
@@ -28,36 +28,36 @@ public:
 	/**
 	 * Overridden function to get the derived object node
 	 */
-	virtual FItemPropertyNode* AsItemPropertyNode() OVERRIDE { return this; }
+	virtual FItemPropertyNode* AsItemPropertyNode() override { return this; }
 
 	/** Display name override to use instead of the property name */
-	void SetDisplayNameOverride( const FString& InDisplayNameOverride ) OVERRIDE;
+	void SetDisplayNameOverride( const FString& InDisplayNameOverride ) override;
 
 	/**
 	 * @return The formatted display name for the property in this node                                                              
 	 */
-	FString GetDisplayName() const OVERRIDE;
+	FString GetDisplayName() const override;
 
 	/**
 	 * Sets the tooltip override to use instead of the property tooltip
 	 */
-	virtual void SetToolTipOverride( const FString& InToolTipOverride ) OVERRIDE;
+	virtual void SetToolTipOverride( const FString& InToolTipOverride ) override;
 
 	/**
 	 * @return The tooltip for the property in this node                                                              
 	 */
-	virtual FString GetToolTipText() const OVERRIDE;
+	virtual FString GetToolTipText() const override;
 
 protected:
 	/**
 	 * Overridden function for special setup
 	 */
-	virtual void InitExpansionFlags() OVERRIDE;
+	virtual void InitExpansionFlags() override;
 
 	/**
 	 * Overridden function for Creating Child Nodes
 	 */
-	virtual void InitChildNodes() OVERRIDE;
+	virtual void InitChildNodes() override;
 
 private:
 	/** Display name override to use instead of the property name */
diff --git a/Engine/Source/Editor/PropertyEditor/Private/ObjectPropertyNode.h b/Engine/Source/Editor/PropertyEditor/Private/ObjectPropertyNode.h
--- a/Engine/Source/Editor/PropertyEditor/Private/ObjectPropertyNode.h
+++ b/Engine/Source/Editor/PropertyEditor/Private/ObjectPropertyNode.h
@@ -18,12 +18,12 @@ public:
 	virtual ~FObjectPropertyNode();
 
 	/** FPropertyNode Interface */
-	virtual FObjectPropertyNode* AsObjectNode() OVERRIDE { return this;}
-	virtual const FObjectPropertyNode* AsObjectNode() const OVERRIDE { return this; }
-	virtual bool GetReadAddressUncached(FPropertyNode& InNode, bool InRequiresSingleSelection, FReadAddressListData& OutAddresses, bool bComparePropertyContents = true, bool bObjectForceCompare = false, bool bArrayPropertiesCanDifferInSize = false) const OVERRIDE;
-	virtual bool GetReadAddressUncached(FPropertyNode& InNode, FReadAddressListData& OutAddresses) const OVERRIDE;
+	virtual FObjectPropertyNode* AsObjectNode() override { return this;}
+	virtual const FObjectPropertyNode* AsObjectNode() const override { return this; }
+	virtual bool GetReadAddressUncached(FPropertyNode& InNode, bool InRequiresSingleSelection, FReadAddressListData& OutAddresses, bool bComparePropertyContents = true, bool bObjectForceCompare = false, bool bArrayPropertiesCanDifferInSize = false) const override;
+	virtual bool GetReadAddressUncached(FPropertyNode& InNode, FReadAddressListData& OutAddresses) const override;
 
-	virtual uint8* GetValueBaseAddress( uint8* Base ) OVERRIDE;
+	virtual uint8* GetValueBaseAddress( uint8* Base ) override;
 
 	/**
 	 * Returns the UObject at index "n" of the Objects Array
@@ -79,9 +79,9 @@ public:
 
 protected:
 	/** FPropertyNode interface */
-	virtual void InitBeforeNodeFlags() OVERRIDE;
-	virtual void InitChildNodes() OVERRIDE;
-	virtual void GetQualifiedName( FString& PathPlusIndex, const bool bWithArrayIndex ) const OVERRIDE;
+	virtual void InitBeforeNodeFlags() override;
+	virtual void InitChildNodes() override;
+	virtual void GetQualifiedName( FString& PathPlusIndex, const bool bWithArrayIndex ) const override;
 
 	/**
 	 * Looks at the Objects array and creates the best base class.  Called by
diff --git a/Engine/Source/Editor/PropertyEditor/Private/Presentation/PropertyTable/DataSource.h b/Engine/Source/Editor/PropertyEditor/Private/Presentation/PropertyTable/DataSource.h
--- a/Engine/Source/Editor/PropertyEditor/Private/Presentation/PropertyTable/DataSource.h
+++ b/Engine/Source/Editor/PropertyEditor/Private/Presentation/PropertyTable/DataSource.h
@@ -12,10 +12,10 @@ public:
 	{
 	}
 
-	virtual TWeakObjectPtr< UObject > AsUObject() const OVERRIDE { return Object; }
-	virtual TSharedPtr< FPropertyPath > AsPropertyPath() const OVERRIDE { return NULL; }
+	virtual TWeakObjectPtr< UObject > AsUObject() const override { return Object; }
+	virtual TSharedPtr< FPropertyPath > AsPropertyPath() const override { return NULL; }
 
-	virtual bool IsValid() const OVERRIDE { return Object.IsValid(); }
+	virtual bool IsValid() const override { return Object.IsValid(); }
 
 private:
 
@@ -31,10 +31,10 @@ public:
 	{
 	}
 
-	virtual TWeakObjectPtr< UObject > AsUObject() const OVERRIDE { return NULL; }
-	virtual TSharedPtr< FPropertyPath > AsPropertyPath() const OVERRIDE { return Path; }
+	virtual TWeakObjectPtr< UObject > AsUObject() const override { return NULL; }
+	virtual TSharedPtr< FPropertyPath > AsPropertyPath() const override { return Path; }
 
-	virtual bool IsValid() const OVERRIDE { return true; }
+	virtual bool IsValid() const override { return true; }
 
 private:
 
@@ -45,8 +45,8 @@ class NoDataSource : public IDataSource
 {
 public:
 
-	virtual TWeakObjectPtr< UObject > AsUObject() const OVERRIDE { return NULL; }
-	virtual TSharedPtr< FPropertyPath > AsPropertyPath() const OVERRIDE { return NULL; }
+	virtual TWeakObjectPtr< UObject > AsUObject() const override { return NULL; }
+	virtual TSharedPtr< FPropertyPath > AsPropertyPath() const override { return NULL; }
 
-	virtual bool IsValid() const OVERRIDE { return false; }
+	virtual bool IsValid() const override { return false; }
 };
\ No newline at end of file
diff --git a/Engine/Source/Editor/PropertyEditor/Private/Presentation/PropertyTable/PropertyTable.h b/Engine/Source/Editor/PropertyEditor/Private/Presentation/PropertyTable/PropertyTable.h
--- a/Engine/Source/Editor/PropertyEditor/Private/Presentation/PropertyTable/PropertyTable.h
+++ b/Engine/Source/Editor/PropertyEditor/Private/Presentation/PropertyTable/PropertyTable.h
@@ -9,124 +9,124 @@ public:
 
 	FPropertyTable();
 
-	virtual void Tick() OVERRIDE;
+	virtual void Tick() override;
 
-	virtual void RequestRefresh() OVERRIDE;
+	virtual void RequestRefresh() override;
 
-	virtual class FNotifyHook* GetNotifyHook() const OVERRIDE;
-	virtual bool AreFavoritesEnabled() const OVERRIDE;
-	virtual void ToggleFavorite( const TSharedRef< class FPropertyEditor >& PropertyEditor ) const OVERRIDE;
-	virtual void CreateColorPickerWindow( const TSharedRef< class FPropertyEditor >& PropertyEditor, bool bUseAlpha ) const OVERRIDE;
-	virtual void EnqueueDeferredAction( FSimpleDelegate DeferredAction ) OVERRIDE;
-	virtual TSharedPtr<class FAssetThumbnailPool> GetThumbnailPool() const OVERRIDE;
-	virtual void NotifyFinishedChangingProperties(const FPropertyChangedEvent& PropertyChangedEvent) OVERRIDE {}
+	virtual class FNotifyHook* GetNotifyHook() const override;
+	virtual bool AreFavoritesEnabled() const override;
+	virtual void ToggleFavorite( const TSharedRef< class FPropertyEditor >& PropertyEditor ) const override;
+	virtual void CreateColorPickerWindow( const TSharedRef< class FPropertyEditor >& PropertyEditor, bool bUseAlpha ) const override;
+	virtual void EnqueueDeferredAction( FSimpleDelegate DeferredAction ) override;
+	virtual TSharedPtr<class FAssetThumbnailPool> GetThumbnailPool() const override;
+	virtual void NotifyFinishedChangingProperties(const FPropertyChangedEvent& PropertyChangedEvent) override {}
 
-	virtual bool GetIsUserAllowedToChangeRoot() OVERRIDE;
-	virtual void SetIsUserAllowedToChangeRoot( bool InAllowUserToChangeRoot ) OVERRIDE;
+	virtual bool GetIsUserAllowedToChangeRoot() override;
+	virtual void SetIsUserAllowedToChangeRoot( bool InAllowUserToChangeRoot ) override;
 
-	virtual void AddColumn( const TWeakObjectPtr< UObject >& Object ) OVERRIDE;
-	virtual void AddColumn( const TWeakObjectPtr< UProperty >& Property ) OVERRIDE;
-	virtual void AddColumn( const TSharedRef< FPropertyPath >& PropertyPath ) OVERRIDE;
-	virtual void AddColumn( const TSharedRef< class IPropertyTableColumn >& Column ) OVERRIDE;
-	virtual void RemoveColumn( const TSharedRef< class IPropertyTableColumn >& Column ) OVERRIDE;
+	virtual void AddColumn( const TWeakObjectPtr< UObject >& Object ) override;
+	virtual void AddColumn( const TWeakObjectPtr< UProperty >& Property ) override;
+	virtual void AddColumn( const TSharedRef< FPropertyPath >& PropertyPath ) override;
+	virtual void AddColumn( const TSharedRef< class IPropertyTableColumn >& Column ) override;
+	virtual void RemoveColumn( const TSharedRef< class IPropertyTableColumn >& Column ) override;
 
-	virtual void AddRow( const TWeakObjectPtr< UObject >& Object ) OVERRIDE;
-	virtual void AddRow( const TWeakObjectPtr< UProperty >& Property ) OVERRIDE;
-	virtual void AddRow( const TSharedRef< FPropertyPath >& PropertyPath ) OVERRIDE;
-	virtual void AddRow( const TSharedRef< class IPropertyTableRow >& Row ) OVERRIDE;
-	virtual void RemoveRow( const TSharedRef< class IPropertyTableRow >& Row ) OVERRIDE;
+	virtual void AddRow( const TWeakObjectPtr< UObject >& Object ) override;
+	virtual void AddRow( const TWeakObjectPtr< UProperty >& Property ) override;
+	virtual void AddRow( const TSharedRef< FPropertyPath >& PropertyPath ) override;
+	virtual void AddRow( const TSharedRef< class IPropertyTableRow >& Row ) override;
+	virtual void RemoveRow( const TSharedRef< class IPropertyTableRow >& Row ) override;
 
-	virtual const EPropertyTableOrientation::Type GetOrientation() const OVERRIDE { return Orientation; }
-	virtual void SetOrientation( EPropertyTableOrientation::Type InOrientation ) OVERRIDE;
+	virtual const EPropertyTableOrientation::Type GetOrientation() const override { return Orientation; }
+	virtual void SetOrientation( EPropertyTableOrientation::Type InOrientation ) override;
 
-	virtual void SetRootPath( const TSharedPtr< FPropertyPath >& Path ) OVERRIDE;
-	virtual TSharedRef< FPropertyPath > GetRootPath() const OVERRIDE { return RootPath; }
-	virtual TArray< FPropertyInfo > GetPossibleExtensionsForPath( const TSharedRef< FPropertyPath >& Path ) const OVERRIDE;
+	virtual void SetRootPath( const TSharedPtr< FPropertyPath >& Path ) override;
+	virtual TSharedRef< FPropertyPath > GetRootPath() const override { return RootPath; }
+	virtual TArray< FPropertyInfo > GetPossibleExtensionsForPath( const TSharedRef< FPropertyPath >& Path ) const override;
 
-	virtual void GetSelectedObjects( TArray< TWeakObjectPtr< UObject > >& OutSelectedObjects) const OVERRIDE;
+	virtual void GetSelectedObjects( TArray< TWeakObjectPtr< UObject > >& OutSelectedObjects) const override;
 
-	virtual const TArray< TWeakObjectPtr< UObject > >& GetObjects() const OVERRIDE { return SourceObjects; }
-	virtual void SetObjects( const TArray< TWeakObjectPtr< UObject > >& Objects ) OVERRIDE;
-	virtual void SetObjects( const TArray< UObject* >& Objects ) OVERRIDE;
+	virtual const TArray< TWeakObjectPtr< UObject > >& GetObjects() const override { return SourceObjects; }
+	virtual void SetObjects( const TArray< TWeakObjectPtr< UObject > >& Objects ) override;
+	virtual void SetObjects( const TArray< UObject* >& Objects ) override;
 
-	virtual TSharedRef< class FObjectPropertyNode > GetObjectPropertyNode( const TSharedRef< IPropertyTableColumn >& Column, const TSharedRef< IPropertyTableRow >& Row ) OVERRIDE;
-	virtual TSharedRef< class FObjectPropertyNode > GetObjectPropertyNode( const TWeakObjectPtr< UObject >& Object ) OVERRIDE;
+	virtual TSharedRef< class FObjectPropertyNode > GetObjectPropertyNode( const TSharedRef< IPropertyTableColumn >& Column, const TSharedRef< IPropertyTableRow >& Row ) override;
+	virtual TSharedRef< class FObjectPropertyNode > GetObjectPropertyNode( const TWeakObjectPtr< UObject >& Object ) override;
 
-	virtual bool GetShowRowHeader() const OVERRIDE { return ShowRowHeader; }
-	virtual void SetShowRowHeader( const bool InShowRowHeader ) OVERRIDE;
+	virtual bool GetShowRowHeader() const override { return ShowRowHeader; }
+	virtual void SetShowRowHeader( const bool InShowRowHeader ) override;
 
-	virtual bool GetShowObjectName() const OVERRIDE { return ShowObjectName; }
-	virtual void SetShowObjectName( const bool ShowObjectName ) OVERRIDE;
+	virtual bool GetShowObjectName() const override { return ShowObjectName; }
+	virtual void SetShowObjectName( const bool ShowObjectName ) override;
 
-	virtual const TArray< TSharedRef< class IPropertyTableColumn > >& GetColumns() OVERRIDE { return Columns; }
-	virtual TArray< TSharedRef< class IPropertyTableRow > >& GetRows() OVERRIDE { return Rows; }
+	virtual const TArray< TSharedRef< class IPropertyTableColumn > >& GetColumns() override { return Columns; }
+	virtual TArray< TSharedRef< class IPropertyTableRow > >& GetRows() override { return Rows; }
 
-	virtual const TSet< TSharedRef< class IPropertyTableRow > >& GetSelectedRows() OVERRIDE { return SelectedRows; }
-	virtual void SetSelectedRows( const TSet< TSharedRef< IPropertyTableRow > >& InSelectedRows ) OVERRIDE;
+	virtual const TSet< TSharedRef< class IPropertyTableRow > >& GetSelectedRows() override { return SelectedRows; }
+	virtual void SetSelectedRows( const TSet< TSharedRef< IPropertyTableRow > >& InSelectedRows ) override;
 
-	virtual const TSet< TSharedRef< class IPropertyTableCell > >& GetSelectedCells() OVERRIDE { return SelectedCells; }
-	virtual void SetSelectedCells( const TSet< TSharedRef< class IPropertyTableCell > >& InSelectedCells ) OVERRIDE;
+	virtual const TSet< TSharedRef< class IPropertyTableCell > >& GetSelectedCells() override { return SelectedCells; }
+	virtual void SetSelectedCells( const TSet< TSharedRef< class IPropertyTableCell > >& InSelectedCells ) override;
 
-	virtual void SelectCellRange( const TSharedRef< class IPropertyTableCell >& StartingCell, const TSharedRef< class IPropertyTableCell >& EndingCell ) OVERRIDE;
+	virtual void SelectCellRange( const TSharedRef< class IPropertyTableCell >& StartingCell, const TSharedRef< class IPropertyTableCell >& EndingCell ) override;
 
-	virtual float GetItemHeight() const OVERRIDE { return ItemHeight; }
-	virtual void SetItemHeight( float NewItemHeight ) OVERRIDE { ItemHeight = NewItemHeight; }
+	virtual float GetItemHeight() const override { return ItemHeight; }
+	virtual void SetItemHeight( float NewItemHeight ) override { ItemHeight = NewItemHeight; }
 
-	virtual TSharedPtr< class IPropertyTableCell > GetLastClickedCell() const OVERRIDE { return LastClickedCell; }
-	virtual void SetLastClickedCell( const TSharedPtr< class IPropertyTableCell >& Cell ) OVERRIDE { LastClickedCell = Cell; }
+	virtual TSharedPtr< class IPropertyTableCell > GetLastClickedCell() const override { return LastClickedCell; }
+	virtual void SetLastClickedCell( const TSharedPtr< class IPropertyTableCell >& Cell ) override { LastClickedCell = Cell; }
 
-	virtual TSharedPtr< class IPropertyTableCell > GetCurrentCell() const OVERRIDE { return CurrentCell; }
-	virtual void SetCurrentCell( const TSharedPtr< class IPropertyTableCell >& Cell ) OVERRIDE;
+	virtual TSharedPtr< class IPropertyTableCell > GetCurrentCell() const override { return CurrentCell; }
+	virtual void SetCurrentCell( const TSharedPtr< class IPropertyTableCell >& Cell ) override;
 
-	virtual TSharedPtr< class IPropertyTableColumn > GetCurrentColumn() const OVERRIDE { return CurrentColumn; }
-	virtual void SetCurrentColumn( const TSharedPtr< class IPropertyTableColumn >& Column ) OVERRIDE;
+	virtual TSharedPtr< class IPropertyTableColumn > GetCurrentColumn() const override { return CurrentColumn; }
+	virtual void SetCurrentColumn( const TSharedPtr< class IPropertyTableColumn >& Column ) override;
 
-	virtual TSharedPtr< class IPropertyTableRow > GetCurrentRow() const OVERRIDE { return CurrentRow; }
-	virtual void SetCurrentRow( const TSharedPtr< class IPropertyTableRow >& Row ) OVERRIDE;
+	virtual TSharedPtr< class IPropertyTableRow > GetCurrentRow() const override { return CurrentRow; }
+	virtual void SetCurrentRow( const TSharedPtr< class IPropertyTableRow >& Row ) override;
 
-	virtual TSharedPtr< class IPropertyTableCell > GetFirstCellInSelection() OVERRIDE;
-	virtual TSharedPtr< class IPropertyTableCell > GetLastCellInSelection() OVERRIDE;
+	virtual TSharedPtr< class IPropertyTableCell > GetFirstCellInSelection() override;
+	virtual TSharedPtr< class IPropertyTableCell > GetLastCellInSelection() override;
 
-	virtual TSharedPtr< class IPropertyTableCell > GetNextCellInRow( const TSharedRef< class IPropertyTableCell >& Cell ) OVERRIDE;
-	virtual TSharedPtr< class IPropertyTableCell > GetPreviousCellInRow( const TSharedRef< class IPropertyTableCell >& Cell ) OVERRIDE;
+	virtual TSharedPtr< class IPropertyTableCell > GetNextCellInRow( const TSharedRef< class IPropertyTableCell >& Cell ) override;
+	virtual TSharedPtr< class IPropertyTableCell > GetPreviousCellInRow( const TSharedRef< class IPropertyTableCell >& Cell ) override;
 
-	virtual TSharedPtr< class IPropertyTableCell > GetNextCellInColumn( const TSharedRef< class IPropertyTableCell >& Cell ) OVERRIDE;
-	virtual TSharedPtr< class IPropertyTableCell > GetPreviousCellInColumn( const TSharedRef< class IPropertyTableCell >& Cell ) OVERRIDE;
+	virtual TSharedPtr< class IPropertyTableCell > GetNextCellInColumn( const TSharedRef< class IPropertyTableCell >& Cell ) override;
+	virtual TSharedPtr< class IPropertyTableCell > GetPreviousCellInColumn( const TSharedRef< class IPropertyTableCell >& Cell ) override;
 
-	virtual TSharedPtr< class IPropertyTableCell > GetFirstCellInRow( const TSharedRef< class IPropertyTableRow >& Row ) OVERRIDE;
-	virtual TSharedPtr< class IPropertyTableCell > GetLastCellInRow( const TSharedRef< class IPropertyTableRow >& Row ) OVERRIDE;
+	virtual TSharedPtr< class IPropertyTableCell > GetFirstCellInRow( const TSharedRef< class IPropertyTableRow >& Row ) override;
+	virtual TSharedPtr< class IPropertyTableCell > GetLastCellInRow( const TSharedRef< class IPropertyTableRow >& Row ) override;
 
-	virtual TSharedPtr< class IPropertyTableCell > GetFirstCellInColumn( const TSharedRef< class IPropertyTableColumn >& Column ) OVERRIDE;
-	virtual TSharedPtr< class IPropertyTableCell > GetLastCellInColumn( const TSharedRef< class IPropertyTableColumn >& Column ) OVERRIDE;
+	virtual TSharedPtr< class IPropertyTableCell > GetFirstCellInColumn( const TSharedRef< class IPropertyTableColumn >& Column ) override;
+	virtual TSharedPtr< class IPropertyTableCell > GetLastCellInColumn( const TSharedRef< class IPropertyTableColumn >& Column ) override;
 
-	virtual TSharedPtr< class IPropertyTableCell > GetFirstCellInTable() OVERRIDE;
-	virtual TSharedPtr< class IPropertyTableCell > GetLastCellInTable() OVERRIDE;
+	virtual TSharedPtr< class IPropertyTableCell > GetFirstCellInTable() override;
+	virtual TSharedPtr< class IPropertyTableCell > GetLastCellInTable() override;
 
-	virtual EPropertyTableSelectionUnit::Type GetSelectionUnit() const OVERRIDE { return SelectionUnit; }
-	virtual void SetSelectionUnit( const EPropertyTableSelectionUnit::Type Unit ) OVERRIDE { SelectionUnit = Unit; }
+	virtual EPropertyTableSelectionUnit::Type GetSelectionUnit() const override { return SelectionUnit; }
+	virtual void SetSelectionUnit( const EPropertyTableSelectionUnit::Type Unit ) override { SelectionUnit = Unit; }
 
-	virtual ESelectionMode::Type GetSelectionMode() const OVERRIDE { return SelectionMode; }
-	virtual void SetSelectionMode( const ESelectionMode::Type Mode ) OVERRIDE;
+	virtual ESelectionMode::Type GetSelectionMode() const override { return SelectionMode; }
+	virtual void SetSelectionMode( const ESelectionMode::Type Mode ) override;
 
-	virtual EColumnSortMode::Type GetColumnSortMode( const TSharedRef< class IPropertyTableColumn > Column ) const OVERRIDE;
-	virtual void SortByColumnWithId( const FName& ColumnId, EColumnSortMode::Type SortMode ) OVERRIDE;
-	virtual void SortByColumn( const TSharedRef< class IPropertyTableColumn >& Column, EColumnSortMode::Type SortMode ) OVERRIDE;
+	virtual EColumnSortMode::Type GetColumnSortMode( const TSharedRef< class IPropertyTableColumn > Column ) const override;
+	virtual void SortByColumnWithId( const FName& ColumnId, EColumnSortMode::Type SortMode ) override;
+	virtual void SortByColumn( const TSharedRef< class IPropertyTableColumn >& Column, EColumnSortMode::Type SortMode ) override;
 
-	virtual void PasteTextAtCell( const FString& Text, const TSharedRef< class IPropertyTableCell >& Cell ) OVERRIDE;
+	virtual void PasteTextAtCell( const FString& Text, const TSharedRef< class IPropertyTableCell >& Cell ) override;
 
 	DECLARE_DERIVED_EVENT( FPropertyTable, IPropertyTable::FSelectionChanged, FSelectionChanged );
-	FSelectionChanged* OnSelectionChanged() OVERRIDE { return &SelectionChanged; }
+	FSelectionChanged* OnSelectionChanged() override { return &SelectionChanged; }
 
 	DECLARE_DERIVED_EVENT( FPropertyTable, IPropertyTable::FColumnsChanged, FColumnsChanged );
-	FColumnsChanged* OnColumnsChanged() OVERRIDE { return &ColumnsChanged; }
+	FColumnsChanged* OnColumnsChanged() override { return &ColumnsChanged; }
 
 	DECLARE_DERIVED_EVENT( FPropertyTable, IPropertyTable::FRowsChanged, FRowsChanged );
-	FRowsChanged* OnRowsChanged() OVERRIDE { return &RowsChanged; }
+	FRowsChanged* OnRowsChanged() override { return &RowsChanged; }
 
 	DECLARE_DERIVED_EVENT( FPropertyTable, IPropertyTable::FRootPathChanged, FRootPathChanged );
-	virtual FRootPathChanged* OnRootPathChanged() OVERRIDE { return &RootPathChanged; }
+	virtual FRootPathChanged* OnRootPathChanged() override { return &RootPathChanged; }
 
-	virtual bool IsPropertyEditingEnabled() const OVERRIDE { return true; }
+	virtual bool IsPropertyEditingEnabled() const override { return true; }
 private:
 
 	void UpdateColumns();
diff --git a/Engine/Source/Editor/PropertyEditor/Private/Presentation/PropertyTable/PropertyTableCell.h b/Engine/Source/Editor/PropertyEditor/Private/Presentation/PropertyTable/PropertyTableCell.h
--- a/Engine/Source/Editor/PropertyEditor/Private/Presentation/PropertyTable/PropertyTableCell.h
+++ b/Engine/Source/Editor/PropertyEditor/Private/Presentation/PropertyTable/PropertyTableCell.h
@@ -11,32 +11,32 @@ public:
 
 	FPropertyTableCell( const TSharedRef< class IPropertyTableColumn >& InColumn, const TSharedRef< class IPropertyTableRow >& InRow );
 
-	virtual void Refresh() OVERRIDE;
+	virtual void Refresh() override;
 
-	virtual bool IsReadOnly() const OVERRIDE;
-	virtual bool IsBound() const OVERRIDE { return bIsBound; }
-	virtual bool InEditMode() const OVERRIDE { return bInEditMode; }
-	virtual bool IsValid() const OVERRIDE;
+	virtual bool IsReadOnly() const override;
+	virtual bool IsBound() const override { return bIsBound; }
+	virtual bool InEditMode() const override { return bInEditMode; }
+	virtual bool IsValid() const override;
 
-	virtual FString GetValueAsString() const OVERRIDE;
-	virtual FText GetValueAsText() const OVERRIDE;
-	virtual void SetValueFromString( const FString& InString ) OVERRIDE;
+	virtual FString GetValueAsString() const override;
+	virtual FText GetValueAsText() const override;
+	virtual void SetValueFromString( const FString& InString ) override;
 
-	virtual TWeakObjectPtr< UObject > GetObject() const OVERRIDE;
-	virtual TSharedPtr< class FPropertyNode > GetNode() const OVERRIDE { return PropertyNode; }
-	virtual TSharedRef< class IPropertyTableColumn > GetColumn() const OVERRIDE { return Column.Pin().ToSharedRef(); }
-	virtual TSharedRef< class IPropertyTableRow > GetRow() const OVERRIDE { return Row.Pin().ToSharedRef(); }
-	virtual TSharedRef< class IPropertyTable > GetTable() const OVERRIDE;
-	virtual TSharedPtr< class IPropertyHandle > GetPropertyHandle() const OVERRIDE;
+	virtual TWeakObjectPtr< UObject > GetObject() const override;
+	virtual TSharedPtr< class FPropertyNode > GetNode() const override { return PropertyNode; }
+	virtual TSharedRef< class IPropertyTableColumn > GetColumn() const override { return Column.Pin().ToSharedRef(); }
+	virtual TSharedRef< class IPropertyTableRow > GetRow() const override { return Row.Pin().ToSharedRef(); }
+	virtual TSharedRef< class IPropertyTable > GetTable() const override;
+	virtual TSharedPtr< class IPropertyHandle > GetPropertyHandle() const override;
 
-	virtual void EnterEditMode() OVERRIDE;
-	virtual void ExitEditMode() OVERRIDE;
+	virtual void EnterEditMode() override;
+	virtual void ExitEditMode() override;
 
 	DECLARE_DERIVED_EVENT( FPropertyTableCell, IPropertyTableCell::FEnteredEditModeEvent, FEnteredEditModeEvent );
-	virtual FEnteredEditModeEvent& OnEnteredEditMode() OVERRIDE { return EnteredEditModeEvent; }
+	virtual FEnteredEditModeEvent& OnEnteredEditMode() override { return EnteredEditModeEvent; }
 
 	DECLARE_DERIVED_EVENT( FPropertyTableCell, IPropertyTableCell::FExitedEditModeEvent, FExitedEditModeEvent  );
-	virtual FExitedEditModeEvent& OnExitedEditMode() OVERRIDE { return ExitedEditModeEvent; }
+	virtual FExitedEditModeEvent& OnExitedEditMode() override { return ExitedEditModeEvent; }
 
 
 private:
diff --git a/Engine/Source/Editor/PropertyEditor/Private/Presentation/PropertyTable/PropertyTableColumn.h b/Engine/Source/Editor/PropertyEditor/Private/Presentation/PropertyTable/PropertyTableColumn.h
--- a/Engine/Source/Editor/PropertyEditor/Private/Presentation/PropertyTable/PropertyTableColumn.h
+++ b/Engine/Source/Editor/PropertyEditor/Private/Presentation/PropertyTable/PropertyTableColumn.h
@@ -15,46 +15,46 @@ public:
 
 	// Begin IPropertyTable Interface
 
-	virtual FName GetId() const OVERRIDE;
+	virtual FName GetId() const override;
 
-	virtual FText GetDisplayName() const OVERRIDE;
+	virtual FText GetDisplayName() const override;
 
-	virtual TSharedRef< IDataSource > GetDataSource() const OVERRIDE { return DataSource; }
+	virtual TSharedRef< IDataSource > GetDataSource() const override { return DataSource; }
 
-	virtual TSharedRef< class FPropertyPath > GetPartialPath() const OVERRIDE { return PartialPath; }
+	virtual TSharedRef< class FPropertyPath > GetPartialPath() const override { return PartialPath; }
 
-	virtual TSharedRef< class IPropertyTableCell > GetCell( const TSharedRef< class IPropertyTableRow >& Row ) OVERRIDE;
+	virtual TSharedRef< class IPropertyTableCell > GetCell( const TSharedRef< class IPropertyTableRow >& Row ) override;
 
-	virtual void RemoveCellsForRow( const TSharedRef< class IPropertyTableRow >& Row ) OVERRIDE;
+	virtual void RemoveCellsForRow( const TSharedRef< class IPropertyTableRow >& Row ) override;
 
-	virtual TSharedRef< class IPropertyTable > GetTable() const OVERRIDE;
+	virtual TSharedRef< class IPropertyTable > GetTable() const override;
 
-	virtual bool CanSelectCells() const OVERRIDE { return !IsHidden(); }
+	virtual bool CanSelectCells() const override { return !IsHidden(); }
 
-	virtual EPropertyTableColumnSizeMode::Type GetSizeMode() const OVERRIDE { return SizeMode; }
+	virtual EPropertyTableColumnSizeMode::Type GetSizeMode() const override { return SizeMode; }
 
-	virtual void SetSizeMode(EPropertyTableColumnSizeMode::Type InSizeMode) OVERRIDE{ SizeMode = InSizeMode; }
+	virtual void SetSizeMode(EPropertyTableColumnSizeMode::Type InSizeMode) override{ SizeMode = InSizeMode; }
 
-	virtual float GetWidth() const OVERRIDE { return Width; } 
+	virtual float GetWidth() const override { return Width; } 
 
-	virtual void SetWidth( float InWidth ) OVERRIDE { Width = InWidth; }
+	virtual void SetWidth( float InWidth ) override { Width = InWidth; }
 
-	virtual bool IsHidden() const OVERRIDE { return bIsHidden; }
+	virtual bool IsHidden() const override { return bIsHidden; }
 
-	virtual void SetHidden( bool InIsHidden ) OVERRIDE { bIsHidden = InIsHidden; }
+	virtual void SetHidden( bool InIsHidden ) override { bIsHidden = InIsHidden; }
 
-	virtual bool IsFrozen() const OVERRIDE { return bIsFrozen; }
+	virtual bool IsFrozen() const override { return bIsFrozen; }
 
-	virtual void SetFrozen( bool InIsFrozen ) OVERRIDE;
+	virtual void SetFrozen( bool InIsFrozen ) override;
 
-	virtual bool CanSortBy() const OVERRIDE;
+	virtual bool CanSortBy() const override;
 
-	virtual void Sort( TArray< TSharedRef< class IPropertyTableRow > >& Rows, const EColumnSortMode::Type SortMode ) OVERRIDE;
+	virtual void Sort( TArray< TSharedRef< class IPropertyTableRow > >& Rows, const EColumnSortMode::Type SortMode ) override;
 
-	virtual void Tick() OVERRIDE;
+	virtual void Tick() override;
 
 	DECLARE_DERIVED_EVENT( FPropertyTableColumn, IPropertyTableColumn::FFrozenStateChanged, FFrozenStateChanged );
-	FFrozenStateChanged* OnFrozenStateChanged() OVERRIDE { return &FrozenStateChanged; }
+	FFrozenStateChanged* OnFrozenStateChanged() override { return &FrozenStateChanged; }
 
 	// End IPropertyTable Interface
 
diff --git a/Engine/Source/Editor/PropertyEditor/Private/Presentation/PropertyTable/PropertyTableObjectNameCell.h b/Engine/Source/Editor/PropertyEditor/Private/Presentation/PropertyTable/PropertyTableObjectNameCell.h
--- a/Engine/Source/Editor/PropertyEditor/Private/Presentation/PropertyTable/PropertyTableObjectNameCell.h
+++ b/Engine/Source/Editor/PropertyEditor/Private/Presentation/PropertyTable/PropertyTableObjectNameCell.h
@@ -11,33 +11,33 @@ public:
 
 	FPropertyTableObjectNameCell( const TSharedRef< class FPropertyTableObjectNameColumn >& InColumn, const TSharedRef< class IPropertyTableRow >& InRow );
 
-	virtual void Refresh() OVERRIDE;
+	virtual void Refresh() override;
 
-	virtual bool IsReadOnly() const OVERRIDE;
-	virtual bool IsBound() const OVERRIDE { return bIsBound; }
-	virtual bool InEditMode() const OVERRIDE { return bInEditMode; }
-	virtual bool IsValid() const OVERRIDE;
+	virtual bool IsReadOnly() const override;
+	virtual bool IsBound() const override { return bIsBound; }
+	virtual bool InEditMode() const override { return bInEditMode; }
+	virtual bool IsValid() const override;
 
-	virtual FString GetValueAsString() const OVERRIDE;
-	virtual FText GetValueAsText() const OVERRIDE;
-	virtual void SetValueFromString( const FString& InString ) OVERRIDE {}
+	virtual FString GetValueAsString() const override;
+	virtual FText GetValueAsText() const override;
+	virtual void SetValueFromString( const FString& InString ) override {}
 
-	virtual TWeakObjectPtr< UObject > GetObject() const OVERRIDE;
+	virtual TWeakObjectPtr< UObject > GetObject() const override;
 
-	virtual TSharedPtr< class FPropertyNode > GetNode() const OVERRIDE { return NULL; }
-	virtual TSharedRef< class IPropertyTableColumn > GetColumn() const OVERRIDE { return Column.Pin().ToSharedRef(); }
-	virtual TSharedRef< class IPropertyTableRow > GetRow() const OVERRIDE { return Row.Pin().ToSharedRef(); }
-	virtual TSharedRef< class IPropertyTable > GetTable() const OVERRIDE;
-	virtual TSharedPtr< class IPropertyHandle > GetPropertyHandle() const OVERRIDE { return NULL; }
+	virtual TSharedPtr< class FPropertyNode > GetNode() const override { return NULL; }
+	virtual TSharedRef< class IPropertyTableColumn > GetColumn() const override { return Column.Pin().ToSharedRef(); }
+	virtual TSharedRef< class IPropertyTableRow > GetRow() const override { return Row.Pin().ToSharedRef(); }
+	virtual TSharedRef< class IPropertyTable > GetTable() const override;
+	virtual TSharedPtr< class IPropertyHandle > GetPropertyHandle() const override { return NULL; }
 
-	virtual void EnterEditMode() OVERRIDE;
-	virtual void ExitEditMode() OVERRIDE;
+	virtual void EnterEditMode() override;
+	virtual void ExitEditMode() override;
 
 	DECLARE_DERIVED_EVENT( FPropertyTableObjectNameCell, IPropertyTableCell::FEnteredEditModeEvent, FEnteredEditModeEvent );
-	virtual FEnteredEditModeEvent& OnEnteredEditMode() OVERRIDE { return EnteredEditModeEvent; }
+	virtual FEnteredEditModeEvent& OnEnteredEditMode() override { return EnteredEditModeEvent; }
 
 	DECLARE_DERIVED_EVENT( FPropertyTableObjectNameCell, IPropertyTableCell::FExitedEditModeEvent, FExitedEditModeEvent  );
-	virtual FExitedEditModeEvent& OnExitedEditMode() OVERRIDE { return ExitedEditModeEvent; }
+	virtual FExitedEditModeEvent& OnExitedEditMode() override { return ExitedEditModeEvent; }
 
 
 private:
diff --git a/Engine/Source/Editor/PropertyEditor/Private/Presentation/PropertyTable/PropertyTableObjectNameColumn.h b/Engine/Source/Editor/PropertyEditor/Private/Presentation/PropertyTable/PropertyTableObjectNameColumn.h
--- a/Engine/Source/Editor/PropertyEditor/Private/Presentation/PropertyTable/PropertyTableObjectNameColumn.h
+++ b/Engine/Source/Editor/PropertyEditor/Private/Presentation/PropertyTable/PropertyTableObjectNameColumn.h
@@ -27,15 +27,15 @@ public:
 
 	//  Begin IPropertyTable Interface
 
-	virtual FName GetId() const OVERRIDE { return FName( TEXT("ObjectName") ); }
+	virtual FName GetId() const override { return FName( TEXT("ObjectName") ); }
 
-	virtual FText GetDisplayName() const OVERRIDE { return LOCTEXT( "DisplayName", "Name" ); }
+	virtual FText GetDisplayName() const override { return LOCTEXT( "DisplayName", "Name" ); }
 
-	virtual TSharedRef< IDataSource > GetDataSource() const OVERRIDE { return MakeShareable( new NoDataSource() ); }
+	virtual TSharedRef< IDataSource > GetDataSource() const override { return MakeShareable( new NoDataSource() ); }
 
-	virtual TSharedRef< class FPropertyPath > GetPartialPath() const OVERRIDE { return FPropertyPath::CreateEmpty(); }
+	virtual TSharedRef< class FPropertyPath > GetPartialPath() const override { return FPropertyPath::CreateEmpty(); }
 
-	virtual TSharedRef< class IPropertyTableCell > GetCell( const TSharedRef< class IPropertyTableRow >& Row ) OVERRIDE
+	virtual TSharedRef< class IPropertyTableCell > GetCell( const TSharedRef< class IPropertyTableRow >& Row ) override
 	{
 		//@todo Clean Cells cache when rows get updated [11/27/2012 Justin.Sargent]
 		TSharedRef< IPropertyTableCell >* CellPtr = Cells.Find( Row );
@@ -51,34 +51,34 @@ public:
 		return Cell;
 	}
 
-	virtual void RemoveCellsForRow( const TSharedRef< class IPropertyTableRow >& Row ) OVERRIDE
+	virtual void RemoveCellsForRow( const TSharedRef< class IPropertyTableRow >& Row ) override
 	{
 		Cells.Remove( Row );
 	}
 
-	virtual TSharedRef< class IPropertyTable > GetTable() const OVERRIDE { return Table.Pin().ToSharedRef(); }
+	virtual TSharedRef< class IPropertyTable > GetTable() const override { return Table.Pin().ToSharedRef(); }
 
-	virtual bool CanSelectCells() const OVERRIDE { return true; }
+	virtual bool CanSelectCells() const override { return true; }
 
-	virtual EPropertyTableColumnSizeMode::Type GetSizeMode() const OVERRIDE { return EPropertyTableColumnSizeMode::Fill; }
+	virtual EPropertyTableColumnSizeMode::Type GetSizeMode() const override { return EPropertyTableColumnSizeMode::Fill; }
 
-	virtual void SetSizeMode(EPropertyTableColumnSizeMode::Type InSizeMode) OVERRIDE {}
+	virtual void SetSizeMode(EPropertyTableColumnSizeMode::Type InSizeMode) override {}
 
-	virtual float GetWidth() const OVERRIDE { return Width; } 
+	virtual float GetWidth() const override { return Width; } 
 
-	virtual void SetWidth( float InWidth ) OVERRIDE { Width = InWidth; }
+	virtual void SetWidth( float InWidth ) override { Width = InWidth; }
 
-	virtual bool IsHidden() const OVERRIDE { return bIsHidden; }
+	virtual bool IsHidden() const override { return bIsHidden; }
 
-	virtual void SetHidden( bool InIsHidden ) OVERRIDE { bIsHidden = InIsHidden; }
+	virtual void SetHidden( bool InIsHidden ) override { bIsHidden = InIsHidden; }
 
-	virtual bool IsFrozen() const OVERRIDE { return true; }
+	virtual bool IsFrozen() const override { return true; }
 
-	virtual void SetFrozen( bool InIsFrozen ) OVERRIDE {}
+	virtual void SetFrozen( bool InIsFrozen ) override {}
 
-	virtual bool CanSortBy() const OVERRIDE { return true; }
+	virtual bool CanSortBy() const override { return true; }
 
-	virtual void Sort( TArray< TSharedRef< class IPropertyTableRow > >& Rows, const EColumnSortMode::Type SortMode ) OVERRIDE
+	virtual void Sort( TArray< TSharedRef< class IPropertyTableRow > >& Rows, const EColumnSortMode::Type SortMode ) override
 	{
 		struct FCompareRowByObjectNameAscending
 		{
@@ -127,10 +127,10 @@ public:
 		}
 	}
 
-	virtual void Tick() OVERRIDE {}
+	virtual void Tick() override {}
 
 	DECLARE_DERIVED_EVENT( FPropertyTableColumn, IPropertyTableColumn::FFrozenStateChanged, FFrozenStateChanged );
-	FFrozenStateChanged* OnFrozenStateChanged() OVERRIDE { return &FrozenStateChanged; }
+	FFrozenStateChanged* OnFrozenStateChanged() override { return &FrozenStateChanged; }
 
 	// End IPropertyTable Interface
 
diff --git a/Engine/Source/Editor/PropertyEditor/Private/Presentation/PropertyTable/PropertyTablePropertyNameCell.h b/Engine/Source/Editor/PropertyEditor/Private/Presentation/PropertyTable/PropertyTablePropertyNameCell.h
--- a/Engine/Source/Editor/PropertyEditor/Private/Presentation/PropertyTable/PropertyTablePropertyNameCell.h
+++ b/Engine/Source/Editor/PropertyEditor/Private/Presentation/PropertyTable/PropertyTablePropertyNameCell.h
@@ -13,41 +13,41 @@ public:
 
 	// Begin IPropertyTableCell Interface
 
-	virtual void EnterEditMode() OVERRIDE;
+	virtual void EnterEditMode() override;
 
-	virtual void ExitEditMode() OVERRIDE;
+	virtual void ExitEditMode() override;
 
-	virtual TSharedRef< class IPropertyTableColumn > GetColumn() const OVERRIDE { return Column.Pin().ToSharedRef(); }
+	virtual TSharedRef< class IPropertyTableColumn > GetColumn() const override { return Column.Pin().ToSharedRef(); }
 
-	virtual TSharedPtr< class FPropertyNode > GetNode() const OVERRIDE { return NULL; }
+	virtual TSharedPtr< class FPropertyNode > GetNode() const override { return NULL; }
 
-	virtual TWeakObjectPtr< UObject > GetObject() const OVERRIDE;
+	virtual TWeakObjectPtr< UObject > GetObject() const override;
 
-	virtual TSharedRef< class IPropertyTableRow > GetRow() const OVERRIDE { return Row.Pin().ToSharedRef(); }
+	virtual TSharedRef< class IPropertyTableRow > GetRow() const override { return Row.Pin().ToSharedRef(); }
 
-	virtual TSharedRef< class IPropertyTable > GetTable() const OVERRIDE;
+	virtual TSharedRef< class IPropertyTable > GetTable() const override;
 
-	virtual FString GetValueAsString() const OVERRIDE;
+	virtual FString GetValueAsString() const override;
 
-	virtual FText GetValueAsText() const OVERRIDE;
+	virtual FText GetValueAsText() const override;
 
-	virtual bool InEditMode() const OVERRIDE { return bInEditMode; }
+	virtual bool InEditMode() const override { return bInEditMode; }
 
-	virtual bool IsReadOnly() const OVERRIDE { return true; }
+	virtual bool IsReadOnly() const override { return true; }
 
-	virtual bool IsBound() const OVERRIDE { return bIsBound; }
+	virtual bool IsBound() const override { return bIsBound; }
 
-	virtual bool IsValid() const OVERRIDE  { return true; }
+	virtual bool IsValid() const override  { return true; }
 
 	DECLARE_DERIVED_EVENT( FPropertyTablePropertyNameCell, IPropertyTableCell::FEnteredEditModeEvent, FEnteredEditModeEvent );
-	virtual FEnteredEditModeEvent& OnEnteredEditMode() OVERRIDE { return EnteredEditModeEvent; }
+	virtual FEnteredEditModeEvent& OnEnteredEditMode() override { return EnteredEditModeEvent; }
 
 	DECLARE_DERIVED_EVENT( FPropertyTablePropertyNameCell, IPropertyTableCell::FExitedEditModeEvent, FExitedEditModeEvent  );
-	virtual FExitedEditModeEvent& OnExitedEditMode() OVERRIDE { return ExitedEditModeEvent; }
+	virtual FExitedEditModeEvent& OnExitedEditMode() override { return ExitedEditModeEvent; }
 
-	virtual void Refresh() OVERRIDE;
+	virtual void Refresh() override;
 
-	virtual void SetValueFromString( const FString& InString ) OVERRIDE {}
+	virtual void SetValueFromString( const FString& InString ) override {}
 
 	// End IPropertyTableCell Interface
 
