UE4 Refactoring. Changed OVERRIDE and FINAL macros to keywords override and final.

diff --git a/Engine/Source/Editor/PropertyEditor/Private/Presentation/PropertyTable/PropertyTablePropertyNameColumn.h b/Engine/Source/Editor/PropertyEditor/Private/Presentation/PropertyTable/PropertyTablePropertyNameColumn.h
--- a/Engine/Source/Editor/PropertyEditor/Private/Presentation/PropertyTable/PropertyTablePropertyNameColumn.h
+++ b/Engine/Source/Editor/PropertyEditor/Private/Presentation/PropertyTable/PropertyTablePropertyNameColumn.h
@@ -20,49 +20,49 @@ public:
 
 	// Begin IPropertyTableColumn Interface
 
-	virtual bool CanSelectCells() const OVERRIDE { return true; }
+	virtual bool CanSelectCells() const override { return true; }
 
-	virtual bool CanSortBy() const OVERRIDE { return true; }
+	virtual bool CanSortBy() const override { return true; }
 
-	virtual TSharedRef< class IPropertyTableCell > GetCell( const TSharedRef< class IPropertyTableRow >& Row ) OVERRIDE;
+	virtual TSharedRef< class IPropertyTableCell > GetCell( const TSharedRef< class IPropertyTableRow >& Row ) override;
 
-	virtual TSharedRef< IDataSource > GetDataSource() const OVERRIDE { return DataSource; }
+	virtual TSharedRef< IDataSource > GetDataSource() const override { return DataSource; }
 
-	virtual TSharedRef< class FPropertyPath > GetPartialPath() const OVERRIDE { return FPropertyPath::CreateEmpty(); }
+	virtual TSharedRef< class FPropertyPath > GetPartialPath() const override { return FPropertyPath::CreateEmpty(); }
 
-	virtual FText GetDisplayName() const OVERRIDE { return LOCTEXT( "DisplayName", "Name" ); }
+	virtual FText GetDisplayName() const override { return LOCTEXT( "DisplayName", "Name" ); }
 
-	virtual FName GetId() const OVERRIDE { return FName( TEXT("PropertyName") ); }
+	virtual FName GetId() const override { return FName( TEXT("PropertyName") ); }
 
-	virtual EPropertyTableColumnSizeMode::Type GetSizeMode() const OVERRIDE { return EPropertyTableColumnSizeMode::Fill; }
+	virtual EPropertyTableColumnSizeMode::Type GetSizeMode() const override { return EPropertyTableColumnSizeMode::Fill; }
 
-	virtual void SetSizeMode(EPropertyTableColumnSizeMode::Type InSizeMode) OVERRIDE {}
+	virtual void SetSizeMode(EPropertyTableColumnSizeMode::Type InSizeMode) override {}
 
-	virtual TSharedRef< class IPropertyTable > GetTable() const OVERRIDE { return Table.Pin().ToSharedRef(); }
+	virtual TSharedRef< class IPropertyTable > GetTable() const override { return Table.Pin().ToSharedRef(); }
 
-	virtual float GetWidth() const OVERRIDE { return Width; }
+	virtual float GetWidth() const override { return Width; }
 
-	virtual bool IsFrozen() const OVERRIDE { return false; }
+	virtual bool IsFrozen() const override { return false; }
 
-	virtual bool IsHidden() const OVERRIDE { return bIsHidden; }
+	virtual bool IsHidden() const override { return bIsHidden; }
 
-	virtual void RemoveCellsForRow( const TSharedRef< class IPropertyTableRow >& Row ) OVERRIDE
+	virtual void RemoveCellsForRow( const TSharedRef< class IPropertyTableRow >& Row ) override
 	{
 		Cells.Remove( Row );
 	}
 
-	virtual void SetFrozen( bool InIsFrozen ) OVERRIDE {}
+	virtual void SetFrozen( bool InIsFrozen ) override {}
 
-	virtual void SetHidden( bool InIsHidden ) OVERRIDE { bIsHidden = InIsHidden; }
+	virtual void SetHidden( bool InIsHidden ) override { bIsHidden = InIsHidden; }
 
-	virtual void SetWidth( float InWidth ) OVERRIDE { Width = InWidth; }
+	virtual void SetWidth( float InWidth ) override { Width = InWidth; }
 
-	virtual void Sort( TArray< TSharedRef< class IPropertyTableRow > >& Rows, const EColumnSortMode::Type SortMode ) OVERRIDE;
+	virtual void Sort( TArray< TSharedRef< class IPropertyTableRow > >& Rows, const EColumnSortMode::Type SortMode ) override;
 
-	virtual void Tick() OVERRIDE {}
+	virtual void Tick() override {}
 
 	DECLARE_DERIVED_EVENT( FPropertyTableColumn, IPropertyTableColumn::FFrozenStateChanged, FFrozenStateChanged );
-	FFrozenStateChanged* OnFrozenStateChanged() OVERRIDE { return &FrozenStateChanged; }
+	FFrozenStateChanged* OnFrozenStateChanged() override { return &FrozenStateChanged; }
 
 	// End IPropertyTableColumn Interface
 
diff --git a/Engine/Source/Editor/PropertyEditor/Private/Presentation/PropertyTable/PropertyTableRow.h b/Engine/Source/Editor/PropertyEditor/Private/Presentation/PropertyTable/PropertyTableRow.h
--- a/Engine/Source/Editor/PropertyEditor/Private/Presentation/PropertyTable/PropertyTableRow.h
+++ b/Engine/Source/Editor/PropertyEditor/Private/Presentation/PropertyTable/PropertyTableRow.h
@@ -15,21 +15,21 @@ public:
 
 	FPropertyTableRow( const TSharedRef< class IPropertyTable >& InTable, const TSharedRef< FPropertyPath >& InPropertyPath, const TSharedRef< FPropertyPath >& InPartialPropertyPath );
 
-	virtual TSharedRef< class IDataSource > GetDataSource() const OVERRIDE { return DataSource; }
+	virtual TSharedRef< class IDataSource > GetDataSource() const override { return DataSource; }
 
-	virtual bool HasChildren() const OVERRIDE;
+	virtual bool HasChildren() const override;
 
-	virtual void GetChildRows( TArray< TSharedRef< class IPropertyTableRow > >& OutChildren ) const OVERRIDE;
+	virtual void GetChildRows( TArray< TSharedRef< class IPropertyTableRow > >& OutChildren ) const override;
 
-	virtual TSharedRef< class IPropertyTable > GetTable() const OVERRIDE;
+	virtual TSharedRef< class IPropertyTable > GetTable() const override;
 
-	virtual bool HasCells() const OVERRIDE { return true; }
+	virtual bool HasCells() const override { return true; }
 
-	virtual TSharedRef< class FPropertyPath > GetPartialPath() const OVERRIDE { return PartialPath; }
+	virtual TSharedRef< class FPropertyPath > GetPartialPath() const override { return PartialPath; }
 
-	virtual void Tick() OVERRIDE;
+	virtual void Tick() override;
 
-	virtual void Refresh() OVERRIDE;
+	virtual void Refresh() override;
 
 	DECLARE_DERIVED_EVENT( FPropertyTableRow, IPropertyTableRow::FRefreshed, FRefreshed );
 	virtual FRefreshed* OnRefresh() { return &Refreshed; }
diff --git a/Engine/Source/Editor/PropertyEditor/Private/Presentation/PropertyTable/PropertyTableRowHeaderColumn.h b/Engine/Source/Editor/PropertyEditor/Private/Presentation/PropertyTable/PropertyTableRowHeaderColumn.h
--- a/Engine/Source/Editor/PropertyEditor/Private/Presentation/PropertyTable/PropertyTableRowHeaderColumn.h
+++ b/Engine/Source/Editor/PropertyEditor/Private/Presentation/PropertyTable/PropertyTableRowHeaderColumn.h
@@ -23,12 +23,12 @@ public:
 
 	}
 
-	virtual FName GetId() const OVERRIDE { return FName( TEXT("RowHeader") ); }
-	virtual FText GetDisplayName() const OVERRIDE { return FText::GetEmpty(); }
-	virtual TSharedRef< IDataSource > GetDataSource() const OVERRIDE { return MakeShareable( new PropertyPathDataSource( FPropertyPath::CreateEmpty() ) ); }
-	virtual TSharedRef< class FPropertyPath > GetPartialPath() const OVERRIDE { return FPropertyPath::CreateEmpty(); }
+	virtual FName GetId() const override { return FName( TEXT("RowHeader") ); }
+	virtual FText GetDisplayName() const override { return FText::GetEmpty(); }
+	virtual TSharedRef< IDataSource > GetDataSource() const override { return MakeShareable( new PropertyPathDataSource( FPropertyPath::CreateEmpty() ) ); }
+	virtual TSharedRef< class FPropertyPath > GetPartialPath() const override { return FPropertyPath::CreateEmpty(); }
 
-	virtual TSharedRef< class IPropertyTableCell > GetCell( const TSharedRef< class IPropertyTableRow >& Row ) OVERRIDE
+	virtual TSharedRef< class IPropertyTableCell > GetCell( const TSharedRef< class IPropertyTableRow >& Row ) override
 	{
 		//@todo Clean Cells cache when rows get updated [11/27/2012 Justin.Sargent]
 		TSharedRef< IPropertyTableCell >* CellPtr = Cells.Find( Row );
@@ -44,39 +44,39 @@ public:
 		return Cell;
 	}
 
-	virtual void RemoveCellsForRow( const TSharedRef< class IPropertyTableRow >& Row ) OVERRIDE
+	virtual void RemoveCellsForRow( const TSharedRef< class IPropertyTableRow >& Row ) override
 	{
 		Cells.Remove( Row );
 	}
 
-	virtual TSharedRef< class IPropertyTable > GetTable() const OVERRIDE { return Table.Pin().ToSharedRef(); }
+	virtual TSharedRef< class IPropertyTable > GetTable() const override { return Table.Pin().ToSharedRef(); }
 
-	virtual bool CanSelectCells() const OVERRIDE { return false; }
+	virtual bool CanSelectCells() const override { return false; }
 
-	virtual EPropertyTableColumnSizeMode::Type GetSizeMode() const OVERRIDE { return EPropertyTableColumnSizeMode::Fixed; }
+	virtual EPropertyTableColumnSizeMode::Type GetSizeMode() const override { return EPropertyTableColumnSizeMode::Fixed; }
 
-	virtual void SetSizeMode(EPropertyTableColumnSizeMode::Type InSizeMode) OVERRIDE {}
+	virtual void SetSizeMode(EPropertyTableColumnSizeMode::Type InSizeMode) override {}
 
-	virtual float GetWidth() const OVERRIDE { return 20.0f; } 
+	virtual float GetWidth() const override { return 20.0f; } 
 
-	virtual void SetWidth( float InWidth ) OVERRIDE {  }
+	virtual void SetWidth( float InWidth ) override {  }
 
-	virtual bool IsHidden() const OVERRIDE { return bIsHidden; }
+	virtual bool IsHidden() const override { return bIsHidden; }
 
-	virtual void SetHidden( bool InIsHidden ) OVERRIDE { bIsHidden = InIsHidden; }
+	virtual void SetHidden( bool InIsHidden ) override { bIsHidden = InIsHidden; }
 
-	virtual bool IsFrozen() const OVERRIDE { return true; }
+	virtual bool IsFrozen() const override { return true; }
 
-	virtual void SetFrozen( bool InIsFrozen ) OVERRIDE {}
+	virtual void SetFrozen( bool InIsFrozen ) override {}
 
-	virtual bool CanSortBy() const OVERRIDE { return false; }
+	virtual bool CanSortBy() const override { return false; }
 
-	virtual void Sort( TArray< TSharedRef< class IPropertyTableRow > >& Rows, const EColumnSortMode::Type SortMode ) OVERRIDE {  }
+	virtual void Sort( TArray< TSharedRef< class IPropertyTableRow > >& Rows, const EColumnSortMode::Type SortMode ) override {  }
 
-	virtual void Tick() OVERRIDE {}
+	virtual void Tick() override {}
 
 	DECLARE_DERIVED_EVENT( FPropertyTableColumn, IPropertyTableColumn::FFrozenStateChanged, FFrozenStateChanged );
-	FFrozenStateChanged* OnFrozenStateChanged() OVERRIDE { return &FrozenStateChanged; }
+	FFrozenStateChanged* OnFrozenStateChanged() override { return &FrozenStateChanged; }
 
 private:
 
diff --git a/Engine/Source/Editor/PropertyEditor/Private/PropertyChangeListener.h b/Engine/Source/Editor/PropertyEditor/Private/PropertyChangeListener.h
--- a/Engine/Source/Editor/PropertyEditor/Private/PropertyChangeListener.h
+++ b/Engine/Source/Editor/PropertyEditor/Private/PropertyChangeListener.h
@@ -9,10 +9,10 @@ class FPropertyChangeListener : public IPropertyChangeListener
 {
 public:
 	/** IPropertyChangeListener interface */
-	virtual void SetObject( UObject& Object, const FPropertyListenerSettings& InPropertyListenerSettings ) OVERRIDE;
-	virtual bool ScanForChanges( bool bRecacheNewValues = true ) OVERRIDE;
-	virtual void TriggerAllPropertiesChangedDelegate() OVERRIDE;
-	virtual FOnPropertyChanged& GetOnPropertyChangedDelegate() OVERRIDE { return OnPropertyChangedDelegate; }
+	virtual void SetObject( UObject& Object, const FPropertyListenerSettings& InPropertyListenerSettings ) override;
+	virtual bool ScanForChanges( bool bRecacheNewValues = true ) override;
+	virtual void TriggerAllPropertiesChangedDelegate() override;
+	virtual FOnPropertyChanged& GetOnPropertyChangedDelegate() override { return OnPropertyChangedDelegate; }
 private:
 	/**
 	 * Recursively creates a property value cache for each property node in the tree starting with the passed in property node
diff --git a/Engine/Source/Editor/PropertyEditor/Private/PropertyCustomizationHelpers.cpp b/Engine/Source/Editor/PropertyEditor/Private/PropertyCustomizationHelpers.cpp
--- a/Engine/Source/Editor/PropertyEditor/Private/PropertyCustomizationHelpers.cpp
+++ b/Engine/Source/Editor/PropertyEditor/Private/PropertyCustomizationHelpers.cpp
@@ -453,7 +453,7 @@ public:
 	 * @param Material		The material being used
 	 * @param bCanBeReplced	Whether or not the material can be replaced by a user
 	 */
-	virtual void AddMaterial( uint32 SlotIndex, UMaterialInterface* Material, bool bCanBeReplaced ) OVERRIDE
+	virtual void AddMaterial( uint32 SlotIndex, UMaterialInterface* Material, bool bCanBeReplaced ) override
 	{
 		int32 NumMaterials = MaterialSlots.Num();
 
diff --git a/Engine/Source/Editor/PropertyEditor/Private/PropertyEditorToolkit.h b/Engine/Source/Editor/PropertyEditor/Private/PropertyEditorToolkit.h
--- a/Engine/Source/Editor/PropertyEditor/Private/PropertyEditorToolkit.h
+++ b/Engine/Source/Editor/PropertyEditor/Private/PropertyEditorToolkit.h
@@ -9,17 +9,17 @@ public:
 
 	FPropertyEditorToolkit();
 
-	virtual void RegisterTabSpawners(const TSharedRef<class FTabManager>& TabManager) OVERRIDE;
+	virtual void RegisterTabSpawners(const TSharedRef<class FTabManager>& TabManager) override;
 
-	virtual void UnregisterTabSpawners(const TSharedRef<class FTabManager>& TabManager) OVERRIDE;
+	virtual void UnregisterTabSpawners(const TSharedRef<class FTabManager>& TabManager) override;
 
-	virtual FName GetToolkitFName() const OVERRIDE;
+	virtual FName GetToolkitFName() const override;
 
-	virtual FText GetBaseToolkitName() const OVERRIDE;
+	virtual FText GetBaseToolkitName() const override;
 
-	virtual FText GetToolkitName() const OVERRIDE;
+	virtual FText GetToolkitName() const override;
 
-	virtual FString GetWorldCentricTabPrefix() const OVERRIDE;
+	virtual FString GetWorldCentricTabPrefix() const override;
 
 	virtual FLinearColor GetWorldCentricTabColorScale() const;
 
@@ -29,9 +29,9 @@ public:
 
 	bool TableHasCustomColumns() const;
 
-	virtual bool CloseWindow() OVERRIDE;
+	virtual bool CloseWindow() override;
 
-	virtual bool IsPrimaryEditor() const OVERRIDE{ return false; };
+	virtual bool IsPrimaryEditor() const override{ return false; };
 
 public:
 
diff --git a/Engine/Source/Editor/PropertyEditor/Private/PropertyHandleImpl.h b/Engine/Source/Editor/PropertyEditor/Private/PropertyHandleImpl.h
--- a/Engine/Source/Editor/PropertyEditor/Private/PropertyHandleImpl.h
+++ b/Engine/Source/Editor/PropertyEditor/Private/PropertyHandleImpl.h
@@ -333,8 +333,8 @@ protected:
 };
 
 #define DECLARE_PROPERTY_ACCESSOR( ValueType ) \
-	virtual FPropertyAccess::Result SetValue( const ValueType& InValue, EPropertyValueSetFlags::Type Flags = EPropertyValueSetFlags::DefaultFlags ) OVERRIDE; \
-	virtual FPropertyAccess::Result GetValue( ValueType& OutValue ) const OVERRIDE; 
+	virtual FPropertyAccess::Result SetValue( const ValueType& InValue, EPropertyValueSetFlags::Type Flags = EPropertyValueSetFlags::DefaultFlags ) override; \
+	virtual FPropertyAccess::Result GetValue( ValueType& OutValue ) const override; 
 
 class FDetailCategoryImpl;
 
@@ -361,52 +361,52 @@ public:
 	DECLARE_PROPERTY_ACCESSOR( UObject* )
 
 	/** IPropertyHandle interface */
-	virtual bool IsValidHandle() const OVERRIDE;
-	virtual FString GetPropertyDisplayName() const OVERRIDE;
-	virtual void ResetToDefault() OVERRIDE;
-	virtual bool DiffersFromDefault() const OVERRIDE;
-	virtual FText GetResetToDefaultLabel() const OVERRIDE;
-	virtual void MarkHiddenByCustomization() OVERRIDE;
-	virtual bool IsCustomized() const OVERRIDE;
-	virtual TSharedRef<SWidget> CreatePropertyNameWidget( const FString& NameOverride = TEXT(""), bool bDisplayResetToDefault = false, bool bDisplayText = true, bool bDisplayThumbnail = true ) const OVERRIDE;
-	virtual TSharedRef<SWidget> CreatePropertyValueWidget() const OVERRIDE;
-	virtual bool IsEditConst() const OVERRIDE;
-	virtual void SetOnPropertyValueChanged( const FSimpleDelegate& InOnPropertyValueChanged ) OVERRIDE;
-	virtual int32 GetIndexInArray() const OVERRIDE;
-	virtual FPropertyAccess::Result GetValueAsFormattedString( FString& OutValue ) const OVERRIDE;
-	virtual FPropertyAccess::Result GetValueAsDisplayString( FString& OutValue ) const OVERRIDE;
-	virtual FPropertyAccess::Result GetValueAsFormattedText( FText& OutValue ) const OVERRIDE;
-	virtual FPropertyAccess::Result GetValueAsDisplayText( FText& OutValue ) const OVERRIDE;
-	virtual FPropertyAccess::Result SetValueFromFormattedString( const FString& InValue,  EPropertyValueSetFlags::Type Flags = EPropertyValueSetFlags::DefaultFlags ) OVERRIDE;
-	virtual TSharedPtr<IPropertyHandle> GetChildHandle( uint32 ChildIndex ) const OVERRIDE;
-	virtual TSharedPtr<IPropertyHandle> GetChildHandle( FName ChildName ) const OVERRIDE;
-	virtual TSharedPtr<IPropertyHandle> GetParentHandle() const OVERRIDE;
-	virtual void AccessRawData( TArray<void*>& RawData ) OVERRIDE;
-	virtual uint32 GetNumOuterObjects() const OVERRIDE;
-	virtual void GetOuterObjects( TArray<UObject*>& OuterObjects ) OVERRIDE;
-	virtual FPropertyAccess::Result GetNumChildren( uint32& OutNumChildren ) const OVERRIDE;
-	virtual TSharedPtr<IPropertyHandleArray> AsArray() OVERRIDE { return NULL; }
-	virtual const UClass* GetPropertyClass() const OVERRIDE;
-	virtual UProperty* GetProperty() const OVERRIDE;
-	virtual UProperty* GetMetaDataProperty() const OVERRIDE;
-	virtual bool HasMetaData(const FName& Key) const OVERRIDE;
-	virtual const FString& GetMetaData(const FName& Key) const OVERRIDE;
-	virtual bool GetBoolMetaData(const FName& Key) const OVERRIDE;
-	virtual int32 GetINTMetaData(const FName& Key) const OVERRIDE;
-	virtual float GetFLOATMetaData(const FName& Key) const OVERRIDE;
-	virtual UClass* GetClassMetaData(const FName& Key) const OVERRIDE;
-	virtual FString GetToolTipText() const OVERRIDE;
-	virtual void SetToolTipText(const FString& ToolTip) OVERRIDE;
-	virtual FPropertyAccess::Result SetPerObjectValues( const TArray<FString>& InPerObjectValues,  EPropertyValueSetFlags::Type Flags = EPropertyValueSetFlags::DefaultFlags ) OVERRIDE;
-	virtual FPropertyAccess::Result GetPerObjectValues( TArray<FString>& OutPerObjectValues ) OVERRIDE;
-	virtual bool GeneratePossibleValues(TArray< TSharedPtr<FString> >& OutOptionStrings, TArray< TSharedPtr<FString> >& OutToolTips, TArray<bool>& OutRestrictedItems) OVERRIDE;
-	virtual FPropertyAccess::Result SetObjectValueFromSelection() OVERRIDE;
-	virtual void NotifyPreChange() OVERRIDE;
-	virtual void NotifyPostChange() OVERRIDE;
-	virtual void AddRestriction( TSharedRef<const FPropertyRestriction> Restriction )OVERRIDE;
-	virtual bool IsRestricted(const FString& Value) const OVERRIDE;
-	virtual bool IsRestricted(const FString& Value, TArray<FText>& OutReasons) const OVERRIDE;
-	virtual bool GenerateRestrictionToolTip(const FString& Value, FText& OutTooltip) const OVERRIDE;
+	virtual bool IsValidHandle() const override;
+	virtual FString GetPropertyDisplayName() const override;
+	virtual void ResetToDefault() override;
+	virtual bool DiffersFromDefault() const override;
+	virtual FText GetResetToDefaultLabel() const override;
+	virtual void MarkHiddenByCustomization() override;
+	virtual bool IsCustomized() const override;
+	virtual TSharedRef<SWidget> CreatePropertyNameWidget( const FString& NameOverride = TEXT(""), bool bDisplayResetToDefault = false, bool bDisplayText = true, bool bDisplayThumbnail = true ) const override;
+	virtual TSharedRef<SWidget> CreatePropertyValueWidget() const override;
+	virtual bool IsEditConst() const override;
+	virtual void SetOnPropertyValueChanged( const FSimpleDelegate& InOnPropertyValueChanged ) override;
+	virtual int32 GetIndexInArray() const override;
+	virtual FPropertyAccess::Result GetValueAsFormattedString( FString& OutValue ) const override;
+	virtual FPropertyAccess::Result GetValueAsDisplayString( FString& OutValue ) const override;
+	virtual FPropertyAccess::Result GetValueAsFormattedText( FText& OutValue ) const override;
+	virtual FPropertyAccess::Result GetValueAsDisplayText( FText& OutValue ) const override;
+	virtual FPropertyAccess::Result SetValueFromFormattedString( const FString& InValue,  EPropertyValueSetFlags::Type Flags = EPropertyValueSetFlags::DefaultFlags ) override;
+	virtual TSharedPtr<IPropertyHandle> GetChildHandle( uint32 ChildIndex ) const override;
+	virtual TSharedPtr<IPropertyHandle> GetChildHandle( FName ChildName ) const override;
+	virtual TSharedPtr<IPropertyHandle> GetParentHandle() const override;
+	virtual void AccessRawData( TArray<void*>& RawData ) override;
+	virtual uint32 GetNumOuterObjects() const override;
+	virtual void GetOuterObjects( TArray<UObject*>& OuterObjects ) override;
+	virtual FPropertyAccess::Result GetNumChildren( uint32& OutNumChildren ) const override;
+	virtual TSharedPtr<IPropertyHandleArray> AsArray() override { return NULL; }
+	virtual const UClass* GetPropertyClass() const override;
+	virtual UProperty* GetProperty() const override;
+	virtual UProperty* GetMetaDataProperty() const override;
+	virtual bool HasMetaData(const FName& Key) const override;
+	virtual const FString& GetMetaData(const FName& Key) const override;
+	virtual bool GetBoolMetaData(const FName& Key) const override;
+	virtual int32 GetINTMetaData(const FName& Key) const override;
+	virtual float GetFLOATMetaData(const FName& Key) const override;
+	virtual UClass* GetClassMetaData(const FName& Key) const override;
+	virtual FString GetToolTipText() const override;
+	virtual void SetToolTipText(const FString& ToolTip) override;
+	virtual FPropertyAccess::Result SetPerObjectValues( const TArray<FString>& InPerObjectValues,  EPropertyValueSetFlags::Type Flags = EPropertyValueSetFlags::DefaultFlags ) override;
+	virtual FPropertyAccess::Result GetPerObjectValues( TArray<FString>& OutPerObjectValues ) override;
+	virtual bool GeneratePossibleValues(TArray< TSharedPtr<FString> >& OutOptionStrings, TArray< TSharedPtr<FString> >& OutToolTips, TArray<bool>& OutRestrictedItems) override;
+	virtual FPropertyAccess::Result SetObjectValueFromSelection() override;
+	virtual void NotifyPreChange() override;
+	virtual void NotifyPostChange() override;
+	virtual void AddRestriction( TSharedRef<const FPropertyRestriction> Restriction )override;
+	virtual bool IsRestricted(const FString& Value) const override;
+	virtual bool IsRestricted(const FString& Value, TArray<FText>& OutReasons) const override;
+	virtual bool GenerateRestrictionToolTip(const FString& Value, FText& OutTooltip) const override;
 
 	TSharedPtr<FPropertyNode> GetPropertyNode();
 protected:
@@ -418,8 +418,8 @@ class FPropertyHandleInt : public FPropertyHandleBase
 public:
 	FPropertyHandleInt( TSharedRef<FPropertyNode> PropertyNode, FNotifyHook* NotifyHook, TSharedPtr<IPropertyUtilities> PropertyUtilities );
 	static bool Supports( TSharedRef<FPropertyNode> PropertyNode );
-	virtual FPropertyAccess::Result GetValue( int32& OutValue ) const OVERRIDE;
-	virtual FPropertyAccess::Result SetValue( const int32& InValue, EPropertyValueSetFlags::Type Flags = EPropertyValueSetFlags::DefaultFlags ) OVERRIDE;
+	virtual FPropertyAccess::Result GetValue( int32& OutValue ) const override;
+	virtual FPropertyAccess::Result SetValue( const int32& InValue, EPropertyValueSetFlags::Type Flags = EPropertyValueSetFlags::DefaultFlags ) override;
 };
 
 class FPropertyHandleFloat : public FPropertyHandleBase
@@ -427,8 +427,8 @@ class FPropertyHandleFloat : public FPropertyHandleBase
 public:
 	FPropertyHandleFloat( TSharedRef<FPropertyNode> PropertyNode, FNotifyHook* NotifyHook, TSharedPtr<IPropertyUtilities> PropertyUtilities );
 	static bool Supports( TSharedRef<FPropertyNode> PropertyNode );
-	virtual FPropertyAccess::Result GetValue( float& OutValue ) const OVERRIDE;
-	virtual FPropertyAccess::Result SetValue( const float& InValue, EPropertyValueSetFlags::Type Flags = EPropertyValueSetFlags::DefaultFlags ) OVERRIDE;
+	virtual FPropertyAccess::Result GetValue( float& OutValue ) const override;
+	virtual FPropertyAccess::Result SetValue( const float& InValue, EPropertyValueSetFlags::Type Flags = EPropertyValueSetFlags::DefaultFlags ) override;
 };
 
 class FPropertyHandleBool : public FPropertyHandleBase
@@ -436,8 +436,8 @@ class FPropertyHandleBool : public FPropertyHandleBase
 public:
 	FPropertyHandleBool( TSharedRef<FPropertyNode> PropertyNode, FNotifyHook* NotifyHook, TSharedPtr<IPropertyUtilities> PropertyUtilities );
 	static bool Supports( TSharedRef<FPropertyNode> PropertyNode );
-	virtual FPropertyAccess::Result GetValue( bool& OutValue ) const OVERRIDE;
-	virtual FPropertyAccess::Result SetValue( const bool& InValue, EPropertyValueSetFlags::Type Flags = EPropertyValueSetFlags::DefaultFlags ) OVERRIDE;
+	virtual FPropertyAccess::Result GetValue( bool& OutValue ) const override;
+	virtual FPropertyAccess::Result SetValue( const bool& InValue, EPropertyValueSetFlags::Type Flags = EPropertyValueSetFlags::DefaultFlags ) override;
 };
 
 class FPropertyHandleByte : public FPropertyHandleBase
@@ -445,8 +445,8 @@ class FPropertyHandleByte : public FPropertyHandleBase
 public:
 	FPropertyHandleByte( TSharedRef<FPropertyNode> PropertyNode, FNotifyHook* NotifyHook, TSharedPtr<IPropertyUtilities> PropertyUtilities );
 	static bool Supports( TSharedRef<FPropertyNode> PropertyNode );
-	virtual FPropertyAccess::Result GetValue( uint8& OutValue ) const OVERRIDE;
-	virtual FPropertyAccess::Result SetValue( const uint8& InValue, EPropertyValueSetFlags::Type Flags = EPropertyValueSetFlags::DefaultFlags ) OVERRIDE;
+	virtual FPropertyAccess::Result GetValue( uint8& OutValue ) const override;
+	virtual FPropertyAccess::Result SetValue( const uint8& InValue, EPropertyValueSetFlags::Type Flags = EPropertyValueSetFlags::DefaultFlags ) override;
 };
 
 class FPropertyHandleString : public FPropertyHandleBase
@@ -454,11 +454,11 @@ class FPropertyHandleString : public FPropertyHandleBase
 public:
 	FPropertyHandleString( TSharedRef<FPropertyNode> PropertyNode, FNotifyHook* NotifyHook,TSharedPtr<IPropertyUtilities> PropertyUtilities );
 	static bool Supports( TSharedRef<FPropertyNode> PropertyNode );
-	virtual FPropertyAccess::Result GetValue( FString& OutValue ) const OVERRIDE;
-	virtual FPropertyAccess::Result SetValue( const FString& InValue, EPropertyValueSetFlags::Type Flags = EPropertyValueSetFlags::DefaultFlags ) OVERRIDE;
+	virtual FPropertyAccess::Result GetValue( FString& OutValue ) const override;
+	virtual FPropertyAccess::Result SetValue( const FString& InValue, EPropertyValueSetFlags::Type Flags = EPropertyValueSetFlags::DefaultFlags ) override;
 
-	virtual FPropertyAccess::Result GetValue( FName& OutValue ) const OVERRIDE;
-	virtual FPropertyAccess::Result SetValue( const FName& InValue, EPropertyValueSetFlags::Type Flags = EPropertyValueSetFlags::DefaultFlags ) OVERRIDE;
+	virtual FPropertyAccess::Result GetValue( FName& OutValue ) const override;
+	virtual FPropertyAccess::Result SetValue( const FName& InValue, EPropertyValueSetFlags::Type Flags = EPropertyValueSetFlags::DefaultFlags ) override;
 };
 
 class FPropertyHandleObject : public FPropertyHandleBase
@@ -466,8 +466,8 @@ class FPropertyHandleObject : public FPropertyHandleBase
 public:
 	FPropertyHandleObject( TSharedRef<FPropertyNode> PropertyNode, FNotifyHook* NotifyHook, TSharedPtr<IPropertyUtilities> PropertyUtilities );
 	static bool Supports( TSharedRef<FPropertyNode> PropertyNode );
-	virtual FPropertyAccess::Result GetValue( UObject*& OutValue ) const OVERRIDE;
-	virtual FPropertyAccess::Result SetValue( const UObject*& InValue, EPropertyValueSetFlags::Type Flags = EPropertyValueSetFlags::DefaultFlags ) OVERRIDE;
+	virtual FPropertyAccess::Result GetValue( UObject*& OutValue ) const override;
+	virtual FPropertyAccess::Result SetValue( const UObject*& InValue, EPropertyValueSetFlags::Type Flags = EPropertyValueSetFlags::DefaultFlags ) override;
 
 };
 
@@ -476,17 +476,17 @@ class FPropertyHandleVector : public FPropertyHandleBase
 public:
 	FPropertyHandleVector( TSharedRef<FPropertyNode> PropertyNode, FNotifyHook* NotifyHook, TSharedPtr<IPropertyUtilities> PropertyUtilities );
 	static bool Supports( TSharedRef<FPropertyNode> PropertyNode );
-	virtual FPropertyAccess::Result GetValue( FVector& OutValue ) const OVERRIDE;
-	virtual FPropertyAccess::Result SetValue( const FVector& InValue, EPropertyValueSetFlags::Type Flags = EPropertyValueSetFlags::DefaultFlags ) OVERRIDE;
+	virtual FPropertyAccess::Result GetValue( FVector& OutValue ) const override;
+	virtual FPropertyAccess::Result SetValue( const FVector& InValue, EPropertyValueSetFlags::Type Flags = EPropertyValueSetFlags::DefaultFlags ) override;
 
-	virtual FPropertyAccess::Result GetValue( FVector2D& OutValue ) const OVERRIDE;
-	virtual FPropertyAccess::Result SetValue( const FVector2D& InValue, EPropertyValueSetFlags::Type Flags = EPropertyValueSetFlags::DefaultFlags ) OVERRIDE;
+	virtual FPropertyAccess::Result GetValue( FVector2D& OutValue ) const override;
+	virtual FPropertyAccess::Result SetValue( const FVector2D& InValue, EPropertyValueSetFlags::Type Flags = EPropertyValueSetFlags::DefaultFlags ) override;
 
-	virtual FPropertyAccess::Result GetValue( FVector4& OutValue ) const OVERRIDE;
-	virtual FPropertyAccess::Result SetValue( const FVector4& InValue, EPropertyValueSetFlags::Type Flags = EPropertyValueSetFlags::DefaultFlags ) OVERRIDE;
+	virtual FPropertyAccess::Result GetValue( FVector4& OutValue ) const override;
+	virtual FPropertyAccess::Result SetValue( const FVector4& InValue, EPropertyValueSetFlags::Type Flags = EPropertyValueSetFlags::DefaultFlags ) override;
 
-	virtual FPropertyAccess::Result GetValue( FQuat& OutValue ) const OVERRIDE;
-	virtual FPropertyAccess::Result SetValue( const FQuat& InValue, EPropertyValueSetFlags::Type Flags = EPropertyValueSetFlags::DefaultFlags ) OVERRIDE;
+	virtual FPropertyAccess::Result GetValue( FQuat& OutValue ) const override;
+	virtual FPropertyAccess::Result SetValue( const FQuat& InValue, EPropertyValueSetFlags::Type Flags = EPropertyValueSetFlags::DefaultFlags ) override;
 
 	virtual FPropertyAccess::Result SetX( float InValue, EPropertyValueSetFlags::Type Flags = EPropertyValueSetFlags::DefaultFlags );
 	virtual FPropertyAccess::Result SetY( float InValue, EPropertyValueSetFlags::Type Flags = EPropertyValueSetFlags::DefaultFlags );
@@ -501,8 +501,8 @@ class FPropertyHandleRotator : public FPropertyHandleBase
 public:
 	FPropertyHandleRotator( TSharedRef<FPropertyNode> PropertyNode, FNotifyHook* NotifyHook, TSharedPtr<IPropertyUtilities> PropertyUtilities  );
 	static bool Supports( TSharedRef<FPropertyNode> PropertyNode );
-	virtual FPropertyAccess::Result GetValue( FRotator& OutValue ) const OVERRIDE;
-	virtual FPropertyAccess::Result SetValue( const FRotator& InValue, EPropertyValueSetFlags::Type Flags = EPropertyValueSetFlags::DefaultFlags ) OVERRIDE;
+	virtual FPropertyAccess::Result GetValue( FRotator& OutValue ) const override;
+	virtual FPropertyAccess::Result SetValue( const FRotator& InValue, EPropertyValueSetFlags::Type Flags = EPropertyValueSetFlags::DefaultFlags ) override;
 
 	virtual FPropertyAccess::Result SetRoll( float InValue, EPropertyValueSetFlags::Type Flags = EPropertyValueSetFlags::DefaultFlags );
 	virtual FPropertyAccess::Result SetPitch( float InValue, EPropertyValueSetFlags::Type Flags = EPropertyValueSetFlags::DefaultFlags );
@@ -520,15 +520,15 @@ public:
 	FPropertyHandleArray( TSharedRef<FPropertyNode> PropertyNode, FNotifyHook* NotifyHook, TSharedPtr<IPropertyUtilities> PropertyUtilities );
 	static bool Supports( TSharedRef<FPropertyNode> PropertyNode );
 	/** IPropertyHandleArray interface */
-	virtual FPropertyAccess::Result AddItem() OVERRIDE;
-	virtual FPropertyAccess::Result EmptyArray() OVERRIDE;
-	virtual FPropertyAccess::Result Insert( int32 Index ) OVERRIDE;
-	virtual FPropertyAccess::Result DuplicateItem( int32 Index ) OVERRIDE;
-	virtual FPropertyAccess::Result DeleteItem( int32 Index ) OVERRIDE;
-	virtual FPropertyAccess::Result GetNumElements( uint32& OutNumItems ) const OVERRIDE;
-	virtual void SetOnNumElementsChanged( FSimpleDelegate& InOnNumElementsChanged ) OVERRIDE;
-	virtual TSharedPtr<IPropertyHandleArray> AsArray() OVERRIDE;
-	virtual TSharedRef<IPropertyHandle> GetElement( int32 Index ) const OVERRIDE;
+	virtual FPropertyAccess::Result AddItem() override;
+	virtual FPropertyAccess::Result EmptyArray() override;
+	virtual FPropertyAccess::Result Insert( int32 Index ) override;
+	virtual FPropertyAccess::Result DuplicateItem( int32 Index ) override;
+	virtual FPropertyAccess::Result DeleteItem( int32 Index ) override;
+	virtual FPropertyAccess::Result GetNumElements( uint32& OutNumItems ) const override;
+	virtual void SetOnNumElementsChanged( FSimpleDelegate& InOnNumElementsChanged ) override;
+	virtual TSharedPtr<IPropertyHandleArray> AsArray() override;
+	virtual TSharedRef<IPropertyHandle> GetElement( int32 Index ) const override;
 private:
 	/**
 	 * @return Whether or not the array can be modified
diff --git a/Engine/Source/Editor/PropertyEditor/Private/SDetailSingleItemRow.h b/Engine/Source/Editor/PropertyEditor/Private/SDetailSingleItemRow.h
--- a/Engine/Source/Editor/PropertyEditor/Private/SDetailSingleItemRow.h
+++ b/Engine/Source/Editor/PropertyEditor/Private/SDetailSingleItemRow.h
@@ -22,7 +22,7 @@ public:
 	void Construct( const FArguments& InArgs, FDetailLayoutCustomization* InCustomization, bool bHasMultipleColumns, TSharedRef<IDetailTreeNode> InOwnerTreeNode, const TSharedRef<STableViewBase>& InOwnerTableView );
 
 protected:
-	virtual bool OnContextMenuOpening( FMenuBuilder& MenuBuilder ) OVERRIDE;
+	virtual bool OnContextMenuOpening( FMenuBuilder& MenuBuilder ) override;
 private:
 	void OnLeftColumnResized( float InNewWidth );
 	void OnCopyProperty();
diff --git a/Engine/Source/Editor/PropertyEditor/Private/SDetailTableRowBase.h b/Engine/Source/Editor/PropertyEditor/Private/SDetailTableRowBase.h
--- a/Engine/Source/Editor/PropertyEditor/Private/SDetailTableRowBase.h
+++ b/Engine/Source/Editor/PropertyEditor/Private/SDetailTableRowBase.h
@@ -5,7 +5,7 @@
 class SDetailTableRowBase : public STableRow< TSharedPtr< IDetailTreeNode > >
 {
 public:
-	virtual FReply OnMouseButtonUp( const FGeometry& MyGeometry, const FPointerEvent& MouseEvent ) OVERRIDE
+	virtual FReply OnMouseButtonUp( const FGeometry& MyGeometry, const FPointerEvent& MouseEvent ) override
 	{
 		if( OwnerTreeNode.IsValid() && MouseEvent.GetEffectingButton() == EKeys::RightMouseButton && !StaticCastSharedRef<STableViewBase>( OwnerTablePtr.Pin()->AsWidget() )->IsRightClickScrolling() )
 		{
diff --git a/Engine/Source/Editor/PropertyEditor/Private/SDetailsView.h b/Engine/Source/Editor/PropertyEditor/Private/SDetailsView.h
--- a/Engine/Source/Editor/PropertyEditor/Private/SDetailsView.h
+++ b/Engine/Source/Editor/PropertyEditor/Private/SDetailsView.h
@@ -19,7 +19,7 @@ public:
 	virtual ~SDetailsView();
 
 	/** Causes the details view to be refreshed (new widgets generated) with the current set of objects */
-	void ForceRefresh() OVERRIDE;
+	void ForceRefresh() override;
 
 	/**
 	 * Constructs the property view widgets                   
@@ -27,9 +27,9 @@ public:
 	void Construct(const FArguments& InArgs);
 
 	/** IDetailsView interface */
-	virtual void SetObjects( const TArray<UObject*>& InObjects, bool bForceRefresh = false ) OVERRIDE;
-	virtual void SetObjects( const TArray< TWeakObjectPtr< UObject > >& InObjects, bool bForceRefresh = false ) OVERRIDE;
-	virtual void SetObject( UObject* InObject, bool bForceRefresh = false ) OVERRIDE;
+	virtual void SetObjects( const TArray<UObject*>& InObjects, bool bForceRefresh = false ) override;
+	virtual void SetObjects( const TArray< TWeakObjectPtr< UObject > >& InObjects, bool bForceRefresh = false ) override;
+	virtual void SetObject( UObject* InObject, bool bForceRefresh = false ) override;
 
 	/**
 	 * Replaces objects being observed by the view with new objects
@@ -58,42 +58,42 @@ public:
 	 * @param PropertyEditor				The slate property node to edit.
 	 * @param bUseAlpha			Whether or not alpha is supported
 	 */
-	virtual void CreateColorPickerWindow(const TSharedRef< class FPropertyEditor >& PropertyEditor, bool bUseAlpha) OVERRIDE;
+	virtual void CreateColorPickerWindow(const TSharedRef< class FPropertyEditor >& PropertyEditor, bool bUseAlpha) override;
 
 	/** Sets the callback for when the property view changes */
 	virtual void SetOnObjectArrayChanged( FOnObjectArrayChanged OnObjectArrayChangedDelegate);
 
 	/** @return	Returns list of selected objects we're inspecting */
-	virtual const TArray< TWeakObjectPtr<UObject> >& GetSelectedObjects() const OVERRIDE
+	virtual const TArray< TWeakObjectPtr<UObject> >& GetSelectedObjects() const override
 	{
 		return SelectedObjects;
 	} 
 
 	/** @return	Returns list of selected actors we're inspecting */
-	virtual const TArray< TWeakObjectPtr<AActor> >& GetSelectedActors() const OVERRIDE
+	virtual const TArray< TWeakObjectPtr<AActor> >& GetSelectedActors() const override
 	{
 		return SelectedActors;
 	}
 
 	/** @return Returns information about the selected set of actors */
-	virtual const FSelectedActorInfo& GetSelectedActorInfo() const OVERRIDE
+	virtual const FSelectedActorInfo& GetSelectedActorInfo() const override
 	{
 		return SelectedActorInfo;
 	}
 
-	virtual bool HasClassDefaultObject() const OVERRIDE
+	virtual bool HasClassDefaultObject() const override
 	{
 		return bViewingClassDefaultObject;
 	}
 
 	/** Gets the base class being viewed */
-	const UClass* GetBaseClass() const OVERRIDE;
-	UClass* GetBaseClass() OVERRIDE;
+	const UClass* GetBaseClass() const override;
+	UClass* GetBaseClass() override;
 
 	// SWidget interface
-	virtual bool SupportsKeyboardFocus() const OVERRIDE;
-	virtual FReply OnKeyboardFocusReceived( const FGeometry& MyGeometry, const FKeyboardFocusEvent& InKeyboardFocusEvent ) OVERRIDE;
-	virtual void Tick( const FGeometry& AllottedGeometry, const double InCurrentTime, const float InDeltaTime ) OVERRIDE;
+	virtual bool SupportsKeyboardFocus() const override;
+	virtual FReply OnKeyboardFocusReceived( const FGeometry& MyGeometry, const FKeyboardFocusEvent& InKeyboardFocusEvent ) override;
+	virtual void Tick( const FGeometry& AllottedGeometry, const double InCurrentTime, const float InDeltaTime ) override;
 	// End of SWidget interface
 
 	/**
@@ -106,7 +106,7 @@ public:
 	/**
 	 * @return True if a category is hidden by any of the uobject classes currently in view by this details panel
 	 */
-	bool IsCategoryHiddenByClass(FName CategoryName) const OVERRIDE;
+	bool IsCategoryHiddenByClass(FName CategoryName) const override;
 
 private:
 	void RegisterInstancedCustomPropertyLayout( UClass* Class, FOnGetDetailCustomizationInstance DetailLayoutDelegate );
@@ -155,7 +155,7 @@ private:
 	/**
 	 * Updates the details with the passed in filter                                                              
 	 */
-	virtual void UpdateFilteredDetails() OVERRIDE;
+	virtual void UpdateFilteredDetails() override;
 
 	/** Called when the filter button is clicked */
 	void OnFilterButtonClicked();
diff --git a/Engine/Source/Editor/PropertyEditor/Private/SDetailsViewBase.h b/Engine/Source/Editor/PropertyEditor/Private/SDetailsViewBase.h
--- a/Engine/Source/Editor/PropertyEditor/Private/SDetailsViewBase.h
+++ b/Engine/Source/Editor/PropertyEditor/Private/SDetailsViewBase.h
@@ -80,13 +80,13 @@ public:
 	/**
 	* @return true of the details view can be updated from editor selection
 	*/
-	virtual bool IsUpdatable() const OVERRIDE
+	virtual bool IsUpdatable() const override
 	{
 		return DetailsViewArgs.bUpdatesFromSelection;
 	}
 
 	/** @return The identifier for this details view, or NAME_None is this view is anonymous */
-	virtual FName GetIdentifier() const OVERRIDE
+	virtual FName GetIdentifier() const override
 	{
 		return DetailsViewArgs.ViewIdentifier;
 	}
@@ -94,9 +94,9 @@ public:
 	/**
 	 * Sets the visible state of the filter box/property grid area
 	 */
-	virtual void HideFilterArea(bool bHide) OVERRIDE;
+	virtual void HideFilterArea(bool bHide) override;
 
-	virtual FOnFinishedChangingProperties& OnFinishedChangingProperties() OVERRIDE
+	virtual FOnFinishedChangingProperties& OnFinishedChangingProperties() override
 	{ 
 		return OnFinishedChangingPropertiesDelegate; 
 	}
@@ -109,7 +109,7 @@ public:
 	/**
 	 * Returns true if the details view is locked and cant have its observed objects changed 
 	 */
-	virtual bool IsLocked() const OVERRIDE
+	virtual bool IsLocked() const override
 	{
 		return bIsLocked;
 	}
@@ -117,17 +117,17 @@ public:
 	/**
 	 * Sets a delegate which is called to determine whether a specific property should be visible
 	 */
-	virtual void SetIsPropertyVisibleDelegate(FIsPropertyVisible InIsPropertyVisible) OVERRIDE;
+	virtual void SetIsPropertyVisibleDelegate(FIsPropertyVisible InIsPropertyVisible) override;
 
 	/**
 	 * Sets a delegate to call to determine if the properties  editing is enabled
 	 */ 
-	virtual void SetIsPropertyEditingEnabledDelegate(FIsPropertyEditingEnabled IsPropertyEditingEnabled) OVERRIDE;
+	virtual void SetIsPropertyEditingEnabledDelegate(FIsPropertyEditingEnabled IsPropertyEditingEnabled) override;
 
 	/**
 	 * @return true if property editing is enabled (based on the FIsPropertyEditingEnabled delegate)
 	 */
-	virtual bool IsPropertyEditingEnabled() const OVERRIDE;
+	virtual bool IsPropertyEditingEnabled() const override;
 
 	/**
 	 * Requests that an item in the tree be expanded or collapsed
@@ -135,7 +135,7 @@ public:
 	 * @param TreeNode	The tree node to expand
 	 * @param bExpand	true if the item should be expanded, false otherwise
 	 */
-	void RequestItemExpanded(TSharedRef<IDetailTreeNode> TreeNode, bool bExpand) OVERRIDE;
+	void RequestItemExpanded(TSharedRef<IDetailTreeNode> TreeNode, bool bExpand) override;
 
 	/**
 	 * Sets the expansion state for a node and optionally all of its children
@@ -144,7 +144,7 @@ public:
 	 * @param bIsItemExpanded	The new expansion state
 	 * @param bRecursive		Whether or not to apply the expansion change to any children
 	 */
-	void SetNodeExpansionState(TSharedRef<IDetailTreeNode> InTreeNode, bool bIsItemExpanded, bool bRecursive) OVERRIDE;
+	void SetNodeExpansionState(TSharedRef<IDetailTreeNode> InTreeNode, bool bIsItemExpanded, bool bRecursive) override;
 
 	/**
 	 * Sets the expansion state all root nodes and optionally all of their children
@@ -157,7 +157,7 @@ public:
 	/**
 	 * Refreshes the detail's treeview
 	 */
-	void RefreshTree() OVERRIDE;
+	void RefreshTree() override;
 
 	/**
 	 * Saves the expansion state of a tree node
@@ -165,7 +165,7 @@ public:
 	 * @param NodePath	The path to the detail node to save
 	 * @param bIsExpanded	true if the node is expanded, false otherwise
 	 */
-	void SaveCustomExpansionState(const FString& NodePath, bool bIsExpanded) OVERRIDE;
+	void SaveCustomExpansionState(const FString& NodePath, bool bIsExpanded) override;
 
 	/**
 	 * Gets the saved expansion state of a tree node in this category	
@@ -173,12 +173,12 @@ public:
 	 * @param NodePath	The path to the detail node to get
 	 * @return true if the node should be expanded, false otherwise
 	 */	
-	bool GetCustomSavedExpansionState(const FString& NodePath) const OVERRIDE;
+	bool GetCustomSavedExpansionState(const FString& NodePath) const override;
 
 	/**
 	 * Called when properties have finished changing (after PostEditChange is called)
 	 */
-	virtual void NotifyFinishedChangingProperties(const FPropertyChangedEvent& PropertyChangedEvent) OVERRIDE;
+	virtual void NotifyFinishedChangingProperties(const FPropertyChangedEvent& PropertyChangedEvent) override;
 
 	/** Column width accessibility */
 	float OnGetLeftColumnWidth() const { return 1.0f - ColumnWidth; }
@@ -188,17 +188,17 @@ public:
 	/**
 	 * @return True if the property is visible
 	 */
-	bool IsPropertyVisible(const UProperty* Property) const OVERRIDE;
+	bool IsPropertyVisible(const UProperty* Property) const override;
 
 	/**
 	* Sets a delegate which is regardless of the objects being viewed to lay out generic details not specific to any object
 	*/
-	virtual void SetGenericLayoutDetailsDelegate(FOnGetDetailCustomizationInstance OnGetGenericDetails) OVERRIDE;
+	virtual void SetGenericLayoutDetailsDelegate(FOnGetDetailCustomizationInstance OnGetGenericDetails) override;
 
 	/**
 	 * Adds an action to execute next tick
 	 */
-	virtual void EnqueueDeferredAction(FSimpleDelegate& DeferredAction) OVERRIDE;
+	virtual void EnqueueDeferredAction(FSimpleDelegate& DeferredAction) override;
 
 	/**
 	 * @return The thumbnail pool that should be used for thumbnails being rendered in this view
@@ -208,10 +208,10 @@ public:
 	/**
 	 * Returns the property utilities for this view
 	 */
-	TSharedPtr<IPropertyUtilities> GetPropertyUtilities() OVERRIDE;
+	TSharedPtr<IPropertyUtilities> GetPropertyUtilities() override;
 
 	/** Returns the notify hook to use when properties change */
-	virtual FNotifyHook* GetNotifyHook() const OVERRIDE
+	virtual FNotifyHook* GetNotifyHook() const override
 	{ 
 		return DetailsViewArgs.NotifyHook; 
 	}
diff --git a/Engine/Source/Editor/PropertyEditor/Private/SPropertyTreeViewImpl.cpp b/Engine/Source/Editor/PropertyEditor/Private/SPropertyTreeViewImpl.cpp
--- a/Engine/Source/Editor/PropertyEditor/Private/SPropertyTreeViewImpl.cpp
+++ b/Engine/Source/Editor/PropertyEditor/Private/SPropertyTreeViewImpl.cpp
@@ -28,47 +28,47 @@ public:
 	{
 	}
 
-	virtual class FNotifyHook* GetNotifyHook() const OVERRIDE
+	virtual class FNotifyHook* GetNotifyHook() const override
 	{
 		return View.GetNotifyHook();
 	}
 
-	virtual bool AreFavoritesEnabled() const OVERRIDE
+	virtual bool AreFavoritesEnabled() const override
 	{
 		return View.AreFavoritesEnabled();
 	}
 
-	virtual void ToggleFavorite( const TSharedRef< class FPropertyEditor >& PropertyEditor ) const OVERRIDE
+	virtual void ToggleFavorite( const TSharedRef< class FPropertyEditor >& PropertyEditor ) const override
 	{
 		View.ToggleFavorite( PropertyEditor );
 	}
 
-	virtual void CreateColorPickerWindow( const TSharedRef< class FPropertyEditor >& PropertyEditor, bool bUseAlpha ) const OVERRIDE
+	virtual void CreateColorPickerWindow( const TSharedRef< class FPropertyEditor >& PropertyEditor, bool bUseAlpha ) const override
 	{
 		View.CreateColorPickerWindow( PropertyEditor, bUseAlpha );
 	}
 
-	virtual void EnqueueDeferredAction( FSimpleDelegate DeferredAction ) OVERRIDE
+	virtual void EnqueueDeferredAction( FSimpleDelegate DeferredAction ) override
 	{
 		View.EnqueueDeferredAction( DeferredAction );
 	}
 
-	virtual void RequestRefresh() OVERRIDE
+	virtual void RequestRefresh() override
 	{
 		View.RequestRefresh();
 	}
 
-	virtual bool IsPropertyEditingEnabled() const OVERRIDE
+	virtual bool IsPropertyEditingEnabled() const override
 	{
 		return true;
 	}
 
-	virtual TSharedPtr<class FAssetThumbnailPool> GetThumbnailPool() const OVERRIDE
+	virtual TSharedPtr<class FAssetThumbnailPool> GetThumbnailPool() const override
 	{
 		return NULL;
 	}
 
-	virtual void NotifyFinishedChangingProperties(const FPropertyChangedEvent& PropertyChangedEvent) OVERRIDE {}
+	virtual void NotifyFinishedChangingProperties(const FPropertyChangedEvent& PropertyChangedEvent) override {}
 private:
 
 	SPropertyTreeViewImpl& View;
diff --git a/Engine/Source/Editor/PropertyEditor/Private/SSingleProperty.cpp b/Engine/Source/Editor/PropertyEditor/Private/SSingleProperty.cpp
--- a/Engine/Source/Editor/PropertyEditor/Private/SSingleProperty.cpp
+++ b/Engine/Source/Editor/PropertyEditor/Private/SSingleProperty.cpp
@@ -20,13 +20,13 @@ public:
 	{
 	}
 
-	virtual class FNotifyHook* GetNotifyHook() const OVERRIDE
+	virtual class FNotifyHook* GetNotifyHook() const override
 	{
 		TSharedPtr< SSingleProperty > PinnedView = View.Pin();
 		return PinnedView->GetNotifyHook();
 	}
 
-	virtual void CreateColorPickerWindow( const TSharedRef< class FPropertyEditor >& PropertyEditor, bool bUseAlpha ) const OVERRIDE
+	virtual void CreateColorPickerWindow( const TSharedRef< class FPropertyEditor >& PropertyEditor, bool bUseAlpha ) const override
 	{
 		TSharedPtr< SSingleProperty > PinnedView = View.Pin();
 
@@ -41,25 +41,25 @@ public:
 		 // not implemented
 	}
 
-	virtual bool AreFavoritesEnabled() const OVERRIDE
+	virtual bool AreFavoritesEnabled() const override
 	{
 		// not implemented
 		return false;
 	}
 
-	virtual void ToggleFavorite( const TSharedRef< class FPropertyEditor >& PropertyEditor ) const OVERRIDE
+	virtual void ToggleFavorite( const TSharedRef< class FPropertyEditor >& PropertyEditor ) const override
 	{
 		// not implemented
 	}
 
-	virtual bool IsPropertyEditingEnabled() const OVERRIDE
+	virtual bool IsPropertyEditingEnabled() const override
 	{
 		return true;
 	}
 
-	virtual void RequestRefresh() OVERRIDE {}
+	virtual void RequestRefresh() override {}
 
-	virtual TSharedPtr<class FAssetThumbnailPool> GetThumbnailPool() const OVERRIDE
+	virtual TSharedPtr<class FAssetThumbnailPool> GetThumbnailPool() const override
 	{
 		// not implemented
 		return NULL;
diff --git a/Engine/Source/Editor/PropertyEditor/Private/SSingleProperty.h b/Engine/Source/Editor/PropertyEditor/Private/SSingleProperty.h
--- a/Engine/Source/Editor/PropertyEditor/Private/SSingleProperty.h
+++ b/Engine/Source/Editor/PropertyEditor/Private/SSingleProperty.h
@@ -23,9 +23,9 @@ public:
 	void Construct( const FArguments& InArgs );
 
 	/** ISinglePropertyView interface */
-	virtual bool HasValidProperty() const OVERRIDE { return RootPropertyNode.IsValid() && ValueNode.IsValid(); }
-	virtual void SetObject( UObject* InObject ) OVERRIDE;
-	virtual void SetOnPropertyValueChanged( FSimpleDelegate& InOnPropertyValueChanged ) OVERRIDE;
+	virtual bool HasValidProperty() const override { return RootPropertyNode.IsValid() && ValueNode.IsValid(); }
+	virtual void SetObject( UObject* InObject ) override;
+	virtual void SetOnPropertyValueChanged( FSimpleDelegate& InOnPropertyValueChanged ) override;
 
 	/**
 	 * Replaces objects being observed by the view with new objects
diff --git a/Engine/Source/Editor/PropertyEditor/Private/UserInterface/PropertyDetails/PropertyDetailsUtilities.h b/Engine/Source/Editor/PropertyEditor/Private/UserInterface/PropertyDetails/PropertyDetailsUtilities.h
--- a/Engine/Source/Editor/PropertyEditor/Private/UserInterface/PropertyDetails/PropertyDetailsUtilities.h
+++ b/Engine/Source/Editor/PropertyEditor/Private/UserInterface/PropertyDetails/PropertyDetailsUtilities.h
@@ -8,15 +8,15 @@ class FPropertyDetailsUtilities : public IPropertyUtilities
 public:
 	FPropertyDetailsUtilities(IDetailsViewPrivate& InDetailsView);
 	/** IPropertyUtilities interface */
-	virtual class FNotifyHook* GetNotifyHook() const OVERRIDE;
-	virtual bool AreFavoritesEnabled() const OVERRIDE;
-	virtual void ToggleFavorite( const TSharedRef< class FPropertyEditor >& PropertyEditor ) const OVERRIDE;
-	virtual void CreateColorPickerWindow( const TSharedRef< class FPropertyEditor >& PropertyEditor, bool bUseAlpha ) const OVERRIDE;
-	virtual void EnqueueDeferredAction( FSimpleDelegate DeferredAction ) OVERRIDE;
-	virtual bool IsPropertyEditingEnabled() const OVERRIDE;
-	virtual void RequestRefresh() OVERRIDE;
-	virtual TSharedPtr<class FAssetThumbnailPool> GetThumbnailPool() const OVERRIDE;
-	virtual void NotifyFinishedChangingProperties(const FPropertyChangedEvent& PropertyChangedEvent) OVERRIDE;
+	virtual class FNotifyHook* GetNotifyHook() const override;
+	virtual bool AreFavoritesEnabled() const override;
+	virtual void ToggleFavorite( const TSharedRef< class FPropertyEditor >& PropertyEditor ) const override;
+	virtual void CreateColorPickerWindow( const TSharedRef< class FPropertyEditor >& PropertyEditor, bool bUseAlpha ) const override;
+	virtual void EnqueueDeferredAction( FSimpleDelegate DeferredAction ) override;
+	virtual bool IsPropertyEditingEnabled() const override;
+	virtual void RequestRefresh() override;
+	virtual TSharedPtr<class FAssetThumbnailPool> GetThumbnailPool() const override;
+	virtual void NotifyFinishedChangingProperties(const FPropertyChangedEvent& PropertyChangedEvent) override;
 private:
 	IDetailsViewPrivate& DetailsView;
 };
diff --git a/Engine/Source/Editor/PropertyEditor/Private/UserInterface/PropertyEditor/SPropertyEditorAsset.h b/Engine/Source/Editor/PropertyEditor/Private/UserInterface/PropertyEditor/SPropertyEditorAsset.h
--- a/Engine/Source/Editor/PropertyEditor/Private/UserInterface/PropertyEditor/SPropertyEditorAsset.h
+++ b/Engine/Source/Editor/PropertyEditor/Private/UserInterface/PropertyEditor/SPropertyEditorAsset.h
@@ -56,7 +56,7 @@ public:
 	void GetDesiredWidth( float& OutMinDesiredWidth, float &OutMaxDesiredWidth );
 
 	/** Override the tick method so we can check for thumbnail differences & update if necessary */
-	void Tick( const FGeometry& AllottedGeometry, const double InCurrentTime, const float InDeltaTime ) OVERRIDE;
+	void Tick( const FGeometry& AllottedGeometry, const double InCurrentTime, const float InDeltaTime ) override;
 
 private:
 	struct FObjectOrAssetData
diff --git a/Engine/Source/Editor/PropertyEditor/Private/UserInterface/PropertyEditor/SPropertyEditorBool.h b/Engine/Source/Editor/PropertyEditor/Private/UserInterface/PropertyEditor/SPropertyEditorBool.h
--- a/Engine/Source/Editor/PropertyEditor/Private/UserInterface/PropertyEditor/SPropertyEditorBool.h
+++ b/Engine/Source/Editor/PropertyEditor/Private/UserInterface/PropertyEditor/SPropertyEditorBool.h
@@ -15,9 +15,9 @@ public:
 	void GetDesiredWidth( float& OutMinDesiredWidth, float& OutMaxDesiredWidth );
 
 	// SWidget overrides
-	virtual bool SupportsKeyboardFocus() const OVERRIDE;
-	virtual bool HasKeyboardFocus() const OVERRIDE;
-	virtual FReply OnKeyboardFocusReceived( const FGeometry& MyGeometry, const FKeyboardFocusEvent& InKeyboardFocusEvent ) OVERRIDE;
+	virtual bool SupportsKeyboardFocus() const override;
+	virtual bool HasKeyboardFocus() const override;
+	virtual FReply OnKeyboardFocusReceived( const FGeometry& MyGeometry, const FKeyboardFocusEvent& InKeyboardFocusEvent ) override;
 
 private:
 
diff --git a/Engine/Source/Editor/PropertyEditor/Private/UserInterface/PropertyEditor/SPropertyEditorClass.cpp b/Engine/Source/Editor/PropertyEditor/Private/UserInterface/PropertyEditor/SPropertyEditorClass.cpp
--- a/Engine/Source/Editor/PropertyEditor/Private/UserInterface/PropertyEditor/SPropertyEditorClass.cpp
+++ b/Engine/Source/Editor/PropertyEditor/Private/UserInterface/PropertyEditor/SPropertyEditorClass.cpp
@@ -41,7 +41,7 @@ public:
 		return false;
 	}
 
-	virtual bool IsUnloadedClassAllowed(const FClassViewerInitializationOptions& InInitOptions, const TSharedRef< const IUnloadedBlueprintData > InClass, TSharedRef< FClassViewerFilterFuncs > InFilterFuncs) OVERRIDE
+	virtual bool IsUnloadedClassAllowed(const FClassViewerInitializationOptions& InInitOptions, const TSharedRef< const IUnloadedBlueprintData > InClass, TSharedRef< FClassViewerFilterFuncs > InFilterFuncs) override
 	{
 		bool bMatchesFlags = !InClass->HasAnyClassFlags(CLASS_Hidden|CLASS_HideDropDown|CLASS_Deprecated) &&
 			(bAllowAbstract || !InClass->HasAnyClassFlags(CLASS_Abstract));
diff --git a/Engine/Source/Editor/PropertyEditor/Private/UserInterface/PropertyEditor/SPropertyEditorClass.h b/Engine/Source/Editor/PropertyEditor/Private/UserInterface/PropertyEditor/SPropertyEditorClass.h
--- a/Engine/Source/Editor/PropertyEditor/Private/UserInterface/PropertyEditor/SPropertyEditorClass.h
+++ b/Engine/Source/Editor/PropertyEditor/Private/UserInterface/PropertyEditor/SPropertyEditorClass.h
@@ -47,7 +47,7 @@ public:
 private:
 	void SendToObjects(const FString& NewValue);
 
-	virtual FReply OnDrop(const FGeometry& MyGeometry, const FDragDropEvent& DragDropEvent) OVERRIDE;
+	virtual FReply OnDrop(const FGeometry& MyGeometry, const FDragDropEvent& DragDropEvent) override;
 
 	/** 
 	 * Generates a class picker with a filter to show only classes allowed to be selected. 
diff --git a/Engine/Source/Editor/PropertyEditor/Private/UserInterface/PropertyEditor/SPropertyEditorEditInline.cpp b/Engine/Source/Editor/PropertyEditor/Private/UserInterface/PropertyEditor/SPropertyEditorEditInline.cpp
--- a/Engine/Source/Editor/PropertyEditor/Private/UserInterface/PropertyEditor/SPropertyEditorEditInline.cpp
+++ b/Engine/Source/Editor/PropertyEditor/Private/UserInterface/PropertyEditor/SPropertyEditorEditInline.cpp
@@ -44,7 +44,7 @@ public:
 		return false;
 	}
 
-	virtual bool IsUnloadedClassAllowed(const FClassViewerInitializationOptions& InInitOptions, const TSharedRef< const IUnloadedBlueprintData > InUnloadedClassData, TSharedRef< FClassViewerFilterFuncs > InFilterFuncs) OVERRIDE
+	virtual bool IsUnloadedClassAllowed(const FClassViewerInitializationOptions& InInitOptions, const TSharedRef< const IUnloadedBlueprintData > InUnloadedClassData, TSharedRef< FClassViewerFilterFuncs > InFilterFuncs) override
 	{
 		const bool bChildOfObjectClass = InUnloadedClassData->IsChildOf(ObjProperty->PropertyClass);
 
diff --git a/Engine/Source/Editor/PropertyEditor/Private/UserInterface/PropertyEditor/SPropertyEditorInteractiveActorPicker.h b/Engine/Source/Editor/PropertyEditor/Private/UserInterface/PropertyEditor/SPropertyEditorInteractiveActorPicker.h
--- a/Engine/Source/Editor/PropertyEditor/Private/UserInterface/PropertyEditor/SPropertyEditorInteractiveActorPicker.h
+++ b/Engine/Source/Editor/PropertyEditor/Private/UserInterface/PropertyEditor/SPropertyEditorInteractiveActorPicker.h
@@ -22,8 +22,8 @@ public:
 	void Construct( const FArguments& InArgs );
 
 	/** Begin SWidget interface */
-	virtual FReply OnKeyDown( const FGeometry& MyGeometry, const FKeyboardEvent& InKeyboardEvent ) OVERRIDE;
-	virtual bool SupportsKeyboardFocus() const OVERRIDE;
+	virtual FReply OnKeyDown( const FGeometry& MyGeometry, const FKeyboardEvent& InKeyboardEvent ) override;
+	virtual bool SupportsKeyboardFocus() const override;
 	/** End SWidget interface */
 
 private:
diff --git a/Engine/Source/Editor/PropertyEditor/Private/UserInterface/PropertyEditor/SPropertyEditorNumeric.h b/Engine/Source/Editor/PropertyEditor/Private/UserInterface/PropertyEditor/SPropertyEditorNumeric.h
--- a/Engine/Source/Editor/PropertyEditor/Private/UserInterface/PropertyEditor/SPropertyEditorNumeric.h
+++ b/Engine/Source/Editor/PropertyEditor/Private/UserInterface/PropertyEditor/SPropertyEditorNumeric.h
@@ -99,12 +99,12 @@ public:
 		SetEnabled( TAttribute<bool>( this, &SPropertyEditorNumeric<NumericType>::CanEdit ) );
 	}
 
-	virtual bool SupportsKeyboardFocus() const OVERRIDE
+	virtual bool SupportsKeyboardFocus() const override
 	{
 		return PrimaryWidget.IsValid() && PrimaryWidget->SupportsKeyboardFocus();
 	}
 
-	virtual FReply OnKeyboardFocusReceived( const FGeometry& MyGeometry, const FKeyboardFocusEvent& InKeyboardFocusEvent ) OVERRIDE
+	virtual FReply OnKeyboardFocusReceived( const FGeometry& MyGeometry, const FKeyboardFocusEvent& InKeyboardFocusEvent ) override
 	{
 		// Forward keyboard focus to our editable text widget
 		return FReply::Handled().SetKeyboardFocus( PrimaryWidget.ToSharedRef(), InKeyboardFocusEvent.GetCause() );
diff --git a/Engine/Source/Editor/PropertyEditor/Private/UserInterface/PropertyEditor/SPropertyEditorTableRow.h b/Engine/Source/Editor/PropertyEditor/Private/UserInterface/PropertyEditor/SPropertyEditorTableRow.h
--- a/Engine/Source/Editor/PropertyEditor/Private/UserInterface/PropertyEditor/SPropertyEditorTableRow.h
+++ b/Engine/Source/Editor/PropertyEditor/Private/UserInterface/PropertyEditor/SPropertyEditorTableRow.h
@@ -27,13 +27,13 @@ public:
 	 * 
 	 * @param ColumnName The name of the column to generate a widget for
 	 */
-	virtual TSharedRef<SWidget> GenerateWidgetForColumn( const FName& ColumnName ) OVERRIDE;
+	virtual TSharedRef<SWidget> GenerateWidgetForColumn( const FName& ColumnName ) override;
 	
-	virtual FReply OnMouseButtonUp( const FGeometry& MyGeometry, const FPointerEvent& MouseEvent ) OVERRIDE;
+	virtual FReply OnMouseButtonUp( const FGeometry& MyGeometry, const FPointerEvent& MouseEvent ) override;
 
-	virtual TSharedPtr< FPropertyPath > GetPropertyPath() const OVERRIDE;
+	virtual TSharedPtr< FPropertyPath > GetPropertyPath() const override;
 
-	virtual bool IsCursorHovering() const OVERRIDE;
+	virtual bool IsCursorHovering() const override;
 
 
 private:
diff --git a/Engine/Source/Editor/PropertyEditor/Private/UserInterface/PropertyEditor/SPropertyEditorText.h b/Engine/Source/Editor/PropertyEditor/Private/UserInterface/PropertyEditor/SPropertyEditorText.h
--- a/Engine/Source/Editor/PropertyEditor/Private/UserInterface/PropertyEditor/SPropertyEditorText.h
+++ b/Engine/Source/Editor/PropertyEditor/Private/UserInterface/PropertyEditor/SPropertyEditorText.h
@@ -19,9 +19,9 @@ public:
 
 	void GetDesiredWidth( float& OutMinDesiredWidth, float& OutMaxDesiredWidth );
 
-	bool SupportsKeyboardFocus() const OVERRIDE;
+	bool SupportsKeyboardFocus() const override;
 
-	FReply OnKeyboardFocusReceived( const FGeometry& MyGeometry, const FKeyboardFocusEvent& InKeyboardFocusEvent ) OVERRIDE;
+	FReply OnKeyboardFocusReceived( const FGeometry& MyGeometry, const FKeyboardFocusEvent& InKeyboardFocusEvent ) override;
 
 private:
 
diff --git a/Engine/Source/Editor/PropertyEditor/Private/UserInterface/PropertyTable/BooleanPropertyTableCellPresenter.h b/Engine/Source/Editor/PropertyEditor/Private/UserInterface/PropertyTable/BooleanPropertyTableCellPresenter.h
--- a/Engine/Source/Editor/PropertyEditor/Private/UserInterface/PropertyTable/BooleanPropertyTableCellPresenter.h
+++ b/Engine/Source/Editor/PropertyEditor/Private/UserInterface/PropertyTable/BooleanPropertyTableCellPresenter.h
@@ -9,21 +9,21 @@ public:
 
 	FBooleanPropertyTableCellPresenter( const TSharedRef< class FPropertyEditor >& InPropertyEditor );
 
-	virtual TSharedRef< class SWidget > ConstructDisplayWidget() OVERRIDE;
+	virtual TSharedRef< class SWidget > ConstructDisplayWidget() override;
 
-	virtual bool RequiresDropDown() OVERRIDE;
+	virtual bool RequiresDropDown() override;
 
-	virtual TSharedRef< class SWidget > ConstructEditModeCellWidget() OVERRIDE;
+	virtual TSharedRef< class SWidget > ConstructEditModeCellWidget() override;
 
-	virtual TSharedRef< class SWidget > ConstructEditModeDropDownWidget() OVERRIDE;
+	virtual TSharedRef< class SWidget > ConstructEditModeDropDownWidget() override;
 
-	virtual TSharedRef< class SWidget > WidgetToFocusOnEdit() OVERRIDE;
+	virtual TSharedRef< class SWidget > WidgetToFocusOnEdit() override;
 
-	virtual FString GetValueAsString() OVERRIDE;
+	virtual FString GetValueAsString() override;
 
-	virtual FText GetValueAsText() OVERRIDE;
+	virtual FText GetValueAsText() override;
 
-	virtual bool HasReadOnlyEditMode() OVERRIDE { return false; }
+	virtual bool HasReadOnlyEditMode() override { return false; }
 
 
 private:
diff --git a/Engine/Source/Editor/PropertyEditor/Private/UserInterface/PropertyTable/ColorPropertyTableCellPresenter.h b/Engine/Source/Editor/PropertyEditor/Private/UserInterface/PropertyTable/ColorPropertyTableCellPresenter.h
--- a/Engine/Source/Editor/PropertyEditor/Private/UserInterface/PropertyTable/ColorPropertyTableCellPresenter.h
+++ b/Engine/Source/Editor/PropertyEditor/Private/UserInterface/PropertyTable/ColorPropertyTableCellPresenter.h
@@ -9,21 +9,21 @@ public:
 
 	FColorPropertyTableCellPresenter( const TSharedRef< class FPropertyEditor >& InPropertyEditor, const TSharedRef< class IPropertyTableUtilities >& InPropertyUtilities );
 
-	virtual TSharedRef< class SWidget > ConstructDisplayWidget() OVERRIDE;
+	virtual TSharedRef< class SWidget > ConstructDisplayWidget() override;
 
-	virtual bool RequiresDropDown() OVERRIDE;
+	virtual bool RequiresDropDown() override;
 
-	virtual TSharedRef< class SWidget > ConstructEditModeCellWidget() OVERRIDE;
+	virtual TSharedRef< class SWidget > ConstructEditModeCellWidget() override;
 
-	virtual TSharedRef< class SWidget > ConstructEditModeDropDownWidget() OVERRIDE;
+	virtual TSharedRef< class SWidget > ConstructEditModeDropDownWidget() override;
 
-	virtual TSharedRef< class SWidget > WidgetToFocusOnEdit() OVERRIDE;
+	virtual TSharedRef< class SWidget > WidgetToFocusOnEdit() override;
 
-	virtual FString GetValueAsString() OVERRIDE;
+	virtual FString GetValueAsString() override;
 
-	virtual FText GetValueAsText() OVERRIDE;
+	virtual FText GetValueAsText() override;
 
-	virtual bool HasReadOnlyEditMode() OVERRIDE { return false; }
+	virtual bool HasReadOnlyEditMode() override { return false; }
 
 
 private:
diff --git a/Engine/Source/Editor/PropertyEditor/Private/UserInterface/PropertyTable/ObjectNameTableCellPresenter.h b/Engine/Source/Editor/PropertyEditor/Private/UserInterface/PropertyTable/ObjectNameTableCellPresenter.h
--- a/Engine/Source/Editor/PropertyEditor/Private/UserInterface/PropertyTable/ObjectNameTableCellPresenter.h
+++ b/Engine/Source/Editor/PropertyEditor/Private/UserInterface/PropertyTable/ObjectNameTableCellPresenter.h
@@ -16,7 +16,7 @@ public:
 
 	}
 
-	virtual TSharedRef< class SWidget > ConstructDisplayWidget() OVERRIDE
+	virtual TSharedRef< class SWidget > ConstructDisplayWidget() override
 	{
 		return SNew( SBox )
 			.VAlign( VAlign_Center)
@@ -30,9 +30,9 @@ public:
 			);
 	}
 
-	virtual bool RequiresDropDown() OVERRIDE { return false; }
+	virtual bool RequiresDropDown() override { return false; }
 
-	virtual TSharedRef< class SWidget > ConstructEditModeCellWidget() OVERRIDE 
+	virtual TSharedRef< class SWidget > ConstructEditModeCellWidget() override 
 	{
 		TSharedPtr< SEditableTextBox > NewFocusWidget;
 		TSharedRef< class SWidget > Result = 
@@ -51,15 +51,15 @@ public:
 		return Result;
 	}
 
-	virtual TSharedRef< class SWidget > ConstructEditModeDropDownWidget() OVERRIDE { return SNullWidget::NullWidget; }
+	virtual TSharedRef< class SWidget > ConstructEditModeDropDownWidget() override { return SNullWidget::NullWidget; }
 
-	virtual TSharedRef< class SWidget > WidgetToFocusOnEdit() OVERRIDE { return FocusWidget.Pin().ToSharedRef(); }
+	virtual TSharedRef< class SWidget > WidgetToFocusOnEdit() override { return FocusWidget.Pin().ToSharedRef(); }
 
-	virtual FString GetValueAsString() OVERRIDE { return Cell->GetValueAsString(); }
+	virtual FString GetValueAsString() override { return Cell->GetValueAsString(); }
 
-	virtual FText GetValueAsText() OVERRIDE { return Cell->GetValueAsText(); }
+	virtual FText GetValueAsText() override { return Cell->GetValueAsText(); }
 
-	virtual bool HasReadOnlyEditMode() OVERRIDE { return true; }
+	virtual bool HasReadOnlyEditMode() override { return true; }
 
 
 private:
diff --git a/Engine/Source/Editor/PropertyEditor/Private/UserInterface/PropertyTable/PropertyTableWidgetHandle.h b/Engine/Source/Editor/PropertyEditor/Private/UserInterface/PropertyTable/PropertyTableWidgetHandle.h
--- a/Engine/Source/Editor/PropertyEditor/Private/UserInterface/PropertyTable/PropertyTableWidgetHandle.h
+++ b/Engine/Source/Editor/PropertyEditor/Private/UserInterface/PropertyTable/PropertyTableWidgetHandle.h
@@ -9,12 +9,12 @@ public:
 	FPropertyTableWidgetHandle(const TSharedRef<SPropertyTable>& InPropertyTableWidget)
 		: PropertyTableWidget(InPropertyTableWidget) {}
 
-	virtual void RequestRefresh() OVERRIDE
+	virtual void RequestRefresh() override
 	{
 		PropertyTableWidget->RequestRefresh();
 	} 
 
-	virtual TSharedRef<SWidget> GetWidget() OVERRIDE
+	virtual TSharedRef<SWidget> GetWidget() override
 	{
 		return PropertyTableWidget;
 	}
diff --git a/Engine/Source/Editor/PropertyEditor/Private/UserInterface/PropertyTable/SBoolColumnHeader.h b/Engine/Source/Editor/PropertyEditor/Private/UserInterface/PropertyTable/SBoolColumnHeader.h
--- a/Engine/Source/Editor/PropertyEditor/Private/UserInterface/PropertyTable/SBoolColumnHeader.h
+++ b/Engine/Source/Editor/PropertyEditor/Private/UserInterface/PropertyTable/SBoolColumnHeader.h
@@ -28,7 +28,7 @@ class SBoolColumnHeader : public SColumnHeader
 		SColumnHeader::Construct( ColumnArgs, InPropertyTableColumn, InPropertyUtilities );
 	}
 
-	virtual TSharedRef< SWidget > GenerateCell( const TSharedRef< class IPropertyTableRow >& PropertyTableRow ) OVERRIDE
+	virtual TSharedRef< SWidget > GenerateCell( const TSharedRef< class IPropertyTableRow >& PropertyTableRow ) override
 	{
 		TSharedRef< IPropertyTableCell > Cell = Column->GetCell( PropertyTableRow );
 
diff --git a/Engine/Source/Editor/PropertyEditor/Private/UserInterface/PropertyTable/SColorColumnHeader.h b/Engine/Source/Editor/PropertyEditor/Private/UserInterface/PropertyTable/SColorColumnHeader.h
--- a/Engine/Source/Editor/PropertyEditor/Private/UserInterface/PropertyTable/SColorColumnHeader.h
+++ b/Engine/Source/Editor/PropertyEditor/Private/UserInterface/PropertyTable/SColorColumnHeader.h
@@ -28,7 +28,7 @@ class SColorColumnHeader : public SColumnHeader
 		SColumnHeader::Construct( ColumnArgs, InPropertyTableColumn, InPropertyUtilities );
 	}
 
-	virtual TSharedRef< SWidget > GenerateCell( const TSharedRef< class IPropertyTableRow >& PropertyTableRow ) OVERRIDE
+	virtual TSharedRef< SWidget > GenerateCell( const TSharedRef< class IPropertyTableRow >& PropertyTableRow ) override
 	{
 		TSharedRef< IPropertyTableCell > Cell = Column->GetCell( PropertyTableRow );
 
diff --git a/Engine/Source/Editor/PropertyEditor/Private/UserInterface/PropertyTable/SColumnHeader.h b/Engine/Source/Editor/PropertyEditor/Private/UserInterface/PropertyTable/SColumnHeader.h
--- a/Engine/Source/Editor/PropertyEditor/Private/UserInterface/PropertyTable/SColumnHeader.h
+++ b/Engine/Source/Editor/PropertyEditor/Private/UserInterface/PropertyTable/SColumnHeader.h
@@ -70,7 +70,7 @@ protected:
 		this->SetToolTip( ToolTip );
 	}
 
-	virtual FReply OnMouseButtonDown( const FGeometry& MyGeometry, const FPointerEvent& MouseEvent ) OVERRIDE
+	virtual FReply OnMouseButtonDown( const FGeometry& MyGeometry, const FPointerEvent& MouseEvent ) override
 	{
 		if ( !Column->IsFrozen() && MouseEvent.IsMouseButtonDown( EKeys::MiddleMouseButton ) )
 		{
diff --git a/Engine/Source/Editor/PropertyEditor/Private/UserInterface/PropertyTable/SPropertyTable.h b/Engine/Source/Editor/PropertyEditor/Private/UserInterface/PropertyTable/SPropertyTable.h
--- a/Engine/Source/Editor/PropertyEditor/Private/UserInterface/PropertyTable/SPropertyTable.h
+++ b/Engine/Source/Editor/PropertyEditor/Private/UserInterface/PropertyTable/SPropertyTable.h
@@ -106,24 +106,24 @@ public:
 		}
 	}
 
-	void Tick( const FGeometry& AllottedGeometry, const double InCurrentTime, const float InDeltaTime ) OVERRIDE
+	void Tick( const FGeometry& AllottedGeometry, const double InCurrentTime, const float InDeltaTime ) override
 	{
 		Table->Tick();
 		STreeView< TSharedRef< IPropertyTableRow > >::Tick( AllottedGeometry, InCurrentTime, InDeltaTime  );
 	}
 
-	virtual void RequestRefresh() OVERRIDE
+	virtual void RequestRefresh() override
 	{
 		WidgetGenerator.Clear();
 		RequestTreeRefresh();
 	}
 
-	virtual TSharedRef<SWidget> GetWidget() OVERRIDE
+	virtual TSharedRef<SWidget> GetWidget() override
 	{
 		return SharedThis(this);
 	}
 
-	virtual FReply OnKeyDown( const FGeometry& MyGeometry, const FKeyboardEvent& InKeyboardEvent ) OVERRIDE
+	virtual FReply OnKeyDown( const FGeometry& MyGeometry, const FKeyboardEvent& InKeyboardEvent ) override
 	{
 		FReply Reply = FReply::Unhandled();
 		const ESelectionMode::Type SelectionMode = Table->GetSelectionMode();
@@ -454,7 +454,7 @@ public:
 		return STreeView< TSharedRef< IPropertyTableRow > >::OnKeyDown( MyGeometry, InKeyboardEvent );
 	}
 
-	virtual FReply OnPreviewKeyDown( const FGeometry& MyGeometry, const FKeyboardEvent& InKeyboardEvent ) OVERRIDE
+	virtual FReply OnPreviewKeyDown( const FGeometry& MyGeometry, const FKeyboardEvent& InKeyboardEvent ) override
 	{
 		const FKey Key = InKeyboardEvent.GetKey();
 		if ( Key == EKeys::Enter )
@@ -542,7 +542,7 @@ public:
 		return FReply::Handled();
 	}
 
-	virtual FReply OnKeyUp( const FGeometry& MyGeometry, const FKeyboardEvent& InKeyboardEvent ) OVERRIDE
+	virtual FReply OnKeyUp( const FGeometry& MyGeometry, const FKeyboardEvent& InKeyboardEvent ) override
 	{
 		FReply Reply = FReply::Unhandled();
 
@@ -741,12 +741,12 @@ public:
 	// These handle selection, expansion, and other such properties common to ItemsWidgets.
 	//
 
-	virtual bool Private_UsesSelectorFocus() const OVERRIDE
+	virtual bool Private_UsesSelectorFocus() const override
 	{
 		return false;
 	}
 
-	virtual bool Private_HasSelectorFocus( const TSharedRef< IPropertyTableRow >& TheItem ) const OVERRIDE
+	virtual bool Private_HasSelectorFocus( const TSharedRef< IPropertyTableRow >& TheItem ) const override
 	{
 		return false;
 	}
diff --git a/Engine/Source/Editor/PropertyEditor/Private/UserInterface/PropertyTable/SPropertyTableCell.h b/Engine/Source/Editor/PropertyEditor/Private/UserInterface/PropertyTable/SPropertyTableCell.h
--- a/Engine/Source/Editor/PropertyEditor/Private/UserInterface/PropertyTable/SPropertyTableCell.h
+++ b/Engine/Source/Editor/PropertyEditor/Private/UserInterface/PropertyTable/SPropertyTableCell.h
@@ -32,13 +32,13 @@ class SPropertyTableCell : public SCompoundWidget
 
 	void Construct( const FArguments& InArgs, const TSharedRef< class IPropertyTableCell >& InCell );
 
-	virtual FReply OnMouseButtonDown( const FGeometry& MyGeometry, const FPointerEvent& MouseEvent ) OVERRIDE;
+	virtual FReply OnMouseButtonDown( const FGeometry& MyGeometry, const FPointerEvent& MouseEvent ) override;
 
-	virtual FReply OnMouseButtonDoubleClick( const FGeometry& InMyGeometry, const FPointerEvent& InMouseEvent ) OVERRIDE;
+	virtual FReply OnMouseButtonDoubleClick( const FGeometry& InMyGeometry, const FPointerEvent& InMouseEvent ) override;
 
-	virtual void Tick( const FGeometry& AllottedGeometry, const double InCurrentTime, const float InDeltaTime ) OVERRIDE;
+	virtual void Tick( const FGeometry& AllottedGeometry, const double InCurrentTime, const float InDeltaTime ) override;
 
-	virtual int32 OnPaint( const FGeometry& AllottedGeometry, const FSlateRect& MyClippingRect, FSlateWindowElementList& OutDrawElements, int32 LayerId, const FWidgetStyle& InWidgetStyle, bool bParentEnabled ) const OVERRIDE;
+	virtual int32 OnPaint( const FGeometry& AllottedGeometry, const FSlateRect& MyClippingRect, FSlateWindowElementList& OutDrawElements, int32 LayerId, const FWidgetStyle& InWidgetStyle, bool bParentEnabled ) const override;
 
 
 private:
diff --git a/Engine/Source/Editor/PropertyEditor/Private/UserInterface/PropertyTable/SPropertyTableRow.h b/Engine/Source/Editor/PropertyEditor/Private/UserInterface/PropertyTable/SPropertyTableRow.h
--- a/Engine/Source/Editor/PropertyEditor/Private/UserInterface/PropertyTable/SPropertyTableRow.h
+++ b/Engine/Source/Editor/PropertyEditor/Private/UserInterface/PropertyTable/SPropertyTableRow.h
@@ -166,7 +166,7 @@ protected:
 		GenerateColumns( HeaderRowWeakPtr.Pin().ToSharedRef() );
 	}
 
-	virtual TSharedRef< SWidget > GenerateWidgetForColumn( const FName& ColumnId ) OVERRIDE
+	virtual TSharedRef< SWidget > GenerateWidgetForColumn( const FName& ColumnId ) override
 	{
 		TSharedPtr< SPropertyTableHeaderRow > PropertyTableHeaderRow = HeaderRowWeakPtr.Pin();
 		if ( !PropertyTableHeaderRow.IsValid() )
diff --git a/Engine/Source/Editor/PropertyEditor/Private/UserInterface/PropertyTable/SRowHeaderCell.h b/Engine/Source/Editor/PropertyEditor/Private/UserInterface/PropertyTable/SRowHeaderCell.h
--- a/Engine/Source/Editor/PropertyEditor/Private/UserInterface/PropertyTable/SRowHeaderCell.h
+++ b/Engine/Source/Editor/PropertyEditor/Private/UserInterface/PropertyTable/SRowHeaderCell.h
@@ -78,7 +78,7 @@ class SRowHeaderCell : public SCompoundWidget
 		];
 	}
 
-	virtual FReply OnMouseButtonDown( const FGeometry& MyGeometry, const FPointerEvent& MouseEvent ) OVERRIDE
+	virtual FReply OnMouseButtonDown( const FGeometry& MyGeometry, const FPointerEvent& MouseEvent ) override
 	{
 		const TSharedRef< IPropertyTable > Table = Cell->GetTable();
 		Table->SetLastClickedCell( Cell );
@@ -86,7 +86,7 @@ class SRowHeaderCell : public SCompoundWidget
 		return FReply::Unhandled();
 	}
 
-	virtual FReply OnMouseButtonDoubleClick( const FGeometry& InMyGeometry, const FPointerEvent& MouseEvent ) OVERRIDE
+	virtual FReply OnMouseButtonDoubleClick( const FGeometry& InMyGeometry, const FPointerEvent& MouseEvent ) override
 	{
 		const TSharedRef< IPropertyTable > Table = Cell->GetTable();
 		Table->SetLastClickedCell( Cell );
diff --git a/Engine/Source/Editor/PropertyEditor/Private/UserInterface/PropertyTable/TextPropertyTableCellPresenter.h b/Engine/Source/Editor/PropertyEditor/Private/UserInterface/PropertyTable/TextPropertyTableCellPresenter.h
--- a/Engine/Source/Editor/PropertyEditor/Private/UserInterface/PropertyTable/TextPropertyTableCellPresenter.h
+++ b/Engine/Source/Editor/PropertyEditor/Private/UserInterface/PropertyTable/TextPropertyTableCellPresenter.h
@@ -10,21 +10,21 @@ public:
 
 	FTextPropertyTableCellPresenter( const TSharedRef< class FPropertyEditor >& InPropertyEditor, const TSharedRef< class IPropertyTableUtilities >& InPropertyUtilities, FSlateFontInfo InFont = FEditorStyle::GetFontStyle( PropertyTableConstants::NormalFontStyle ) );
 
-	virtual TSharedRef< class SWidget > ConstructDisplayWidget() OVERRIDE;
+	virtual TSharedRef< class SWidget > ConstructDisplayWidget() override;
 
-	virtual bool RequiresDropDown() OVERRIDE;
+	virtual bool RequiresDropDown() override;
 
-	virtual TSharedRef< class SWidget > ConstructEditModeCellWidget() OVERRIDE;
+	virtual TSharedRef< class SWidget > ConstructEditModeCellWidget() override;
 
-	virtual TSharedRef< class SWidget > ConstructEditModeDropDownWidget() OVERRIDE;
+	virtual TSharedRef< class SWidget > ConstructEditModeDropDownWidget() override;
 
-	virtual TSharedRef< class SWidget > WidgetToFocusOnEdit() OVERRIDE;
+	virtual TSharedRef< class SWidget > WidgetToFocusOnEdit() override;
 
-	virtual FString GetValueAsString() OVERRIDE;
+	virtual FString GetValueAsString() override;
 
-	virtual FText GetValueAsText() OVERRIDE;
+	virtual FText GetValueAsText() override;
 
-	virtual bool HasReadOnlyEditMode() OVERRIDE { return HasReadOnlyEditingWidget; }
+	virtual bool HasReadOnlyEditMode() override { return HasReadOnlyEditingWidget; }
 
 
 private:
diff --git a/Engine/Source/Editor/PropertyEditor/Public/PropertyCustomizationHelpers.h b/Engine/Source/Editor/PropertyEditor/Public/PropertyCustomizationHelpers.h
--- a/Engine/Source/Editor/PropertyEditor/Public/PropertyCustomizationHelpers.h
+++ b/Engine/Source/Editor/PropertyEditor/Public/PropertyCustomizationHelpers.h
@@ -213,18 +213,18 @@ public:
 		OnGenerateArrayElementWidgetDelegate = InOnGenerateArrayElementWidget;
 	}
 
-	virtual bool RequiresTick() const OVERRIDE { return false; }
+	virtual bool RequiresTick() const override { return false; }
 
-	virtual void Tick( float DeltaTime ) OVERRIDE {}
+	virtual void Tick( float DeltaTime ) override {}
 
-	virtual FName GetName() const OVERRIDE
+	virtual FName GetName() const override
 	{
 		return BaseProperty->GetProperty()->GetFName();
 	}
 	
-	virtual bool InitiallyCollapsed() const OVERRIDE { return false; }
+	virtual bool InitiallyCollapsed() const override { return false; }
 
-	virtual void GenerateHeaderRowContent( FDetailWidgetRow& NodeRow ) OVERRIDE
+	virtual void GenerateHeaderRowContent( FDetailWidgetRow& NodeRow ) override
 	{
 		if (bGenerateHeader)
 		{
@@ -255,7 +255,7 @@ public:
 		}
 	}
 
-	virtual void GenerateChildContent( IDetailChildrenBuilder& ChildrenBuilder ) OVERRIDE
+	virtual void GenerateChildContent( IDetailChildrenBuilder& ChildrenBuilder ) override
 	{
 		uint32 NumChildren = 0;
 		ArrayProperty->GetNumElements( NumChildren );
@@ -274,7 +274,7 @@ public:
 	}
 
 protected:
-	virtual void SetOnRebuildChildren( FSimpleDelegate InOnRebuildChildren  ) OVERRIDE { OnRebuildChildren = InOnRebuildChildren; } 
+	virtual void SetOnRebuildChildren( FSimpleDelegate InOnRebuildChildren  ) override { OnRebuildChildren = InOnRebuildChildren; } 
 
 	void OnNumChildrenChanged()
 	{
@@ -397,13 +397,13 @@ private:
 	void OnHideMaterialsForElement( int32 SlotIndex );
 
 	/** IDetailCustomNodeBuilder interface */
-	virtual void SetOnRebuildChildren( FSimpleDelegate InOnRebuildChildren  ) OVERRIDE { OnRebuildChildren = InOnRebuildChildren; } 
-	virtual bool RequiresTick() const OVERRIDE { return true; }
-	virtual void Tick( float DeltaTime ) OVERRIDE;
-	virtual void GenerateHeaderRowContent( FDetailWidgetRow& NodeRow ) OVERRIDE;
-	virtual void GenerateChildContent( IDetailChildrenBuilder& ChildrenBuilder ) OVERRIDE;
-	virtual FName GetName() const OVERRIDE { return NAME_None; }
-	virtual bool InitiallyCollapsed() const OVERRIDE { return false; }
+	virtual void SetOnRebuildChildren( FSimpleDelegate InOnRebuildChildren  ) override { OnRebuildChildren = InOnRebuildChildren; } 
+	virtual bool RequiresTick() const override { return true; }
+	virtual void Tick( float DeltaTime ) override;
+	virtual void GenerateHeaderRowContent( FDetailWidgetRow& NodeRow ) override;
+	virtual void GenerateChildContent( IDetailChildrenBuilder& ChildrenBuilder ) override;
+	virtual FName GetName() const override { return NAME_None; }
+	virtual bool InitiallyCollapsed() const override { return false; }
 
 	/**
 	 * Adds a new material item to the list
diff --git a/Engine/Source/Editor/ReferenceViewer/Classes/EdGraphNode_Reference.h b/Engine/Source/Editor/ReferenceViewer/Classes/EdGraphNode_Reference.h
--- a/Engine/Source/Editor/ReferenceViewer/Classes/EdGraphNode_Reference.h
+++ b/Engine/Source/Editor/ReferenceViewer/Classes/EdGraphNode_Reference.h
@@ -16,9 +16,9 @@ class UEdGraphNode_Reference : public UEdGraphNode
 	class UEdGraph_ReferenceViewer* GetReferenceViewerGraph() const;
 
 	// UEdGraphNode implementation
-	virtual FText GetNodeTitle(ENodeTitleType::Type TitleType) const OVERRIDE;
-	virtual void AllocateDefaultPins() OVERRIDE;
-	virtual UObject* GetJumpTargetForDoubleClick() const OVERRIDE;
+	virtual FText GetNodeTitle(ENodeTitleType::Type TitleType) const override;
+	virtual void AllocateDefaultPins() override;
+	virtual UObject* GetJumpTargetForDoubleClick() const override;
 	// End UEdGraphNode implementation
 
 	void CacheAssetData(const FAssetData& AssetData);
diff --git a/Engine/Source/Editor/ReferenceViewer/Classes/ReferenceViewerSchema.h b/Engine/Source/Editor/ReferenceViewer/Classes/ReferenceViewerSchema.h
--- a/Engine/Source/Editor/ReferenceViewer/Classes/ReferenceViewerSchema.h
+++ b/Engine/Source/Editor/ReferenceViewer/Classes/ReferenceViewerSchema.h
@@ -10,12 +10,12 @@ class UReferenceViewerSchema : public UEdGraphSchema
 
 public:
 	// UEdGraphSchema interface
-	virtual void GetContextMenuActions(const UEdGraph* CurrentGraph, const UEdGraphNode* InGraphNode, const UEdGraphPin* InGraphPin, class FMenuBuilder* MenuBuilder, bool bIsDebugging) const OVERRIDE;
-	virtual FLinearColor GetPinTypeColor(const FEdGraphPinType& PinType) const OVERRIDE;
-	virtual void BreakPinLinks(UEdGraphPin& TargetPin, bool bSendsNodeNotifcation) const OVERRIDE;
-	virtual void BreakSinglePinLink(UEdGraphPin* SourcePin, UEdGraphPin* TargetPin) OVERRIDE;
-	virtual FPinConnectionResponse MovePinLinks(UEdGraphPin& MoveFromPin, UEdGraphPin& MoveToPin, bool bIsItermeadiateMove = false) const OVERRIDE;
-	virtual FPinConnectionResponse CopyPinLinks(UEdGraphPin& CopyFromPin, UEdGraphPin& CopyToPin, bool bIsItermeadiateCopy = false) const OVERRIDE;
+	virtual void GetContextMenuActions(const UEdGraph* CurrentGraph, const UEdGraphNode* InGraphNode, const UEdGraphPin* InGraphPin, class FMenuBuilder* MenuBuilder, bool bIsDebugging) const override;
+	virtual FLinearColor GetPinTypeColor(const FEdGraphPinType& PinType) const override;
+	virtual void BreakPinLinks(UEdGraphPin& TargetPin, bool bSendsNodeNotifcation) const override;
+	virtual void BreakSinglePinLink(UEdGraphPin* SourcePin, UEdGraphPin* TargetPin) override;
+	virtual FPinConnectionResponse MovePinLinks(UEdGraphPin& MoveFromPin, UEdGraphPin& MoveToPin, bool bIsItermeadiateMove = false) const override;
+	virtual FPinConnectionResponse CopyPinLinks(UEdGraphPin& CopyFromPin, UEdGraphPin& CopyToPin, bool bIsItermeadiateCopy = false) const override;
 	// End of UEdGraphSchema interface
 
 private:
diff --git a/Engine/Source/Editor/ReferenceViewer/Private/ReferenceViewer.cpp b/Engine/Source/Editor/ReferenceViewer/Private/ReferenceViewer.cpp
--- a/Engine/Source/Editor/ReferenceViewer/Private/ReferenceViewer.cpp
+++ b/Engine/Source/Editor/ReferenceViewer/Private/ReferenceViewer.cpp
@@ -10,7 +10,7 @@
 
 class FGraphPanelNodeFactory_ReferenceViewer : public FGraphPanelNodeFactory
 {
-	virtual TSharedPtr<class SGraphNode> CreateNode(UEdGraphNode* Node) const OVERRIDE
+	virtual TSharedPtr<class SGraphNode> CreateNode(UEdGraphNode* Node) const override
 	{
 		if ( UEdGraphNode_Reference* DependencyNode = Cast<UEdGraphNode_Reference>(Node) )
 		{
@@ -30,7 +30,7 @@ public:
 		
 	}
 
-	virtual void StartupModule() OVERRIDE
+	virtual void StartupModule() override
 	{
 		GraphPanelNodeFactory = MakeShareable( new FGraphPanelNodeFactory_ReferenceViewer() );
 		FEdGraphUtilities::RegisterVisualNodeFactory(GraphPanelNodeFactory);
@@ -40,7 +40,7 @@ public:
 			.SetMenuType( ETabSpawnerMenuType::Hide );
 	}
 
-	virtual void ShutdownModule() OVERRIDE
+	virtual void ShutdownModule() override
 	{
 		if ( GraphPanelNodeFactory.IsValid() )
 		{
@@ -51,14 +51,14 @@ public:
 		FGlobalTabmanager::Get()->UnregisterNomadTabSpawner(ReferenceViewerTabId);
 	}
 
-	virtual void InvokeReferenceViewerTab(const TArray<FName>& GraphRootPackageNames) OVERRIDE
+	virtual void InvokeReferenceViewerTab(const TArray<FName>& GraphRootPackageNames) override
 	{
 		TSharedRef<SDockTab> NewTab = FGlobalTabmanager::Get()->InvokeTab( ReferenceViewerTabId );
 		TSharedRef<SReferenceViewer> ReferenceViewer = StaticCastSharedRef<SReferenceViewer>( NewTab->GetContent() );
 		ReferenceViewer->SetGraphRootPackageNames(GraphRootPackageNames);
 	}
 
-	virtual TSharedRef<SWidget> CreateReferenceViewer(const TArray<FName>& GraphRootPackageNames) OVERRIDE
+	virtual TSharedRef<SWidget> CreateReferenceViewer(const TArray<FName>& GraphRootPackageNames) override
 	{
 		TSharedRef<SReferenceViewer> ReferenceViewer = SNew(SReferenceViewer);
 		ReferenceViewer->SetGraphRootPackageNames(GraphRootPackageNames);
diff --git a/Engine/Source/Editor/SceneOutliner/Private/SOutlinerTreeView.h b/Engine/Source/Editor/SceneOutliner/Private/SOutlinerTreeView.h
--- a/Engine/Source/Editor/SceneOutliner/Private/SOutlinerTreeView.h
+++ b/Engine/Source/Editor/SceneOutliner/Private/SOutlinerTreeView.h
@@ -28,11 +28,11 @@ namespace SceneOutliner
 
 	protected:
 
-		virtual FReply OnDragOver( const FGeometry& MyGeometry, const FDragDropEvent& DragDropEvent ) OVERRIDE;
+		virtual FReply OnDragOver( const FGeometry& MyGeometry, const FDragDropEvent& DragDropEvent ) override;
 
-		virtual void OnDragLeave( const FDragDropEvent& DragDropEvent ) OVERRIDE;
+		virtual void OnDragLeave( const FDragDropEvent& DragDropEvent ) override;
 
-		virtual FReply OnDrop( const FGeometry& MyGeometry, const FDragDropEvent& DragDropEvent ) OVERRIDE;
+		virtual FReply OnDrop( const FGeometry& MyGeometry, const FDragDropEvent& DragDropEvent ) override;
 
 		/** Weak reference to the outliner widget that owns this list */
 		TWeakPtr<SSceneOutliner> SceneOutlinerWeak;
@@ -61,13 +61,13 @@ namespace SceneOutliner
 
 
 		/** Overridden from SMultiColumnTableRow.  Generates a widget for this column of the tree row. */
-		virtual TSharedRef<SWidget> GenerateWidgetForColumn( const FName& ColumnName ) OVERRIDE;
+		virtual TSharedRef<SWidget> GenerateWidgetForColumn( const FName& ColumnName ) override;
 
 		TSharedPtr< SSceneOutlinerItemWidget > GetItemWidget();
 
 	protected:
 
-		virtual FReply OnMouseButtonDown(const FGeometry& MyGeometry, const FPointerEvent& MouseEvent) OVERRIDE;
+		virtual FReply OnMouseButtonDown(const FGeometry& MyGeometry, const FPointerEvent& MouseEvent) override;
 
 		FReply HandleDragDetected( const FGeometry& MyGeometry, const FPointerEvent& MouseEvent );
 		
@@ -75,14 +75,14 @@ namespace SceneOutliner
 
 		void PerformAttachment(FName SocketName, TWeakObjectPtr<AActor> ParentActorPtr, const TArray< TWeakObjectPtr<AActor> > ChildrenPtrs);
 
-		virtual void OnDragEnter( const FGeometry& MyGeometry, const FDragDropEvent& DragDropEvent ) OVERRIDE;
+		virtual void OnDragEnter( const FGeometry& MyGeometry, const FDragDropEvent& DragDropEvent ) override;
 
-		virtual FReply OnDragOver( const FGeometry& MyGeometry, const FDragDropEvent& DragDropEvent ) OVERRIDE
+		virtual FReply OnDragOver( const FGeometry& MyGeometry, const FDragDropEvent& DragDropEvent ) override
 		{
 			return FReply::Handled();
 		}
 
-		virtual FReply OnDrop( const FGeometry& MyGeometry, const FDragDropEvent& DragDropEvent ) OVERRIDE;
+		virtual FReply OnDrop( const FGeometry& MyGeometry, const FDragDropEvent& DragDropEvent ) override;
 
 
 	private:
diff --git a/Engine/Source/Editor/SceneOutliner/Private/SSceneOutliner.h b/Engine/Source/Editor/SceneOutliner/Private/SSceneOutliner.h
--- a/Engine/Source/Editor/SceneOutliner/Private/SSceneOutliner.h
+++ b/Engine/Source/Editor/SceneOutliner/Private/SSceneOutliner.h
@@ -92,17 +92,17 @@ namespace SceneOutliner
 		void ExpandActor(AActor* Actor);
 
 		/** SWidget interface */
-		virtual void Tick( const FGeometry& AllottedGeometry, const double InCurrentTime, const float InDeltaTime ) OVERRIDE;
+		virtual void Tick( const FGeometry& AllottedGeometry, const double InCurrentTime, const float InDeltaTime ) override;
 
 		/**	Broadcasts whenever the current selection changes */
 		FSimpleMulticastDelegate SelectionChanged;
 
 		/** Sends a requests to the Scene Outliner to refresh itself the next chance it gets */
-		virtual void Refresh() OVERRIDE;
+		virtual void Refresh() override;
 
 		// Begin FEditorUndoClient Interface
-		virtual void PostUndo(bool bSuccess) OVERRIDE;
-		virtual void PostRedo(bool bSuccess) OVERRIDE { PostUndo(bSuccess); }
+		virtual void PostUndo(bool bSuccess) override;
+		virtual void PostRedo(bool bSuccess) override { PostUndo(bSuccess); }
 		// End of FEditorUndoClient
 
 		/** Get an array of the currently selected items in this tree */
@@ -277,7 +277,7 @@ namespace SceneOutliner
 		bool IsActorDisplayable( const AActor* Actor ) const;
 
 		/** @return Returns a string to use for highlighting results in the outliner list */
-		virtual FText GetFilterHighlightText() const OVERRIDE;
+		virtual FText GetFilterHighlightText() const override;
 
 		/**
 		 * Handler for when a property changes on any object
@@ -325,10 +325,10 @@ namespace SceneOutliner
 		EVisibility GetSearchBoxVisibility() const;
 
 		/** Overridden from SWidget: Checks to see if this widget supports keyboard focus */
-		virtual bool SupportsKeyboardFocus() const OVERRIDE;
+		virtual bool SupportsKeyboardFocus() const override;
 
 		/** Overridden from SWidget: Called when a key is pressed down */
-		virtual FReply OnKeyDown( const FGeometry& MyGeometry, const FKeyboardEvent& InKeyboardEvent ) OVERRIDE;
+		virtual FReply OnKeyDown( const FGeometry& MyGeometry, const FKeyboardEvent& InKeyboardEvent ) override;
 
 		/** Function to validate actor list shown in scene outliner widget */
 		void ValidateOutlinerTreeView();
diff --git a/Engine/Source/Editor/SceneOutliner/Private/SSceneOutlinerItemWidget.h b/Engine/Source/Editor/SceneOutliner/Private/SSceneOutlinerItemWidget.h
--- a/Engine/Source/Editor/SceneOutliner/Private/SSceneOutlinerItemWidget.h
+++ b/Engine/Source/Editor/SceneOutliner/Private/SSceneOutlinerItemWidget.h
@@ -32,7 +32,7 @@ namespace SceneOutliner
 		 * @param  InCurrentTime  Current absolute real time
 		 * @param  InDeltaTime  Real time passed since last tick
 		 */
-		virtual void Tick( const FGeometry& AllottedGeometry, const double InCurrentTime, const float InDeltaTime ) OVERRIDE;
+		virtual void Tick( const FGeometry& AllottedGeometry, const double InCurrentTime, const float InDeltaTime ) override;
 
 		/**
 		 * The widget should respond by populating the OutDrawElements array with FDrawElements 
@@ -47,7 +47,7 @@ namespace SceneOutliner
 		 *
 		 * @return The maximum layer ID attained by this widget or any of its children.
 		 */
-		virtual int32 OnPaint( const FGeometry& AllottedGeometry, const FSlateRect& MyClippingRect, FSlateWindowElementList& OutDrawElements, int32 LayerId, const FWidgetStyle& InWidgetStyle, bool bParentEnabled ) const OVERRIDE;
+		virtual int32 OnPaint( const FGeometry& AllottedGeometry, const FSlateRect& MyClippingRect, FSlateWindowElementList& OutDrawElements, int32 LayerId, const FWidgetStyle& InWidgetStyle, bool bParentEnabled ) const override;
 
 	private:
 		/** How many pixels to extend the highlight rectangle's left side horizontally */
diff --git a/Engine/Source/Editor/SceneOutliner/Private/SceneOutlinerActorInfoColumn.h b/Engine/Source/Editor/SceneOutliner/Private/SceneOutlinerActorInfoColumn.h
--- a/Engine/Source/Editor/SceneOutliner/Private/SceneOutlinerActorInfoColumn.h
+++ b/Engine/Source/Editor/SceneOutliner/Private/SceneOutlinerActorInfoColumn.h
@@ -22,19 +22,19 @@ public:
 	//////////////////////////////////////////////////////////////////////////
 	// Begin ISceneOutlinerColumn Implementation
 
-	virtual FName GetColumnID() OVERRIDE;
+	virtual FName GetColumnID() override;
 
-	virtual SHeaderRow::FColumn::FArguments ConstructHeaderRowColumn() OVERRIDE;
+	virtual SHeaderRow::FColumn::FArguments ConstructHeaderRowColumn() override;
 
-	virtual const TSharedRef< SWidget > ConstructRowWidget( const TSharedRef<SceneOutliner::TOutlinerTreeItem> TreeItem ) OVERRIDE;
+	virtual const TSharedRef< SWidget > ConstructRowWidget( const TSharedRef<SceneOutliner::TOutlinerTreeItem> TreeItem ) override;
 
-	virtual bool ProvidesSearchStrings() OVERRIDE;
+	virtual bool ProvidesSearchStrings() override;
 
-	virtual void PopulateActorSearchStrings( const AActor* const Actor, OUT TArray< FString >& OutSearchStrings ) const OVERRIDE;
+	virtual void PopulateActorSearchStrings( const AActor* const Actor, OUT TArray< FString >& OutSearchStrings ) const override;
 
-	virtual bool SupportsSorting() const OVERRIDE;
+	virtual bool SupportsSorting() const override;
 
-	virtual void SortItems(TArray<TSharedPtr<SceneOutliner::TOutlinerTreeItem>>& RootItems, const EColumnSortMode::Type SortMode) const OVERRIDE;
+	virtual void SortItems(TArray<TSharedPtr<SceneOutliner::TOutlinerTreeItem>>& RootItems, const EColumnSortMode::Type SortMode) const override;
 	
 	// End ISceneOutlinerColumn Implementation
 	//////////////////////////////////////////////////////////////////////////
diff --git a/Engine/Source/Editor/SceneOutliner/Private/SceneOutlinerGutter.cpp b/Engine/Source/Editor/SceneOutliner/Private/SceneOutlinerGutter.cpp
--- a/Engine/Source/Editor/SceneOutliner/Private/SceneOutlinerGutter.cpp
+++ b/Engine/Source/Editor/SceneOutliner/Private/SceneOutlinerGutter.cpp
@@ -19,7 +19,7 @@ public:
 	TUniquePtr<FScopedTransaction> UndoTransaction;
 
 	/** The widget decorator to use */
-	virtual TSharedPtr<SWidget> GetDefaultDecorator() const OVERRIDE
+	virtual TSharedPtr<SWidget> GetDefaultDecorator() const override
 	{
 		return SNullWidget::NullWidget;
 	}
@@ -67,7 +67,7 @@ public:
 private:
 
 	/** Start a new drag/drop operation for this widget */
-	virtual FReply OnDragDetected(const FGeometry& MyGeometry, const FPointerEvent& MouseEvent) OVERRIDE
+	virtual FReply OnDragDetected(const FGeometry& MyGeometry, const FPointerEvent& MouseEvent) override
 	{
 		if (MouseEvent.IsMouseButtonDown(EKeys::LeftMouseButton))
 		{
@@ -80,7 +80,7 @@ private:
 	}
 
 	/** If a visibility drag drop operation has entered this widget, set its actor to the new visibility state */
-	virtual void OnDragEnter(const FGeometry& MyGeometry, const FDragDropEvent& DragDropEvent) OVERRIDE
+	virtual void OnDragEnter(const FGeometry& MyGeometry, const FDragDropEvent& DragDropEvent) override
 	{
 		auto VisibilityOp = DragDropEvent.GetOperationAs<FVisibilityDragDropOp>();
 		if (VisibilityOp.IsValid())
@@ -90,7 +90,7 @@ private:
 	}
 
 	/** Called when the mouse button is pressed down on this widget */
-	virtual FReply OnMouseButtonDown(const FGeometry& MyGeometry, const FPointerEvent& MouseEvent) OVERRIDE
+	virtual FReply OnMouseButtonDown(const FGeometry& MyGeometry, const FPointerEvent& MouseEvent) override
 	{
 		if (MouseEvent.GetEffectingButton() != EKeys::LeftMouseButton)
 		{
@@ -106,7 +106,7 @@ private:
 	}
 
 	/** Process a mouse up message */
-	virtual FReply OnMouseButtonUp( const FGeometry& MyGeometry, const FPointerEvent& MouseEvent ) OVERRIDE
+	virtual FReply OnMouseButtonUp( const FGeometry& MyGeometry, const FPointerEvent& MouseEvent ) override
 	{
 		if (MouseEvent.GetEffectingButton() == EKeys::LeftMouseButton)
 		{
@@ -118,7 +118,7 @@ private:
 	}
 
 	/** Called when this widget had captured the mouse, but that capture has been revoked for some reason. */
-	virtual void OnMouseCaptureLost() OVERRIDE
+	virtual void OnMouseCaptureLost() override
 	{
 		UndoTransaction.Reset();
 	}
diff --git a/Engine/Source/Editor/SceneOutliner/Private/SceneOutlinerGutter.h b/Engine/Source/Editor/SceneOutliner/Private/SceneOutlinerGutter.h
--- a/Engine/Source/Editor/SceneOutliner/Private/SceneOutlinerGutter.h
+++ b/Engine/Source/Editor/SceneOutliner/Private/SceneOutlinerGutter.h
@@ -21,19 +21,19 @@ public:
 
 	// -----------------------------------------
 	// ISceneOutlinerColumn Implementation
-	virtual FName GetColumnID() OVERRIDE;
+	virtual FName GetColumnID() override;
 
-	virtual SHeaderRow::FColumn::FArguments ConstructHeaderRowColumn() OVERRIDE;
+	virtual SHeaderRow::FColumn::FArguments ConstructHeaderRowColumn() override;
 
 	virtual const TSharedRef< SWidget > ConstructRowWidget( const TSharedRef<SceneOutliner::TOutlinerTreeItem> TreeItem ) override;
 
-	virtual bool ProvidesSearchStrings() OVERRIDE { return false; }
+	virtual bool ProvidesSearchStrings() override { return false; }
 
-	virtual void PopulateActorSearchStrings(const AActor* const Actor, OUT TArray< FString >& OutSearchStrings) const OVERRIDE {}
+	virtual void PopulateActorSearchStrings(const AActor* const Actor, OUT TArray< FString >& OutSearchStrings) const override {}
 
-	virtual bool SupportsSorting() const OVERRIDE { return true; }
+	virtual bool SupportsSorting() const override { return true; }
 
-	virtual void SortItems(TArray<TSharedPtr<SceneOutliner::TOutlinerTreeItem>>& RootItems, const EColumnSortMode::Type SortMode) const OVERRIDE;
+	virtual void SortItems(TArray<TSharedPtr<SceneOutliner::TOutlinerTreeItem>>& RootItems, const EColumnSortMode::Type SortMode) const override;
 	// -----------------------------------------
 
 private:
diff --git a/Engine/Source/Editor/SceneOutliner/Public/SceneOutlinerFilters.h b/Engine/Source/Editor/SceneOutliner/Public/SceneOutlinerFilters.h
--- a/Engine/Source/Editor/SceneOutliner/Public/SceneOutlinerFilters.h
+++ b/Engine/Source/Editor/SceneOutliner/Public/SceneOutlinerFilters.h
@@ -68,19 +68,19 @@ namespace SceneOutliner
 			ProxyFilter(TFilterPtr InFilter) : Filter(InFilter) {}
 
 			/** Returns whether the specified Item passes the Filter's restrictions */
-			virtual bool PassesFilter(const FOutlinerFilterProxy& InItem) const OVERRIDE
+			virtual bool PassesFilter(const FOutlinerFilterProxy& InItem) const override
 			{
 				return ProxyPassesFilter(InItem, Filter);
 			}
 
 			/** Get the on changed event - does nothing for this wrapper */
-			virtual FChangedEvent& OnChanged() OVERRIDE
+			virtual FChangedEvent& OnChanged() override
 			{
 				return ChangedEvent;
 			}
 
 			/** Override the relevant function that matches our type */
-			virtual bool MatchesWrappedFilter(TFilterPtr WrappedFilter) const OVERRIDE
+			virtual bool MatchesWrappedFilter(TFilterPtr WrappedFilter) const override
 			{
 				return WrappedFilter == Filter;
 			}
