PropertyEditor refactored. It's ready for StructureDefaultView. Some SDetailsView functions were moved to SDetailsViewBase. FPropertyValueImpl calls NotifyFinishedChangingProperties properly.

diff --git a/Engine/Source/Editor/Persona/Private/AnimNotifyDetails.cpp b/Engine/Source/Editor/Persona/Private/AnimNotifyDetails.cpp
--- a/Engine/Source/Editor/Persona/Private/AnimNotifyDetails.cpp
+++ b/Engine/Source/Editor/Persona/Private/AnimNotifyDetails.cpp
@@ -19,23 +19,25 @@ TSharedRef<IDetailCustomization> FAnimNotifyDetails::MakeInstance()
 void FAnimNotifyDetails::CustomizeDetails( IDetailLayoutBuilder& DetailBuilder )
 {
 	const UClass* DetailObjectClass = DetailBuilder.GetDetailsView().GetBaseClass();
-
-	if (DetailObjectClass->GetName().Find(TEXT("AnimNotify_PlayParticleEffect")) != INDEX_NONE)
-	{
-		AddBoneNameProperty(DetailBuilder, DetailObjectClass, TEXT("SocketName"), TEXT("AnimNotify"));
-	}
-	else if (DetailObjectClass->GetName().Find(TEXT("AnimNotify_PlaySound")) != INDEX_NONE)
-	{
-		AddBoneNameProperty(DetailBuilder, DetailObjectClass, TEXT("AttachName"), TEXT("AnimNotify"));
-	}
-	else if (DetailObjectClass->GetName().Find(TEXT("AnimNotifyState_Trail")) != INDEX_NONE)
-	{
-		AddBoneNameProperty(DetailBuilder, DetailObjectClass, TEXT("FirstSocketName"), TEXT("Trail"));
-		AddBoneNameProperty(DetailBuilder, DetailObjectClass, TEXT("SecondSocketName"), TEXT("Trail"));
-	}
-	else if(DetailObjectClass->GetName().Find(TEXT("AnimNotifyState_TimedParticleEffect")) != INDEX_NONE)
+	if (DetailObjectClass)
 	{
-		AddBoneNameProperty(DetailBuilder, DetailObjectClass, TEXT("SocketName"), TEXT("ParticleSystem"));
+		if (DetailObjectClass->GetName().Find(TEXT("AnimNotify_PlayParticleEffect")) != INDEX_NONE)
+		{
+			AddBoneNameProperty(DetailBuilder, DetailObjectClass, TEXT("SocketName"), TEXT("AnimNotify"));
+		}
+		else if (DetailObjectClass->GetName().Find(TEXT("AnimNotify_PlaySound")) != INDEX_NONE)
+		{
+			AddBoneNameProperty(DetailBuilder, DetailObjectClass, TEXT("AttachName"), TEXT("AnimNotify"));
+		}
+		else if (DetailObjectClass->GetName().Find(TEXT("AnimNotifyState_Trail")) != INDEX_NONE)
+		{
+			AddBoneNameProperty(DetailBuilder, DetailObjectClass, TEXT("FirstSocketName"), TEXT("Trail"));
+			AddBoneNameProperty(DetailBuilder, DetailObjectClass, TEXT("SecondSocketName"), TEXT("Trail"));
+		}
+		else if (DetailObjectClass->GetName().Find(TEXT("AnimNotifyState_TimedParticleEffect")) != INDEX_NONE)
+		{
+			AddBoneNameProperty(DetailBuilder, DetailObjectClass, TEXT("SocketName"), TEXT("ParticleSystem"));
+		}
 	}
 }
 
diff --git a/Engine/Source/Editor/PropertyEditor/Private/CategoryPropertyNode.cpp b/Engine/Source/Editor/PropertyEditor/Private/CategoryPropertyNode.cpp
--- a/Engine/Source/Editor/PropertyEditor/Private/CategoryPropertyNode.cpp
+++ b/Engine/Source/Editor/PropertyEditor/Private/CategoryPropertyNode.cpp
@@ -58,26 +58,29 @@ void FCategoryPropertyNode::InitChildNodes()
 
 	TArray<UProperty*> Properties;
 	// The parent of a category window has to be an object window.
-	FObjectPropertyNode* itemobj = FindObjectItemParent();
-	// Get a list of properties that are in the same category
-	for( TFieldIterator<UProperty> It(itemobj->GetObjectBaseClass()); It; ++It )
+	auto ComplexNode = FindComplexParent();
+	if (ComplexNode)
 	{
-		bool bMetaDataAllowVisible = true;
-		FString MetaDataVisibilityCheckString = It->GetMetaData(TEXT("bShowOnlyWhenTrue"));
-		if (MetaDataVisibilityCheckString.Len())
+		// Get a list of properties that are in the same category
+		for (TFieldIterator<UProperty> It(ComplexNode->GetBaseStructure()); It; ++It)
 		{
-			//ensure that the metadata visibility string is actually set to true in order to show this property
-			// @todo Remove this
-			GConfig->GetBool(TEXT("UnrealEd.PropertyFilters"), *MetaDataVisibilityCheckString, bMetaDataAllowVisible, GEditorUserSettingsIni);
-		}
+			bool bMetaDataAllowVisible = true;
+			FString MetaDataVisibilityCheckString = It->GetMetaData(TEXT("bShowOnlyWhenTrue"));
+			if (MetaDataVisibilityCheckString.Len())
+			{
+				//ensure that the metadata visibility string is actually set to true in order to show this property
+				// @todo Remove this
+				GConfig->GetBool(TEXT("UnrealEd.PropertyFilters"), *MetaDataVisibilityCheckString, bMetaDataAllowVisible, GEditorUserSettingsIni);
+			}
 
-		if (bMetaDataAllowVisible)
-		{
-			// Add if we are showing non-editable props and this is the 'None' category, 
-			// or if this is the right category, and we are either showing non-editable
-			if ( FObjectEditorUtils::GetCategoryFName(*It) == CategoryName && (bShowHiddenProperties || (It->PropertyFlags & CPF_Edit)) )
+			if (bMetaDataAllowVisible)
 			{
-				Properties.Add( *It );
+				// Add if we are showing non-editable props and this is the 'None' category, 
+				// or if this is the right category, and we are either showing non-editable
+				if (FObjectEditorUtils::GetCategoryFName(*It) == CategoryName && (bShowHiddenProperties || (It->PropertyFlags & CPF_Edit)))
+				{
+					Properties.Add(*It);
+				}
 			}
 		}
 	}
diff --git a/Engine/Source/Editor/PropertyEditor/Private/DetailCategoryBuilderImpl.cpp b/Engine/Source/Editor/PropertyEditor/Private/DetailCategoryBuilderImpl.cpp
--- a/Engine/Source/Editor/PropertyEditor/Private/DetailCategoryBuilderImpl.cpp
+++ b/Engine/Source/Editor/PropertyEditor/Private/DetailCategoryBuilderImpl.cpp
@@ -86,11 +86,11 @@ FDetailCategoryImpl::FDetailCategoryImpl( FName InCategoryName, TSharedRef<FDeta
 	, bHasFilterStrings( false )
 	, bHasVisibleDetails( true )
 {
-	const UClass* BaseClass = InDetailLayout->GetDetailsView().GetBaseClass();
+	const UStruct* BaseStruct = InDetailLayout->GetDetailsView().GetBaseStruct();
 	// Use the base class name if there is one otherwise this is a generic category not specific to a class
-	FName BaseClassName = BaseClass ? BaseClass->GetFName() : FName("Generic");
+	FName BaseStructName = BaseStruct ? BaseStruct->GetFName() : FName("Generic");
 
-	CategoryPathName = BaseClassName.ToString() + TEXT(".") + CategoryName.ToString();
+	CategoryPathName = BaseStructName.ToString() + TEXT(".") + CategoryName.ToString();
 
 	GConfig->GetBool( TEXT("DetailCategoriesAdvanced"), *CategoryPathName, bUserShowAdvanced, GEditorUserSettingsIni );
 
@@ -447,11 +447,11 @@ void FDetailCategoryImpl::SetDisplayName( FName InCategoryName, const FString& L
 	}
 	else
 	{
-		const UClass* BaseClass = GetParentLayoutImpl().GetDetailsView().GetBaseClass();
+		const UStruct* BaseStruct = GetParentLayoutImpl().GetDetailsView().GetBaseStruct();
 		// Use the base class name if there is one otherwise this is a generic category not specific to a class
-		FName BaseClassName = BaseClass ? BaseClass->GetFName() : FName("Generic");
+		FName BaseStructName = BaseStruct ? BaseStruct->GetFName() : FName("Generic");
 
-		FString CategoryStr = InCategoryName != NAME_None ? InCategoryName.ToString() : BaseClassName.ToString();
+		FString CategoryStr = InCategoryName != NAME_None ? InCategoryName.ToString() : BaseStructName.ToString();
 		FString SourceCategoryStr = FName::NameToDisplayString( CategoryStr, false );
 
 		bool FoundText = false;
@@ -462,7 +462,7 @@ void FDetailCategoryImpl::SetDisplayName( FName InCategoryName, const FString& L
 		}
 		else
 		{
-			FoundText = FText::FindText( TEXT("DetailCategory.ClassName"), BaseClassName.ToString(), /*OUT*/DisplayNameText );
+			FoundText = FText::FindText(TEXT("DetailCategory.ClassName"), BaseStructName.ToString(), /*OUT*/DisplayNameText);
 		}
 
 		if ( FoundText )
diff --git a/Engine/Source/Editor/PropertyEditor/Private/DetailLayoutBuilderImpl.cpp b/Engine/Source/Editor/PropertyEditor/Private/DetailLayoutBuilderImpl.cpp
--- a/Engine/Source/Editor/PropertyEditor/Private/DetailLayoutBuilderImpl.cpp
+++ b/Engine/Source/Editor/PropertyEditor/Private/DetailLayoutBuilderImpl.cpp
@@ -21,8 +21,8 @@ IDetailCategoryBuilder& FDetailLayoutBuilderImpl::EditCategory( FName CategoryNa
 	// Use the base class name if the user did not specify a name
 	if( CategoryName == NAME_None )
 	{
-		const UClass* BaseClass = GetDetailsView().GetBaseClass();
-		CategoryName = BaseClass ? BaseClass->GetFName() : FName("Generic");
+		const UStruct* BaseStruct = GetDetailsView().GetBaseStruct();
+		CategoryName = BaseStruct ? BaseStruct->GetFName() : FName("Generic");
 	}
 
 	TSharedPtr<FDetailCategoryImpl> CategoryImpl;
@@ -443,10 +443,13 @@ void FDetailLayoutBuilderImpl::GetObjectsBeingCustomized( TArray< TWeakObjectPtr
 			if( PropertyNodeMapPtr )
 			{
 				FPropertyNodeMap& PropertyNodeMap = *PropertyNodeMapPtr;
-
-				for( int32 ObjectIndex = 0; ObjectIndex < PropertyNodeMap.ParentObjectProperty->GetNumObjects(); ++ObjectIndex )
+				FObjectPropertyNode* ParentObjectProperty = PropertyNodeMap.ParentProperty ? PropertyNodeMap.ParentProperty->AsObjectNode() : NULL;
+				if (ParentObjectProperty)
 				{
-					OutObjects.Add( PropertyNodeMap.ParentObjectProperty->GetUObject( ObjectIndex ) );
+					for (int32 ObjectIndex = 0; ObjectIndex < ParentObjectProperty->GetNumObjects(); ++ObjectIndex)
+					{
+						OutObjects.Add(ParentObjectProperty->GetUObject(ObjectIndex));
+					}
 				}
 			}
 		}
diff --git a/Engine/Source/Editor/PropertyEditor/Private/ObjectPropertyNode.cpp b/Engine/Source/Editor/PropertyEditor/Private/ObjectPropertyNode.cpp
--- a/Engine/Source/Editor/PropertyEditor/Private/ObjectPropertyNode.cpp
+++ b/Engine/Source/Editor/PropertyEditor/Private/ObjectPropertyNode.cpp
@@ -9,7 +9,7 @@
 #include "ObjectEditorUtils.h"
 
 FObjectPropertyNode::FObjectPropertyNode(void)
-	: FPropertyNode()
+	: FComplexPropertyNode()
 	, BaseClass(NULL)
 {
 }
diff --git a/Engine/Source/Editor/PropertyEditor/Private/ObjectPropertyNode.h b/Engine/Source/Editor/PropertyEditor/Private/ObjectPropertyNode.h
--- a/Engine/Source/Editor/PropertyEditor/Private/ObjectPropertyNode.h
+++ b/Engine/Source/Editor/PropertyEditor/Private/ObjectPropertyNode.h
@@ -11,7 +11,7 @@
 typedef TArray< TWeakObjectPtr<UObject> >::TIterator TPropObjectIterator;
 typedef TArray< TWeakObjectPtr<UObject> >::TConstIterator TPropObjectConstIterator;
 
-class FObjectPropertyNode : public FPropertyNode
+class FObjectPropertyNode : public FComplexPropertyNode
 {
 public:
 	FObjectPropertyNode();
@@ -61,6 +61,26 @@ public:
 	/** @return		The base-est baseclass for objects in this list. */
 	const UClass*	GetObjectBaseClass() const { return BaseClass.IsValid() ? BaseClass.Get() : NULL; }
 
+
+	// FComplexPropertyNode implementation
+	virtual UStruct* GetBaseStructure() override { return GetObjectBaseClass(); }
+	virtual const UStruct* GetBaseStructure() const override{ return GetObjectBaseClass(); }
+	virtual int32 GetInstancesNum() const override{ return GetNumObjects(); }
+	virtual uint8* GetMemoryOfInstance(int32 Index)
+	{
+		return (uint8*)GetUObject(Index);
+	}
+	virtual TWeakObjectPtr<UObject> GetInstanceAsUObject(int32 Index) override
+	{
+		check(Objects.IsValidIndex(Index));
+		return Objects[Index];
+	}
+	virtual EPropertyType GetPropertyType() const override{ return EPT_Object; }
+	virtual void Disconnect() override 
+	{
+		RemoveAllObjects();
+	}
+
 	//////////////////////////////////////////////////////////////////////////
 	/** @return		The property stored at this node, to be passed to Pre/PostEditChange. */
 	virtual UProperty*		GetStoredProperty()		{ return StoredProperty.IsValid() ? StoredProperty.Get() : NULL; }
diff --git a/Engine/Source/Editor/PropertyEditor/Private/Presentation/PropertyEditor/PropertyEditor.cpp b/Engine/Source/Editor/PropertyEditor/Private/Presentation/PropertyEditor/PropertyEditor.cpp
--- a/Engine/Source/Editor/PropertyEditor/Private/Presentation/PropertyEditor/PropertyEditor.cpp
+++ b/Engine/Source/Editor/PropertyEditor/Private/Presentation/PropertyEditor/PropertyEditor.cpp
@@ -505,10 +505,12 @@ bool FPropertyEditor::GetEditConditionPropertyAddress( UBoolProperty*& Condition
 			}
 		}
 
-		FObjectPropertyNode* ObjectNode = ParentNode->FindObjectItemParent();
-		for ( TPropObjectIterator Itor( ObjectNode->ObjectIterator() ) ; Itor ; ++Itor )
+		auto ComplexParentNode = ParentNode->FindComplexParent();
+		if (ComplexParentNode)
 		{
-			TWeakObjectPtr<UObject> Object = *Itor;
+			for (int32 Index = 0; Index < ComplexParentNode->GetInstancesNum(); ++Index)
+			{
+				TWeakObjectPtr<UObject> Object = ComplexParentNode->GetInstanceAsUObject(Index);
 
 			if( Object.IsValid() )
 			{
@@ -527,6 +529,7 @@ bool FPropertyEditor::GetEditConditionPropertyAddress( UBoolProperty*& Condition
 			}
 		}
 	}
+	}
 
 	if ( bResult )
 	{
@@ -544,12 +547,12 @@ bool FPropertyEditor::SupportsEditConditionToggle( UProperty* InProperty )
 	if (!InProperty->HasMetaData(TEXT("HideEditConditionToggle")))
 	{
 		bool bNegateValue = false;
-		UBoolProperty* ConditionalProperty = GetEditConditionProperty(InProperty, bNegateValue);
-		if (ConditionalProperty != NULL)
+		UBoolProperty* ConditionalProperty = GetEditConditionProperty( InProperty, bNegateValue );
+		if( ConditionalProperty != NULL )
 		{
 			bShowEditConditionToggle = true;
 
-			if (ConditionalProperty->HasAllPropertyFlags(CPF_Edit))
+			if( ConditionalProperty->HasAllPropertyFlags( CPF_Edit ) )
 			{
 				// Conditionally-dependent property is already exposed for editing, so no need to draw another
 				// check box next to this property's label
diff --git a/Engine/Source/Editor/PropertyEditor/Private/PropertyChangeListener.cpp b/Engine/Source/Editor/PropertyEditor/Private/PropertyChangeListener.cpp
--- a/Engine/Source/Editor/PropertyEditor/Private/PropertyChangeListener.cpp
+++ b/Engine/Source/Editor/PropertyEditor/Private/PropertyChangeListener.cpp
@@ -200,9 +200,12 @@ bool FPropertyChangeListener::ScanForChanges( bool bRecacheNewValues )
 
 			// Get each object that changed
 			TArray<UObject*> ObjectsThatChanged;
-			for( auto It = ObjectNode->ObjectConstIterator(); It; ++It )
+			if (ObjectNode)
 			{
-				ObjectsThatChanged.Add( It->Get() );
+				for (auto It = ObjectNode->ObjectConstIterator(); It; ++It)
+				{
+					ObjectsThatChanged.Add(It->Get());
+				}
 			}
 
 			if( ObjectsThatChanged.Num() > 0 )
@@ -232,9 +235,12 @@ void FPropertyChangeListener::TriggerAllPropertiesChangedDelegate()
 
 		// Get each object that changed
 		TArray<UObject*> ObjectsThatChanged;
-		for( auto It = ObjectNode->ObjectConstIterator(); It; ++It )
+		if (ObjectNode)
 		{
-			ObjectsThatChanged.Add( It->Get() );
+			for (auto It = ObjectNode->ObjectConstIterator(); It; ++It)
+			{
+				ObjectsThatChanged.Add(It->Get());
+			}
 		}
 
 		if( ObjectsThatChanged.Num() > 0 )
diff --git a/Engine/Source/Editor/PropertyEditor/Private/PropertyHandleImpl.cpp b/Engine/Source/Editor/PropertyEditor/Private/PropertyHandleImpl.cpp
--- a/Engine/Source/Editor/PropertyEditor/Private/PropertyHandleImpl.cpp
+++ b/Engine/Source/Editor/PropertyEditor/Private/PropertyHandleImpl.cpp
@@ -30,16 +30,15 @@ FPropertyValueImpl::~FPropertyValueImpl()
 void FPropertyValueImpl::GetObjectsToModify( TArray<FObjectBaseAddress>& ObjectsToModify, FPropertyNode* InPropertyNode ) const
 {
 	// Find the parent object node which contains offset addresses for reading a property value on an object
-	FObjectPropertyNode* ObjectNode = InPropertyNode->FindObjectItemParent();
-
-	if( ObjectNode )
+	FComplexPropertyNode* ComplexNode = InPropertyNode->FindComplexParent();
+	if (ComplexNode)
 	{
-		// Add each object and address to the list
-		for( TPropObjectIterator Itor( ObjectNode->ObjectIterator() ) ; Itor ; ++Itor )
+		const bool bIsStruct = FComplexPropertyNode::EPT_StandaloneStructure == ComplexNode->GetPropertyType();
+		for (int32 Index = 0; Index < ComplexNode->GetInstancesNum(); ++Index)
 		{
-			UObject*	Object = Itor->Get();
-			uint8*		Addr = InPropertyNode->GetValueBaseAddress( (uint8*) Object );
-			ObjectsToModify.Add( FObjectBaseAddress( Object, Addr ) );
+			UObject*	Object = ComplexNode->GetInstanceAsUObject(Index).Get();
+			uint8*		Addr = InPropertyNode->GetValueBaseAddress(ComplexNode->GetMemoryOfInstance(Index));
+			ObjectsToModify.Add(FObjectBaseAddress(Object, Addr, bIsStruct));
 		}
 	}
 }
@@ -241,11 +240,11 @@ bool FPropertyValueImpl::SendTextToObjectProperty( const FString& Text, EPropert
 	TSharedPtr<FPropertyNode> PropertyNodePin = PropertyNode.Pin();
 	if( PropertyNodePin.IsValid() )
 	{
-		FObjectPropertyNode* ObjectNode = PropertyNodePin->FindObjectItemParent();
+		FComplexPropertyNode* ParentNode = PropertyNodePin->FindComplexParent();
 
 		// If more than one object is selected, an empty field indicates their values for this property differ.
 		// Don't send it to the objects value in this case (if we did, they would all get set to None which isn't good).
-		if( !ObjectNode || ObjectNode->GetNumObjects() > 1 && !Text.Len() )
+		if (!ParentNode || ParentNode->GetInstancesNum() > 1 && !Text.Len())
 		{
 			return false;
 		}
@@ -285,7 +284,7 @@ FPropertyAccess::Result FPropertyValueImpl::ImportText( const FString& InValue,
 	TArray<FString> Values;
 	for( int32 ObjectIndex = 0 ; ObjectIndex < ObjectsToModify.Num() ; ++ObjectIndex )
 	{
-		if( ObjectsToModify[ObjectIndex].Object != NULL )
+		if (ObjectsToModify[ObjectIndex].Object != NULL || ObjectsToModify[ObjectIndex].bIsStruct)
 		{
 			Values.Add( InValue );
 		}
@@ -327,12 +326,14 @@ FPropertyAccess::Result FPropertyValueImpl::ImportText( const TArray<FObjectBase
 		bool bIsGameWorld = false;
 		// If the object we are modifying is in the PIE world, than make the PIE world the active
 		// GWorld.  Assumes all objects managed by this property window belong to the same world.
-		UPackage* ObjectPackage = InObjects[0].Object->GetOutermost();
-		bool bIsPIEPackage = !!(ObjectPackage->PackageFlags & PKG_PlayInEditor);
-		if ( GUnrealEd && GUnrealEd->PlayWorld && bIsPIEPackage && !GIsPlayInEditorWorld)
+		if (UPackage* ObjectPackage = (InObjects[0].Object ? InObjects[0].Object->GetOutermost() : NULL))
 		{
-			OldGWorld = SetPlayInEditorWorld(GUnrealEd->PlayWorld);
-			bIsGameWorld = true;
+			bool bIsPIEPackage = !!(ObjectPackage->PackageFlags & PKG_PlayInEditor);
+			if (GUnrealEd && GUnrealEd->PlayWorld && bIsPIEPackage && !GIsPlayInEditorWorld)
+			{
+				OldGWorld = SetPlayInEditorWorld(GUnrealEd->PlayWorld);
+				bIsGameWorld = true;
+			}
 		}
 		///////////////
 
@@ -361,16 +362,18 @@ FPropertyAccess::Result FPropertyValueImpl::ImportText( const TArray<FObjectBase
 
 			// If this property is the inner-property of an array, cache the current value as well
 			FString PreviousArrayValue;
-
-			FPropertyNode* ParentNode = InPropertyNode->GetParentNode();
-			if (ParentNode != NULL && ParentNode->GetProperty() && ParentNode->GetProperty()->IsA(UArrayProperty::StaticClass()))
+			if (Cur.Object)
 			{
-				UArrayProperty* ArrayProp = Cast<UArrayProperty>(ParentNode->GetProperty());
-				if (ArrayProp->Inner == NodeProperty)
+				FPropertyNode* ParentNode = InPropertyNode->GetParentNode();
+				if (ParentNode != NULL && ParentNode->GetProperty() && ParentNode->GetProperty()->IsA(UArrayProperty::StaticClass()))
 				{
-					uint8* Addr = ParentNode->GetValueBaseAddress((uint8*)Cur.Object);
-					
-					ArrayProp->ExportText_Direct(PreviousArrayValue, Addr, Addr, NULL, 0);
+					UArrayProperty* ArrayProp = Cast<UArrayProperty>(ParentNode->GetProperty());
+					if (ArrayProp->Inner == NodeProperty)
+					{
+						uint8* Addr = ParentNode->GetValueBaseAddress((uint8*)Cur.Object);
+
+						ArrayProp->ExportText_Direct(PreviousArrayValue, Addr, Addr, NULL, 0);
+					}
 				}
 			}
 
@@ -399,21 +402,24 @@ FPropertyAccess::Result FPropertyValueImpl::ImportText( const TArray<FObjectBase
 			const TCHAR* NewValue = *InValues[ObjectIndex];
 			NodeProperty->ImportText( NewValue, Cur.BaseAddress, 0, Cur.Object );
 
-			// Cache the value of the property after having modified it.
-			FString ValueAfterImport;
-			NodeProperty->ExportText_Direct(ValueAfterImport, Cur.BaseAddress, Cur.BaseAddress, NULL, 0 );
-
-			if ((Cur.Object->HasAnyFlags(RF_ClassDefaultObject|RF_ArchetypeObject) ||
-				  (Cur.Object->HasAnyFlags(RF_DefaultSubObject) && Cur.Object->GetOuter()->HasAnyFlags(RF_ClassDefaultObject|RF_ArchetypeObject))) && 
-					 !bIsGameWorld)
+			if (Cur.Object)
 			{
-				InPropertyNode->PropagatePropertyChange(Cur.Object, NewValue, PreviousArrayValue.IsEmpty() ? PreviousValue : PreviousArrayValue);
-			}
+				// Cache the value of the property after having modified it.
+				FString ValueAfterImport;
+				NodeProperty->ExportText_Direct(ValueAfterImport, Cur.BaseAddress, Cur.BaseAddress, NULL, 0);
 
-			// If the values before and after setting the property differ, mark the object dirty.
-			if ( FCString::Strcmp( *PreviousValue, *ValueAfterImport ) != 0 )
-			{
-				Cur.Object->MarkPackageDirty();
+				if ((Cur.Object->HasAnyFlags(RF_ClassDefaultObject | RF_ArchetypeObject) ||
+					(Cur.Object->HasAnyFlags(RF_DefaultSubObject) && Cur.Object->GetOuter()->HasAnyFlags(RF_ClassDefaultObject | RF_ArchetypeObject))) &&
+					!bIsGameWorld)
+				{
+					InPropertyNode->PropagatePropertyChange(Cur.Object, NewValue, PreviousArrayValue.IsEmpty() ? PreviousValue : PreviousArrayValue);
+				}
+
+				// If the values before and after setting the property differ, mark the object dirty.
+				if (FCString::Strcmp(*PreviousValue, *ValueAfterImport) != 0)
+				{
+					Cur.Object->MarkPackageDirty();
+				}
 			}
 
 			//add on array index so we can tell which entry just changed
@@ -449,11 +455,11 @@ FPropertyAccess::Result FPropertyValueImpl::ImportText( const TArray<FObjectBase
 			RestoreEditorWorld( OldGWorld );
 		}
 
-
-		if( PropertyUtilities.IsValid() && !bInteractiveChangeInProgress )
+		if (PropertyUtilities.IsValid() && !bInteractiveChangeInProgress)
 		{
 			const bool bTopologyChange = false;
-			FPropertyChangedEvent ChangeEvent(NodeProperty, bTopologyChange, bFinished ? EPropertyChangeType::ValueSet : EPropertyChangeType::Interactive );
+			FPropertyChangedEvent ChangeEvent(NodeProperty, bTopologyChange, bFinished ? EPropertyChangeType::ValueSet : EPropertyChangeType::Interactive);
+			InPropertyNode->FixPropertiesInEvent(ChangeEvent);
 			PropertyUtilities.Pin()->NotifyFinishedChangingProperties(ChangeEvent);
 		}
 
@@ -669,11 +675,11 @@ FPropertyAccess::Result FPropertyValueImpl::SetValueAsString( const FString& InV
 
 		// If more than one object is selected, an empty field indicates their values for this property differ.
 		// Don't send it to the objects value in this case (if we did, they would all get set to None which isn't good).
-		FObjectPropertyNode* ObjectNode = PropertyNodePin->FindObjectItemParent();
+		FComplexPropertyNode* ParentNode = PropertyNodePin->FindComplexParent();
 
 		FString PreviousValue;
 		GetValueAsString( PreviousValue );
-		if( ObjectNode && ( ObjectNode->GetNumObjects() == 1 || ( Value.Len() && PreviousValue != Value ) ) )
+		if (ParentNode && (ParentNode->GetInstancesNum() == 1 || (Value.Len() && PreviousValue != Value)))
 		{
 			ImportText( Value, PropertyNodePin.Get(), Flags );
 		}
@@ -1055,15 +1061,18 @@ void FPropertyValueImpl::AddChild()
 					ArrayIndicesPerObject.Add(TMap<FString,int32>());
 					FPropertyValueImpl::GenerateArrayIndexMapToObjectNode(ArrayIndicesPerObject[i], PropertyNodePin.Get());
 					
-					UObject* Obj = ObjectNode->GetUObject(i);
-					if ((Obj->HasAnyFlags(RF_ClassDefaultObject|RF_ArchetypeObject) || 
-						  (Obj->HasAnyFlags(RF_DefaultSubObject) && Obj->GetOuter()->HasAnyFlags(RF_ClassDefaultObject|RF_ArchetypeObject))) && 
-							!FApp::IsGame())
+					UObject* Obj = ObjectNode ? ObjectNode->GetUObject(i) : NULL;
+					if (Obj)
 					{
-						FString OrgArrayContent;
-						Array->ExportText_Direct(OrgArrayContent, Addr, Addr, NULL, 0);
+						if ((Obj->HasAnyFlags(RF_ClassDefaultObject | RF_ArchetypeObject) ||
+							(Obj->HasAnyFlags(RF_DefaultSubObject) && Obj->GetOuter()->HasAnyFlags(RF_ClassDefaultObject | RF_ArchetypeObject))) &&
+							!FApp::IsGame())
+						{
+							FString OrgArrayContent;
+							Array->ExportText_Direct(OrgArrayContent, Addr, Addr, NULL, 0);
 
-						PropertyNodePin->PropagateArrayPropertyChange(Obj, OrgArrayContent, EPropertyArrayChangeType::Add, -1);
+							PropertyNodePin->PropagateArrayPropertyChange(Obj, OrgArrayContent, EPropertyArrayChangeType::Add, -1);
+						}
 					}
 
 					FScriptArrayHelper	ArrayHelper(Array,Addr);
@@ -1081,22 +1090,19 @@ void FPropertyValueImpl::AddChild()
 				}
 			}
 
+			const bool bTopologyChange = true;
 			if ( bNotifiedPreChange )
 			{
-				const bool bExpand   = true;
-				const bool bRecurse  = false;
-
 				// send the PostEditChange notification; it will be propagated to all selected objects
-				const bool bTopologyChange = true;
 				FPropertyChangedEvent ChangeEvent(NodeProperty, bTopologyChange, EPropertyChangeType::ArrayAdd);
 				ChangeEvent.SetArrayIndexPerObject(ArrayIndicesPerObject);
 				PropertyNodePin->NotifyPostChange( ChangeEvent, NotifyHook );
 			}
 
-			if( PropertyUtilities.IsValid() )
+			if (PropertyUtilities.IsValid())
 			{
-				const bool bTopologyChange = true;
 				FPropertyChangedEvent ChangeEvent(NodeProperty, bTopologyChange, EPropertyChangeType::ArrayAdd);
+				PropertyNodePin->FixPropertiesInEvent(ChangeEvent);
 				PropertyUtilities.Pin()->NotifyFinishedChangingProperties(ChangeEvent);
 			}
 		}
@@ -1139,15 +1145,18 @@ void FPropertyValueImpl::ClearChildren()
 						PropertyNodePin->NotifyPreChange( NodeProperty, NotifyHook );
 					}
 
-					UObject* Obj = ObjectNode->GetUObject(i);
-					if ((Obj->HasAnyFlags(RF_ClassDefaultObject|RF_ArchetypeObject) || 
-						  (Obj->HasAnyFlags(RF_DefaultSubObject) && Obj->GetOuter()->HasAnyFlags(RF_ClassDefaultObject|RF_ArchetypeObject))) && 
-							!FApp::IsGame())
+					UObject* Obj = ObjectNode ? ObjectNode->GetUObject(i) : NULL;
+					if (Obj)
 					{
-						FString OrgArrayContent;
-						Array->ExportText_Direct(OrgArrayContent, Addr, Addr, NULL, 0);
+						if ((Obj->HasAnyFlags(RF_ClassDefaultObject | RF_ArchetypeObject) ||
+							(Obj->HasAnyFlags(RF_DefaultSubObject) && Obj->GetOuter()->HasAnyFlags(RF_ClassDefaultObject | RF_ArchetypeObject))) &&
+							!FApp::IsGame())
+						{
+							FString OrgArrayContent;
+							Array->ExportText_Direct(OrgArrayContent, Addr, Addr, NULL, 0);
 
-						PropertyNodePin->PropagateArrayPropertyChange(Obj, OrgArrayContent, EPropertyArrayChangeType::Clear, -1);
+							PropertyNodePin->PropagateArrayPropertyChange(Obj, OrgArrayContent, EPropertyArrayChangeType::Clear, -1);
+						}
 					}
 
 					FScriptArrayHelper	ArrayHelper(Array,Addr);
@@ -1157,13 +1166,19 @@ void FPropertyValueImpl::ClearChildren()
 
 			if ( bNotifiedPreChange )
 			{
-
-				// send the PostEditChange notification; it will be propagated to all selected objects
 				const bool bTopologyChange = true;
+				// send the PostEditChange notification; it will be propagated to all selected objects
 				FPropertyChangedEvent ChangeEvent(NodeProperty, bTopologyChange, EPropertyChangeType::ValueSet);
 				PropertyNodePin->NotifyPostChange( ChangeEvent, NotifyHook );
 			}
 		}
+
+		if (PropertyUtilities.IsValid())
+		{
+			FPropertyChangedEvent ChangeEvent(NodeProperty, true, EPropertyChangeType::ValueSet);
+			PropertyNodePin->FixPropertiesInEvent(ChangeEvent);
+			PropertyUtilities.Pin()->NotifyFinishedChangingProperties(ChangeEvent);
+		}
 	}
 }
 
@@ -1203,16 +1218,18 @@ void FPropertyValueImpl::InsertChild( TSharedPtr<FPropertyNode> ChildNodeToInser
 		FScriptArrayHelper	ArrayHelper(ArrayProperty,Addr);
 		int32 Index = ChildNodePtr->GetArrayIndex();
 
-		UObject* Obj = ObjectNode->GetUObject(0);
-
-		if ((Obj->HasAnyFlags(RF_ClassDefaultObject|RF_ArchetypeObject) || 
-				(Obj->HasAnyFlags(RF_DefaultSubObject) && Obj->GetOuter()->HasAnyFlags(RF_ClassDefaultObject|RF_ArchetypeObject))) && 
-				!FApp::IsGame())
+		UObject* Obj = ObjectNode ? ObjectNode->GetUObject(0) : NULL;
+		if (Obj)
 		{
-			FString OrgArrayContent;
-			ArrayProperty->ExportText_Direct(OrgArrayContent, Addr, Addr, NULL, 0);
+			if ((Obj->HasAnyFlags(RF_ClassDefaultObject | RF_ArchetypeObject) ||
+				(Obj->HasAnyFlags(RF_DefaultSubObject) && Obj->GetOuter()->HasAnyFlags(RF_ClassDefaultObject | RF_ArchetypeObject))) &&
+				!FApp::IsGame())
+			{
+				FString OrgArrayContent;
+				ArrayProperty->ExportText_Direct(OrgArrayContent, Addr, Addr, NULL, 0);
 
-			ChildNodePtr->PropagateArrayPropertyChange(Obj, OrgArrayContent, EPropertyArrayChangeType::Insert, Index);
+				ChildNodePtr->PropagateArrayPropertyChange(Obj, OrgArrayContent, EPropertyArrayChangeType::Insert, Index);
+			}
 		}
 
 		ArrayHelper.InsertValues(Index, 1 );	
@@ -1227,9 +1244,18 @@ void FPropertyValueImpl::InsertChild( TSharedPtr<FPropertyNode> ChildNodeToInser
 		}
 
 		const bool bTopologyChange = true;
-		FPropertyChangedEvent ChangeEvent(ParentNode->GetProperty(), bTopologyChange, EPropertyChangeType::ArrayAdd);
-		ChangeEvent.SetArrayIndexPerObject(ArrayIndicesPerObject);
-		ChildNodePtr->NotifyPostChange( ChangeEvent, NotifyHook );
+		{
+			FPropertyChangedEvent ChangeEvent(ParentNode->GetProperty(), bTopologyChange, EPropertyChangeType::ArrayAdd);
+			ChangeEvent.SetArrayIndexPerObject(ArrayIndicesPerObject);
+			ChildNodePtr->NotifyPostChange(ChangeEvent, NotifyHook);
+		}
+
+		if (PropertyUtilities.IsValid())
+		{
+			FPropertyChangedEvent ChangeEvent(ParentNode->GetProperty(), bTopologyChange, EPropertyChangeType::ArrayAdd);
+			ChildNodePtr->FixPropertiesInEvent(ChangeEvent);
+			PropertyUtilities.Pin()->NotifyFinishedChangingProperties(ChangeEvent);
+		}
 	}
 }
 
@@ -1271,16 +1297,18 @@ void FPropertyValueImpl::DeleteChild( TSharedPtr<FPropertyNode> ChildNodeToDelet
 				FScriptArrayHelper ArrayHelper(ArrayProperty,Address);
 				int32 Index = ChildNodePtr->GetArrayIndex();
 
-				UObject* Obj = ObjectNode->GetUObject(i);
-
-				if ((Obj->HasAnyFlags(RF_ClassDefaultObject|RF_ArchetypeObject) || 
-						(Obj->HasAnyFlags(RF_DefaultSubObject) && Obj->GetOuter()->HasAnyFlags(RF_ClassDefaultObject|RF_ArchetypeObject))) && 
-						!FApp::IsGame())
+				UObject* Obj = ObjectNode ? ObjectNode->GetUObject(i) : NULL;
+				if (Obj)
 				{
-					FString OrgArrayContent;
-					ArrayProperty->ExportText_Direct(OrgArrayContent, Address, Address, NULL, 0);
+					if ((Obj->HasAnyFlags(RF_ClassDefaultObject | RF_ArchetypeObject) ||
+						(Obj->HasAnyFlags(RF_DefaultSubObject) && Obj->GetOuter()->HasAnyFlags(RF_ClassDefaultObject | RF_ArchetypeObject))) &&
+						!FApp::IsGame())
+					{
+						FString OrgArrayContent;
+						ArrayProperty->ExportText_Direct(OrgArrayContent, Address, Address, NULL, 0);
 
-					ChildNodePtr->PropagateArrayPropertyChange(Obj, OrgArrayContent, EPropertyArrayChangeType::Delete, Index);
+						ChildNodePtr->PropagateArrayPropertyChange(Obj, OrgArrayContent, EPropertyArrayChangeType::Delete, Index);
+					}
 				}
 
 				ArrayHelper.RemoveValues(ChildNodePtr->GetArrayIndex());
@@ -1288,8 +1316,17 @@ void FPropertyValueImpl::DeleteChild( TSharedPtr<FPropertyNode> ChildNodeToDelet
 		}
 
 		const bool bTopologyChange = true;
-		FPropertyChangedEvent ChangeEvent(ParentNode->GetProperty(), bTopologyChange);
-		ChildNodePtr->NotifyPostChange( ChangeEvent, NotifyHook );
+		{
+			FPropertyChangedEvent ChangeEvent(ParentNode->GetProperty(), bTopologyChange);
+			ChildNodePtr->NotifyPostChange(ChangeEvent, NotifyHook);
+		}
+
+		if (PropertyUtilities.IsValid())
+		{
+			FPropertyChangedEvent ChangeEvent(ParentNode->GetProperty(), bTopologyChange);
+			ChildNodePtr->FixPropertiesInEvent(ChangeEvent);
+			PropertyUtilities.Pin()->NotifyFinishedChangingProperties(ChangeEvent);
+		}
 	}
 }
 
@@ -1328,17 +1365,19 @@ void FPropertyValueImpl::DuplicateChild( TSharedPtr<FPropertyNode> ChildNodeToDu
 
 		FScriptArrayHelper	ArrayHelper(ArrayProperty,Addr);
 
-		UObject* Obj = ObjectNode->GetUObject(0);
 		int32 Index = ChildNodePtr->GetArrayIndex();
-
-		if ((Obj->HasAnyFlags(RF_ClassDefaultObject|RF_ArchetypeObject) || 
-				(Obj->HasAnyFlags(RF_DefaultSubObject) && Obj->GetOuter()->HasAnyFlags(RF_ClassDefaultObject|RF_ArchetypeObject))) && 
-				!FApp::IsGame())
+		UObject* Obj = ObjectNode ? ObjectNode->GetUObject(0) : NULL;
+		if (Obj)
 		{
-			FString OrgArrayContent;
-			ArrayProperty->ExportText_Direct(OrgArrayContent, Addr, Addr, NULL, 0);
+			if ((Obj->HasAnyFlags(RF_ClassDefaultObject | RF_ArchetypeObject) ||
+				(Obj->HasAnyFlags(RF_DefaultSubObject) && Obj->GetOuter()->HasAnyFlags(RF_ClassDefaultObject | RF_ArchetypeObject))) &&
+				!FApp::IsGame())
+			{
+				FString OrgArrayContent;
+				ArrayProperty->ExportText_Direct(OrgArrayContent, Addr, Addr, NULL, 0);
 
-			ChildNodePtr->PropagateArrayPropertyChange(Obj, OrgArrayContent, EPropertyArrayChangeType::Duplicate, Index);
+				ChildNodePtr->PropagateArrayPropertyChange(Obj, OrgArrayContent, EPropertyArrayChangeType::Duplicate, Index);
+			}
 		}
 
 		ArrayHelper.InsertValues(Index);
@@ -1348,15 +1387,19 @@ void FPropertyValueImpl::DuplicateChild( TSharedPtr<FPropertyNode> ChildNodeToDu
 		NodeProperty->CopyCompleteValue(ArrayHelper.GetRawPtr(ChildNodePtr->GetArrayIndex()), ArrayHelper.GetRawPtr(ChildNodePtr->GetArrayIndex() + 1));
 
 		// Find the object that owns the array and instance any subobjects
-		FObjectPropertyNode* ObjectPropertyNode = ChildNodePtr->FindObjectItemParent();
-		UObject* ArrayOwner = NULL;
-		for ( TPropObjectIterator Itor( ObjectPropertyNode->ObjectIterator() ) ; Itor && !ArrayOwner ; ++Itor )
+		if (FObjectPropertyNode* ObjectPropertyNode = ChildNodePtr->FindObjectItemParent())
 		{
-			ArrayOwner = Itor->Get();
+			UObject* ArrayOwner = NULL;
+			for (TPropObjectIterator Itor(ObjectPropertyNode->ObjectIterator()); Itor && !ArrayOwner; ++Itor)
+			{
+				ArrayOwner = Itor->Get();
+			}
+			if (ArrayOwner)
+			{
+				ArrayOwner->InstanceSubobjectTemplates();
+			}
 		}
 
-		ArrayOwner->InstanceSubobjectTemplates();
-
 
 
 		//@todo Slate Property Window
@@ -1365,8 +1408,16 @@ void FPropertyValueImpl::DuplicateChild( TSharedPtr<FPropertyNode> ChildNodeToDu
 		// 		ParentNode->SetExpanded(bExpand, bRecurse);
 
 		const bool bTopologyChange = true;
-		FPropertyChangedEvent ChangeEvent(ParentNode->GetProperty(), bTopologyChange, EPropertyChangeType::ValueSet);
-		ChildNodePtr->NotifyPostChange( ChangeEvent, NotifyHook );
+		{
+			FPropertyChangedEvent ChangeEvent(ParentNode->GetProperty(), bTopologyChange, EPropertyChangeType::ValueSet);
+			ChildNodePtr->NotifyPostChange(ChangeEvent, NotifyHook);
+		}
+		if (PropertyUtilities.IsValid())
+		{
+			FPropertyChangedEvent ChangeEvent(ParentNode->GetProperty(), bTopologyChange, EPropertyChangeType::ValueSet);
+			ChildNodePtr->FixPropertiesInEvent(ChangeEvent);
+			PropertyUtilities.Pin()->NotifyFinishedChangingProperties(ChangeEvent);
+		}
 	}
 }
 
@@ -1710,9 +1761,9 @@ FPropertyAccess::Result FPropertyHandleBase::SetPerObjectValues( const TArray<FS
 	TSharedPtr<FPropertyNode> PropertyNode = Implementation->GetPropertyNode();
 	if( PropertyNode.IsValid() && PropertyNode->GetProperty() )
 	{
-		FObjectPropertyNode* ObjectNode = PropertyNode->FindObjectItemParent();
+		FComplexPropertyNode* ComplexNode = PropertyNode->FindComplexParent();
 
-		if( ObjectNode && InPerObjectValues.Num() == ObjectNode->GetNumObjects() )
+		if (ComplexNode && InPerObjectValues.Num() == ComplexNode->GetInstancesNum())
 		{
 			TArray<FObjectBaseAddress> ObjectsToModify;
 			Implementation->GetObjectsToModify( ObjectsToModify, PropertyNode.Get() );
diff --git a/Engine/Source/Editor/PropertyEditor/Private/PropertyHandleImpl.h b/Engine/Source/Editor/PropertyEditor/Private/PropertyHandleImpl.h
--- a/Engine/Source/Editor/PropertyEditor/Private/PropertyHandleImpl.h
+++ b/Engine/Source/Editor/PropertyEditor/Private/PropertyHandleImpl.h
@@ -8,14 +8,17 @@ public:
 	FObjectBaseAddress()
 		:	Object( NULL )
 		,	BaseAddress( NULL )
+		,	bIsStruct(false)
 	{}
-	FObjectBaseAddress(UObject* InObject, uint8* InBaseAddress)
+	FObjectBaseAddress(UObject* InObject, uint8* InBaseAddress, bool InIsStruct)
 		:	Object( InObject )
 		,	BaseAddress( InBaseAddress )
+		,	bIsStruct(InIsStruct)
 	{}
 
 	UObject*	Object;
 	uint8*		BaseAddress;
+	bool		bIsStruct;
 };
 
 
diff --git a/Engine/Source/Editor/PropertyEditor/Private/PropertyNode.cpp b/Engine/Source/Editor/PropertyEditor/Private/PropertyNode.cpp
--- a/Engine/Source/Editor/PropertyEditor/Private/PropertyNode.cpp
+++ b/Engine/Source/Editor/PropertyEditor/Private/PropertyNode.cpp
@@ -226,14 +226,14 @@ void FPropertyNode::ClearCachedReadAddresses( bool bRecursive )
 
 
 // Follows the chain of items upwards until it finds the object window that houses this item.
-FObjectPropertyNode* FPropertyNode::FindObjectItemParent()
+FComplexPropertyNode* FPropertyNode::FindComplexParent()
 {
 	FPropertyNode* Cur = this;
-	FObjectPropertyNode* Found = NULL;
+	FComplexPropertyNode* Found = NULL;
 
 	while( true )
 	{
-		Found = Cur->AsObjectNode();
+		Found = Cur->AsComplexNode();
 		if( Found )
 		{
 			break;
@@ -251,14 +251,14 @@ FObjectPropertyNode* FPropertyNode::FindObjectItemParent()
 
 
 // Follows the chain of items upwards until it finds the object window that houses this item.
-const FObjectPropertyNode* FPropertyNode::FindObjectItemParent() const
+const FComplexPropertyNode* FPropertyNode::FindComplexParent() const
 {
 	const FPropertyNode* Cur = this;
-	const FObjectPropertyNode* Found = NULL;
+	const FComplexPropertyNode* Found = NULL;
 
 	while( true )
 	{
-		Found = Cur->AsObjectNode();
+		Found = Cur->AsComplexNode();
 		if( Found )
 		{
 			break;
@@ -275,6 +275,17 @@ const FObjectPropertyNode* FPropertyNode::FindObjectItemParent() const
 	return Found;
 }
 
+class FObjectPropertyNode* FPropertyNode::FindObjectItemParent()
+{
+	auto ComplexParent = FindComplexParent();
+	return ComplexParent ? ComplexParent->AsObjectNode() : NULL;
+}
+
+const class FObjectPropertyNode* FPropertyNode::FindObjectItemParent() const
+{
+	const auto ComplexParent = FindComplexParent();
+	return ComplexParent ? ComplexParent->AsObjectNode() : NULL;
+}
 
 /**
  * Follows the top-most object window that contains this property window item.
@@ -2017,12 +2028,12 @@ TSharedRef<FEditPropertyChain> FPropertyNode::BuildPropertyChain( UProperty* InP
 
 	FPropertyNode* ItemNode = this;
 
-	FObjectPropertyNode* ObjectNode = FindObjectItemParent();
+	FComplexPropertyNode* ComplexNode = FindComplexParent();
 	UProperty* MemberProperty = InProperty;
 
 	do
 	{
-		if ( ItemNode == ObjectNode )
+		if (ItemNode == ComplexNode)
 		{
 			MemberProperty = PropertyChain->GetHead()->GetValue();
 		}
@@ -2042,7 +2053,7 @@ TSharedRef<FEditPropertyChain> FPropertyNode::BuildPropertyChain( UProperty* InP
 	while( ItemNode != NULL );
 
 	// If the modified property was a property of the object at the root of this property window, the member property will not have been set correctly
-	if ( ItemNode == ObjectNode )
+	if (ItemNode == ComplexNode)
 	{
 		MemberProperty = PropertyChain->GetHead()->GetValue();
 	}
@@ -2053,6 +2064,20 @@ TSharedRef<FEditPropertyChain> FPropertyNode::BuildPropertyChain( UProperty* InP
 	return PropertyChain;
 }
 
+FPropertyChangedEvent& FPropertyNode::FixPropertiesInEvent(FPropertyChangedEvent& Event)
+{
+	ensure(Event.Property);
+
+	auto PropertyChain = BuildPropertyChain(Event.Property);
+	auto MemberProperty = PropertyChain->GetActiveMemberNode() ? PropertyChain->GetActiveMemberNode()->GetValue() : NULL;
+	if (ensure(MemberProperty))
+	{
+		Event.SetActiveMemberProperty(MemberProperty);
+	}
+
+	return Event;
+}
+
 /**
  * Does the string compares to ensure this Name is acceptable to the filter that is passed in
  * @return		Return True if this property should be displayed.  False if it should be culled
diff --git a/Engine/Source/Editor/PropertyEditor/Private/PropertyNode.h b/Engine/Source/Editor/PropertyEditor/Private/PropertyNode.h
--- a/Engine/Source/Editor/PropertyEditor/Private/PropertyNode.h
+++ b/Engine/Source/Editor/PropertyEditor/Private/PropertyNode.h
@@ -70,12 +70,14 @@ class FPropertyNode;
 
 struct FAddressPair
 {
-	FAddressPair( const UObject* InObject, uint8* Address )
+	FAddressPair(const UObject* InObject, uint8* Address, bool bInIsStruct)
 		: Object( InObject)
 		, ReadAddress( Address )
+		, bIsStruct(bInIsStruct)
 	{}
 	TWeakObjectPtr<const UObject> Object;
 	uint8* ReadAddress;
+	bool bIsStruct;
 };
 
 template<> struct TIsPODType<FAddressPair> { enum { Value = true }; };
@@ -88,9 +90,9 @@ public:
 		, bRequiresCache( true )
 	{
 	}
-	void Add( const UObject* Object, uint8* Address )
+	void Add(const UObject* Object, uint8* Address, bool bIsStruct = false)
 	{
-		ReadAddresses.Add( FAddressPair( Object, Address ) );
+		ReadAddresses.Add(FAddressPair(Object, Address, bIsStruct));
 	}
 
 	int32 Num() const
@@ -101,7 +103,7 @@ public:
 	uint8* GetAddress( int32 Index )
 	{
 		const FAddressPair& Pair = ReadAddresses[Index];
-		return Pair.Object.IsValid() ? Pair.ReadAddress : NULL;
+		return (Pair.Object.IsValid() || Pair.bIsStruct) ? Pair.ReadAddress : NULL;
 	}
 
 	void Reset()
@@ -111,9 +113,10 @@ public:
 		bRequiresCache = true;
 	}
 	
-	TArray<FAddressPair> ReadAddresses;
 	bool bAllValuesTheSame;
 	bool bRequiresCache;
+private:
+	TArray<FAddressPair> ReadAddresses;
 };
 
 /**
@@ -129,7 +132,7 @@ public:
 
 	int32 Num() const
 	{
-		return (ReadAddressListData != NULL) ? ReadAddressListData->ReadAddresses.Num() : 0;
+		return (ReadAddressListData != NULL) ? ReadAddressListData->Num() : 0;
 	}
 
 	uint8* GetAddress( int32 Index )
@@ -200,6 +203,8 @@ struct EPropertyArrayChangeType
 	};
 };
 
+class FComplexPropertyNode;
+
 /**
  * The base class for all property nodes                                                              
  */
@@ -256,6 +261,9 @@ public:
 	virtual class FObjectPropertyNode* AsObjectNode() { return NULL; }
 	virtual const FObjectPropertyNode* AsObjectNode() const { return NULL; }
 
+	virtual FComplexPropertyNode* AsComplexNode() { return NULL; }
+	virtual const FComplexPropertyNode* AsComplexNode() const { return NULL; }
+
 	/**
 	 * Interface function to get at the dervied FCategoryPropertyNodeWx class
 	 */
@@ -269,6 +277,9 @@ public:
 	/**
 	 * Follows the chain of items upwards until it finds the object window that houses this item.
 	 */
+	class FComplexPropertyNode* FindComplexParent();
+	const class FComplexPropertyNode* FindComplexParent() const;
+
 	class FObjectPropertyNode* FindObjectItemParent();
 	const class FObjectPropertyNode* FindObjectItemParent() const;
 
@@ -634,6 +645,8 @@ public:
 		return Restrictions;
 	}
 
+	FPropertyChangedEvent& FixPropertiesInEvent(FPropertyChangedEvent& Event);
+
 protected:
 
 	TSharedRef<FEditPropertyChain> BuildPropertyChain( UProperty* PropertyAboutToChange );
@@ -752,4 +765,29 @@ protected:
 	TArray<TSharedRef<const class FPropertyRestriction>> Restrictions;
 };
 
+class FComplexPropertyNode : public FPropertyNode
+{
+public:
 
+	enum EPropertyType
+	{
+		EPT_Object,
+		EPT_StandaloneStructure,
+	};
+
+	FComplexPropertyNode() : FPropertyNode() {}
+	virtual ~FComplexPropertyNode() {}
+
+	virtual FComplexPropertyNode* AsComplexNode() override{ return this; }
+	virtual const FComplexPropertyNode* AsComplexNode() const override{ return this; }
+
+	virtual UStruct* GetBaseStructure() = 0;
+	virtual const UStruct* GetBaseStructure() const = 0;
+
+	virtual int32 GetInstancesNum() const = 0;
+	virtual uint8* GetMemoryOfInstance(int32 Index) = 0;
+	virtual TWeakObjectPtr<UObject> GetInstanceAsUObject(int32 Index) = 0;
+	virtual EPropertyType GetPropertyType() const = 0;
+
+	virtual void Disconnect() = 0;
+};
diff --git a/Engine/Source/Editor/PropertyEditor/Private/SDetailsView.cpp b/Engine/Source/Editor/PropertyEditor/Private/SDetailsView.cpp
--- a/Engine/Source/Editor/PropertyEditor/Private/SDetailsView.cpp
+++ b/Engine/Source/Editor/PropertyEditor/Private/SDetailsView.cpp
@@ -207,28 +207,6 @@ TSharedRef<SDetailTree> SDetailsView::ConstructTreeView( TSharedRef<SScrollBar>&
 	.ExternalScrollbar( ScrollBar );
 }
 
-void SDetailsView::SetColorPropertyFromColorPicker(FLinearColor NewColor)
-{
-	const TSharedPtr< FPropertyNode > PinnedColorPropertyNode = ColorPropertyNode.Pin();
-	if( ensure(PinnedColorPropertyNode.IsValid()) )
-	{
-		UProperty* Property = PinnedColorPropertyNode->GetProperty();
-		check(Property);
-
-		FObjectPropertyNode* ObjectNode = PinnedColorPropertyNode->FindObjectItemParent();
-
-		if( ObjectNode && ObjectNode->GetNumObjects() )
-		{
-			FScopedTransaction Transaction( NSLOCTEXT("UnrealEd", "SetColorProperty", "Set Color Property") );
-
-			PinnedColorPropertyNode->NotifyPreChange(Property, GetNotifyHook());
-
-			FPropertyChangedEvent ChangeEvent(Property, false, EPropertyChangeType::ValueSet);
-			PinnedColorPropertyNode->NotifyPostChange( ChangeEvent, GetNotifyHook() );
-		}
-	}
-}
-
 FReply SDetailsView::OnOpenRawPropertyEditorClicked()
 {
 	FPropertyEditorModule& PropertyEditorModule = FModuleManager::LoadModuleChecked<FPropertyEditorModule>( "PropertyEditor" );
@@ -243,222 +221,6 @@ EVisibility SDetailsView::GetActorNameAreaVisibility() const
 	return bVisible ? EVisibility::Visible : EVisibility::Collapsed; 
 }
 
-EVisibility SDetailsView::GetFilterBoxVisibility() const
-{
-	// Visible if we allow search and we have anything to search otherwise collapsed so it doesn't take up room
-	return DetailsViewArgs.bAllowSearch && RootPropertyNode->GetNumObjects() > 0 ? EVisibility::Visible : EVisibility::Collapsed;
-}
-
-/**
- * Determines whether or not a property should be visible in the default generated detail layout
- *
- * @param PropertyNode	The property node to check
- * @param ParentNode	The parent property node to check
- * @return true if the property should be visible
- */
-static bool IsVisibleStandaloneProperty( const FPropertyNode& PropertyNode, const FPropertyNode& ParentNode )
-{
-	const UProperty* Property = PropertyNode.GetProperty();
-	const UArrayProperty* ParentArrayProperty = Cast<const UArrayProperty>( ParentNode.GetProperty() );
-	bool bIsVisibleStandalone = false;
-	if( Property )
-	{
-		if( Property->IsA( UObjectPropertyBase::StaticClass() ) )
-		{
-			// Do not add this child node to the current map if its a single object property in a category (serves no purpose for UI)
-			bIsVisibleStandalone = !ParentArrayProperty && (PropertyNode.GetNumChildNodes() == 0 || PropertyNode.GetNumChildNodes() > 1) ;
-		}
-		else if( Property->IsA( UArrayProperty::StaticClass() ) || Property->ArrayDim > 1 && PropertyNode.GetArrayIndex() == INDEX_NONE )
-		{
-			// Base array properties are always visible
-			bIsVisibleStandalone = true;
-		}
-		else
-		{
-			bIsVisibleStandalone = true;
-		}
-		
-	}
-
-	return bIsVisibleStandalone;
-}
-
-void SDetailsView::UpdatePropertyMapRecursive( FPropertyNode& InNode, FDetailLayoutBuilderImpl& InDetailLayout, FName CurCategory, FObjectPropertyNode* CurObjectNode )
-{
-	UProperty* ParentProperty = InNode.GetProperty();
-	UStructProperty* ParentStructProp = Cast<UStructProperty>(ParentProperty);
-	
-	for( int32 ChildIndex = 0; ChildIndex < InNode.GetNumChildNodes(); ++ChildIndex )
-	{
-		TSharedPtr<FPropertyNode> ChildNodePtr = InNode.GetChildNode(ChildIndex);
-		FPropertyNode& ChildNode = *ChildNodePtr;
-		UProperty* Property = ChildNode.GetProperty();
-
-		{
-			FObjectPropertyNode* ObjNode = ChildNode.AsObjectNode();
-			FCategoryPropertyNode* CategoryNode = ChildNode.AsCategoryNode();
-			if( ObjNode )
-			{
-				// Currently object property nodes do not provide any useful information other than being a container for its children.  We do not draw anything for them.
-				// When we encounter object property nodes, add their children instead of adding them to the tree.
-				UpdatePropertyMapRecursive( ChildNode, InDetailLayout, CurCategory, ObjNode );
-			}
-			else if( CategoryNode )
-			{
-				// For category nodes, we just set the current category and recurse through the children
-				UpdatePropertyMapRecursive( ChildNode, InDetailLayout, CategoryNode->GetCategoryName(), CurObjectNode );
-			}
-			else
-			{
-				// Whether or not the property can be visible in the default detail layout
-				bool bVisibleByDefault = IsVisibleStandaloneProperty( ChildNode, InNode );
-
-				// Whether or not the property is a struct
-				UStructProperty* StructProperty = Cast<UStructProperty>( Property );
-
-				bool bIsStruct = StructProperty != NULL;
-
-				static FName ShowOnlyInners("ShowOnlyInnerProperties");
-
-				bool bIsChildOfCustomizedStruct = false;
-				bool bIsCustomizedStruct = false;
-
-				const UStruct* Struct = StructProperty ? StructProperty->Struct : NULL;
-				const UStruct* ParentStruct = ParentStructProp ? ParentStructProp->Struct : NULL;
-				if (Struct || ParentStruct)
-				{
-					FPropertyEditorModule& ParentPlugin = FModuleManager::GetModuleChecked<FPropertyEditorModule>("PropertyEditor");
-					if (Struct)
-					{
-						bIsCustomizedStruct = ParentPlugin.IsCustomizedStruct(Struct);
-					}
-
-					if (ParentStruct)
-					{
-						bIsChildOfCustomizedStruct = ParentPlugin.IsCustomizedStruct(ParentStruct);
-					}
-				}
-
-				// Whether or not to push out struct properties to their own categories or show them inside an expandable struct 
-				bool bPushOutStructProps = bIsStruct && !bIsCustomizedStruct && !ParentStructProp && Property->HasMetaData(ShowOnlyInners);
-
-				// Is the property edit inline new 
-				const bool bIsEditInlineNew = SPropertyEditorEditInline::Supports( &ChildNode, ChildNode.GetArrayIndex() );
-
-				// Is this a property of an array
-				bool bIsChildOfArray = PropertyEditorHelpers::IsChildOfArray( ChildNode );
-
-				// Edit inline new properties should be visible by default
-				bVisibleByDefault |= bIsEditInlineNew;
-
-				// Children of arrays are not visible directly,
-				bVisibleByDefault &= !bIsChildOfArray;
-
-				// See if the user requested specific property visibility 
-				const bool bIsUserVisible = IsPropertyVisible( Property ); 
-
-				// Inners of customized in structs should not be taken into consideration for customizing.  They are not designed to be individually customized when their parent is already customized
-				if( !bIsChildOfCustomizedStruct )
-				{
-					// Add any object classes with properties so we can ask them for custom property layouts later
-					ClassesWithProperties.Add( Property->GetOwnerStruct() );
-				}
-
-				// If there is no outer object then the class is the object root and there is only one instance
-				FName InstanceName = NAME_None;
-				if( CurObjectNode && CurObjectNode->GetParentNode() )
-				{
-					InstanceName = CurObjectNode->GetParentNode()->GetProperty()->GetFName();
-				}
-				else if( ParentStructProp )
-				{
-					InstanceName = ParentStructProp->GetFName();
-				}
-
-				// Do not add children of customized in struct properties or arrays
-				if( !bIsChildOfCustomizedStruct && !bIsChildOfArray )
-				{
-					// Get the class property map
-					FClassInstanceToPropertyMap& ClassInstanceMap = ClassToPropertyMap.FindOrAdd( Property->GetOwnerStruct()->GetFName() );
-
-					FPropertyNodeMap& PropertyNodeMap = ClassInstanceMap.FindOrAdd( InstanceName );
-
-					if( !PropertyNodeMap.ParentObjectProperty )
-					{
-						PropertyNodeMap.ParentObjectProperty = CurObjectNode;
-					}
-					else
-					{
-						ensure( PropertyNodeMap.ParentObjectProperty == CurObjectNode );
-					}
-
-					checkSlow( !PropertyNodeMap.Contains( Property->GetFName() ) );
-
-					PropertyNodeMap.Add( Property->GetFName(), ChildNodePtr );
-				}
-
-				if( bVisibleByDefault && bIsUserVisible && !bPushOutStructProps )
-				{
-					FName CategoryName = CurCategory;
-
-					// For properties inside a struct, add them to their own category unless they just take the name of the parent struct.  
-					// In that case push them to the parent category
-					FName PropertyCatagoryName = FObjectEditorUtils::GetCategoryFName(Property);
-					if( !ParentStructProp || (PropertyCatagoryName != ParentStructProp->Struct->GetFName()) )
-					{
-						CategoryName = PropertyCatagoryName;
-					}
-
-					// Add a property to the default category
-					FDetailCategoryImpl& CategoryImpl = InDetailLayout.DefaultCategory( CategoryName );
-					CategoryImpl.AddPropertyNode( ChildNodePtr.ToSharedRef(), InstanceName );
-				}
-
-				bool bRecurseIntoChildren = 
-						!bIsChildOfCustomizedStruct // Don't recurse into built in struct children, we already know what they are and how to display them
-					&&  !bIsCustomizedStruct // Don't recurse into customized structs
-					&&	!bIsChildOfArray // Do not recurse into arrays, the children are drawn by the array property parent
-					&&	!bIsEditInlineNew // Edit inline new children are not supported for customization yet
-					&&	bIsUserVisible // Properties must be allowed to be visible by a user if they are not then their children are not visible either
-					&&  (!bIsStruct || bPushOutStructProps); //  Only recurse into struct properties if they are going to be displayed as standalone properties in categories instead of inside an expandable area inside a category
-
-				if( bRecurseIntoChildren )
-				{
-					// Built in struct properties or children of arras 
-					UpdatePropertyMapRecursive( ChildNode, InDetailLayout, CurCategory, CurObjectNode );
-				}
-			}		
-		}
-	}
-}
-
-void SDetailsView::UpdatePropertyMap()
-{
-	check( RootPropertyNode.IsValid() );
-
-	// Reset everything
-	ClassToPropertyMap.Empty();
-	ClassesWithProperties.Empty();
-
-	// We need to be able to create a new detail layout and properly clean up the old one in the process
-	check( !DetailLayout.IsValid() || DetailLayout.IsUnique() );
-
-	RootTreeNodes.Empty();
-
-	DetailLayout = MakeShareable( new FDetailLayoutBuilderImpl( ClassToPropertyMap, PropertyUtilities.ToSharedRef(), SharedThis( this ) ) );
-
-	// Currently object property nodes do not provide any useful information other than being a container for its children.  We do not draw anything for them.
-	// When we encounter object property nodes, add their children instead of adding them to the tree.
-	UpdatePropertyMapRecursive( *RootPropertyNode, *DetailLayout, NAME_None, RootPropertyNode.Get() );
-
-	// Ask for custom detail layouts
-	QueryCustomDetailLayout( *DetailLayout );
-
-	DetailLayout->GenerateDetailLayout();
-
-	UpdateFilteredDetails();
-}
-
 void SDetailsView::ForceRefresh()
 {
 	TArray< TWeakObjectPtr< UObject > > NewObjectList;
@@ -746,126 +508,6 @@ void SDetailsView::RemoveInvalidActors( const TSet<AActor*>& ValidActors )
 	}
 }
 
-/**
- * Recursively gets expanded items for a node
- * 
- * @param InPropertyNode			The node to get expanded items from
- * @param OutExpandedItems	List of expanded items that were found
- */
-void GetExpandedItems( TSharedPtr<FPropertyNode> InPropertyNode, TArray<FString>& OutExpandedItems )
-{
-	if( InPropertyNode->HasNodeFlags(EPropertyNodeFlags::Expanded) )
-	{
-		const bool bWithArrayIndex = true;
-		FString Path;
-		Path.Empty(128);
-		InPropertyNode->GetQualifiedName(Path, bWithArrayIndex);
-
-		OutExpandedItems.Add( Path );
-	}
-
-	for( int32 ChildIndex = 0; ChildIndex < InPropertyNode->GetNumChildNodes(); ++ChildIndex )
-	{
-		GetExpandedItems( InPropertyNode->GetChildNode(ChildIndex), OutExpandedItems );
-	}
-
-}
-
-/**
- * Recursively sets expanded items for a node
- * 
- * @param InNode			The node to set expanded items on
- * @param OutExpandedItems	List of expanded items to set
- */
-void SetExpandedItems( TSharedPtr<FPropertyNode> InPropertyNode, const TArray<FString>& InExpandedItems )
-{
-	if( InExpandedItems.Num() > 0 )
-	{
-		const bool bWithArrayIndex = true;
-		FString Path;
-		Path.Empty(128);
-		InPropertyNode->GetQualifiedName(Path, bWithArrayIndex);
-
-		for (int32 ItemIndex = 0; ItemIndex < InExpandedItems.Num(); ++ItemIndex)
-		{
-			if ( InExpandedItems[ItemIndex] == Path )
-			{
-				InPropertyNode->SetNodeFlags( EPropertyNodeFlags::Expanded, true );
-				break;
-			}
-		}
-
-		for( int32 NodeIndex = 0; NodeIndex < InPropertyNode->GetNumChildNodes(); ++NodeIndex )
-		{
-			SetExpandedItems( InPropertyNode->GetChildNode( NodeIndex ), InExpandedItems );
-		}
-	}
-}
-
-void SDetailsView::SaveExpandedItems()
-{
-	UClass* BestBaseClass = RootPropertyNode->GetObjectBaseClass();
-
-	TArray<FString> ExpandedPropertyItems;
-	GetExpandedItems( RootPropertyNode, ExpandedPropertyItems );
-
-	TArray<FString> ExpandedCustomItems = ExpandedDetailNodes.Array();
-
-	// Expanded custom items may have spaces but SetSingleLineArray doesnt support spaces (treats it as another element in the array)
-	// Append a '|' after each element instead
-	FString ExpandedCustomItemsString;
-	for( auto It = ExpandedDetailNodes.CreateConstIterator(); It; ++It )
-	{
-		ExpandedCustomItemsString += *It;
-		ExpandedCustomItemsString += TEXT(",");
-	}
-
-	//while a valid class, and we're either the same as the base class (for multiple actors being selected and base class is AActor) OR we're not down to AActor yet)
-	for( UClass* Class = BestBaseClass; Class && ((BestBaseClass == Class) || (Class!=AActor::StaticClass())); Class = Class->GetSuperClass() )
-	{
-		if( RootPropertyNode->GetNumChildNodes() > 0 && ExpandedPropertyItems.Num() > 0 )
-		{
-			GConfig->SetSingleLineArray(TEXT("DetailPropertyExpansion"), *Class->GetName(), ExpandedPropertyItems, GEditorUserSettingsIni);
-		}
-	}
-
-	if( DetailLayout.IsValid() && BestBaseClass && !ExpandedCustomItemsString.IsEmpty() )
-	{
-		GConfig->SetString(TEXT("DetailCustomWidgetExpansion"), *BestBaseClass->GetName(), *ExpandedCustomItemsString, GEditorUserSettingsIni);
-	}
-}
-
-void SDetailsView::RestoreExpandedItems( TSharedPtr<FPropertyNode> InitialStartNode )
-{
-	TSharedPtr<FPropertyNode> StartNode = InitialStartNode;
-	if( !StartNode.IsValid() )
-	{
-		StartNode = RootPropertyNode;
-	}
-
-	ExpandedDetailNodes.Empty();
-
-	TArray<FString> ExpandedPropertyItems;
-	FString ExpandedCustomItems;
-	UClass* BestBaseClass = RootPropertyNode->GetObjectBaseClass();
-	//while a valid class, and we're either the same as the base class (for multiple actors being selected and base class is AActor) OR we're not down to AActor yet)
-	for( UClass* Class = BestBaseClass; Class && ((BestBaseClass == Class) || (Class!=AActor::StaticClass())); Class = Class->GetSuperClass() )
-	{
-		GConfig->GetSingleLineArray(TEXT("DetailPropertyExpansion"), *Class->GetName(), ExpandedPropertyItems, GEditorUserSettingsIni);
-		SetExpandedItems( StartNode, ExpandedPropertyItems );
-	}
-
-	if( BestBaseClass )
-	{
-		GConfig->GetString(TEXT("DetailCustomWidgetExpansion"), *BestBaseClass->GetName(), ExpandedCustomItems, GEditorUserSettingsIni);
-		TArray<FString> ExpandedCustomItemsArray;
-		ExpandedCustomItems.ParseIntoArray( &ExpandedCustomItemsArray, TEXT(","), true );
-
-		ExpandedDetailNodes.Append( ExpandedCustomItemsArray );
-	}
-
-}
-
 /** Called before during SetObjectArray before we change the objects being observed */
 void SDetailsView::PreSetObject()
 {
@@ -876,12 +518,13 @@ void SDetailsView::PreSetObject()
 
 	RootNodePendingKill = RootPropertyNode;
 
-	RootPropertyNode = MakeShareable( new FObjectPropertyNode );
+	RootPropertyNode = MakeShareable(new FObjectPropertyNode);
 
 	SelectedActors.Empty();
 	SelectedObjects.Empty();
 }
 
+
 /** Called at the end of SetObjectArray after we change the objects being observed */
 void SDetailsView::PostSetObject()
 {
@@ -906,363 +549,6 @@ void SDetailsView::PostSetObject()
 	UpdatePropertyMap();
 }
 
-void SDetailsView::QueryLayoutForClass( FDetailLayoutBuilderImpl& CustomDetailLayout, UStruct* Class )
-{
-	CustomDetailLayout.SetCurrentCustomizationClass( CastChecked<UClass>( Class ), NAME_None );
-
-	FPropertyEditorModule& ParentPlugin = FModuleManager::GetModuleChecked<FPropertyEditorModule>("PropertyEditor");
-	FCustomDetailLayoutNameMap& GlobalCustomLayoutNameMap = ParentPlugin.ClassNameToDetailLayoutNameMap;
-
-	// Check the instanced map first
-	FDetailLayoutCallback* Callback = InstancedClassToDetailLayoutMap.Find( TWeakObjectPtr<UStruct>(Class) );
-
-	if( !Callback )
-	{
-		// callback wasn't found in the per instance map, try the global instances instead
-		Callback = GlobalCustomLayoutNameMap.Find( Class->GetFName() );
-	}
-
-	if( Callback && Callback->DetailLayoutDelegate.IsBound() )
-	{
-		// Create a new instance of the custom detail layout for the current class
-		TSharedRef<IDetailCustomization> CustomizationInstance = Callback->DetailLayoutDelegate.Execute();
-
-		// Ask for details immediately
-		CustomizationInstance->CustomizeDetails( CustomDetailLayout );
-
-		// Save the instance from destruction until we refresh
-		CustomizationClassInstances.Add( CustomizationInstance );
-	}
-}
-
-void SDetailsView::QueryCustomDetailLayout( FDetailLayoutBuilderImpl& CustomDetailLayout )
-{
-	FPropertyEditorModule& ParentPlugin = FModuleManager::GetModuleChecked<FPropertyEditorModule>("PropertyEditor");
-
-	// Get the registered classes that customize details
-	FCustomDetailLayoutNameMap& GlobalCustomLayoutNameMap = ParentPlugin.ClassNameToDetailLayoutNameMap;
-	
-	UClass* BaseClass = GetBaseClass();
-
-	// All the current customization instances need to be deleted when it is safe
-	CustomizationClassInstancesPendingDelete = CustomizationClassInstances;
-
-	CustomizationClassInstances.Empty();
-
-	//Ask for generic details not specific to an object being viewed 
-	if( GenericLayoutDelegate.IsBound() )
-	{
-		// Create a new instance of the custom detail layout for the current class
-		TSharedRef<IDetailCustomization> CustomizationInstance = GenericLayoutDelegate.Execute();
-
-		// Ask for details immediately
-		CustomizationInstance->CustomizeDetails( CustomDetailLayout );
-
-		// Save the instance from destruction until we refresh
-		CustomizationClassInstances.Add( CustomizationInstance );
-	}
-
-
-	// Sort them by query order.  @todo not good enough
-	struct FCompareFDetailLayoutCallback
-	{
-		FORCEINLINE bool operator()( const FDetailLayoutCallback& A, const FDetailLayoutCallback& B ) const
-		{
-			return A.Order < B.Order;
-		}
-	};
-
-	TMap< TWeakObjectPtr<UStruct>, FDetailLayoutCallback*> FinalCallbackMap;
-
-	for( auto ClassIt = ClassesWithProperties.CreateConstIterator(); ClassIt; ++ClassIt )
-	{
-		// Check the instanced map first
-		FDetailLayoutCallback* Callback = InstancedClassToDetailLayoutMap.Find( *ClassIt );
-
-		if( !Callback )
-		{
-			// callback wasn't found in the per instance map, try the global instances instead
-			Callback = GlobalCustomLayoutNameMap.Find( (*ClassIt)->GetFName() );
-		}
-
-		if( Callback )
-		{
-			FinalCallbackMap.Add( *ClassIt, Callback );
-		}
-	}
-
-
-	FinalCallbackMap.ValueSort( FCompareFDetailLayoutCallback() );
-
-	TSet<UStruct*> QueriedClasses;
-
-	if( FinalCallbackMap.Num() > 0 )
-	{
-		// Ask each class that we have properties for to customize its layout
-		for( auto LayoutIt(FinalCallbackMap.CreateConstIterator()); LayoutIt; ++LayoutIt )
-		{
-			const TWeakObjectPtr<UStruct> WeakClass = LayoutIt.Key();
-
-			if( WeakClass.IsValid() )
-			{
-				UStruct* Class = WeakClass.Get();
-
-				FClassInstanceToPropertyMap& InstancedPropertyMap = ClassToPropertyMap.FindChecked( Class->GetFName() );
-				for( FClassInstanceToPropertyMap::TIterator InstanceIt(InstancedPropertyMap); InstanceIt; ++InstanceIt )
-				{
-					CustomDetailLayout.SetCurrentCustomizationClass( CastChecked<UClass>( Class ), InstanceIt.Key() );
-
-					const FOnGetDetailCustomizationInstance& DetailDelegate = LayoutIt.Value()->DetailLayoutDelegate;
-
-					if( DetailDelegate.IsBound() )
-					{
-						QueriedClasses.Add( Class );
-
-						// Create a new instance of the custom detail layout for the current class
-						TSharedRef<IDetailCustomization> CustomizationInstance = DetailDelegate.Execute();
-							
-						// Ask for details immediately
-						CustomizationInstance->CustomizeDetails( CustomDetailLayout );
-
-						// Save the instance from destruction until we refresh
-						CustomizationClassInstances.Add( CustomizationInstance );
-					}
-				}
-			}
-		}
-	}
-
-	// Ensure that the base class and its parents are always queried
-	TSet<UStruct*> ParentClassesToQuery;
-	if( BaseClass && !QueriedClasses.Contains( BaseClass ) )
-	{
-		ParentClassesToQuery.Add( BaseClass );
-		ClassesWithProperties.Add( BaseClass );
-	}
-
-	// Find base classes of queried classes that were not queried and add them to the query list
-	// this supports cases where a parent class has no properties but still wants to add customization
-	for( auto QueriedClassIt = ClassesWithProperties.CreateConstIterator();QueriedClassIt; ++QueriedClassIt )
-	{
-		UStruct* ParentStruct = (*QueriedClassIt)->GetSuperStruct();
-		
-		while( ParentStruct && ParentStruct->IsA( UClass::StaticClass() ) && !QueriedClasses.Contains(ParentStruct)  && !ClassesWithProperties.Contains( ParentStruct ) )
-		{
-			ParentClassesToQuery.Add( ParentStruct );
-			ParentStruct = ParentStruct->GetSuperStruct();
-
-		}
-	}
-
-	// Query extra base classes
-	for( auto ParentIt = ParentClassesToQuery.CreateConstIterator(); ParentIt; ++ParentIt )
-	{
-		QueryLayoutForClass( CustomDetailLayout, *ParentIt );
-	}
-}
-
-void SDetailsView::UpdateFilteredDetails()
-{
-	RootPropertyNode->FilterNodes( CurrentFilter.FilterStrings );
-	RootPropertyNode->ProcessSeenFlags( true );
-
-	for( int32 NodeIndex = 0; NodeIndex < ExternalRootPropertyNodes.Num(); ++NodeIndex )
-	{
-		TSharedPtr<FObjectPropertyNode> ObjectNode = ExternalRootPropertyNodes[NodeIndex].Pin();
-
-		if( ObjectNode.IsValid() )
-		{
-			ObjectNode->FilterNodes( CurrentFilter.FilterStrings );
-			ObjectNode->ProcessSeenFlags( true );
-		}
-	}
-
-	if( DetailLayout.IsValid() )
-	{
-		DetailLayout->FilterDetailLayout( CurrentFilter );
-	}
-
-	RootTreeNodes = DetailLayout->GetRootTreeNodes();
-
-	DetailTree->RequestTreeRefresh();
-}
-
-/** Ticks the property view.  This function performs a data consistency check */
-void SDetailsView::Tick( const FGeometry& AllottedGeometry, const double InCurrentTime, const float InDeltaTime )
-{
-	for( int32 i = 0; i < CustomizationClassInstancesPendingDelete.Num(); ++i )
-	{
-		ensure( CustomizationClassInstancesPendingDelete[i].IsUnique() );
-	}
-
-	if( RootNodePendingKill.IsValid() )
-	{
-		RootNodePendingKill->RemoveAllObjects();
-		RootNodePendingKill.Reset();
-	}
-
-	// Empty all the customization instances that need to be deleted
-	CustomizationClassInstancesPendingDelete.Empty();
-
-	// Purge any objects that are marked pending kill from the object list
-	RootPropertyNode->PurgeKilledObjects();
-
-	if( DeferredActions.Num() > 0 )
-	{
-		// Any deferred actions are likely to cause the node  tree to be at least partially rebuilt
-		// Save the expansion state of existing nodes so we can expand them later
-		SaveExpandedItems();
-
-		// Execute any deferred actions
-		for( int32 ActionIndex = 0; ActionIndex < DeferredActions.Num(); ++ActionIndex )
-		{
-			DeferredActions[ActionIndex].ExecuteIfBound();
-		}
-		DeferredActions.Empty();
-	}
-
-	bool bValidateExternalNodes = true;
-	FPropertyNode::DataValidationResult Result = RootPropertyNode->EnsureDataIsValid();
-	if( Result == FPropertyNode::PropertiesChanged || Result == FPropertyNode::EditInlineNewValueChanged  )
-	{
-		RestoreExpandedItems();
-		UpdatePropertyMap();
-	}
-	else if(Result == FPropertyNode::ArraySizeChanged  )
-	{
-		RestoreExpandedItems();
-		UpdateFilteredDetails();
-	}
-	else if( Result == FPropertyNode::ObjectInvalid )
-	{
-		TArray< TWeakObjectPtr< UObject > > ResetArray;
-		for ( TPropObjectIterator Itor( RootPropertyNode->ObjectIterator() ) ; Itor ; ++Itor )
-		{
-			TWeakObjectPtr<UObject> Object = *Itor;
-			if( Object.IsValid() )
-			{
-				ResetArray.Add( Object.Get() );
-			}
-		}
-
-		SetObjectArrayPrivate(ResetArray);
-
-		// All objects are being reset, no need to validate external nodes
-		bValidateExternalNodes = false;
-	}
-
-	if( bValidateExternalNodes )
-	{
-		for( int32 NodeIndex = 0; NodeIndex < ExternalRootPropertyNodes.Num(); ++NodeIndex )
-		{
-			TSharedPtr<FObjectPropertyNode> ObjectNode = ExternalRootPropertyNodes[NodeIndex].Pin();
-
-			if( ObjectNode.IsValid() )
-			{
-				Result = ObjectNode->EnsureDataIsValid();
-				if( Result == FPropertyNode::PropertiesChanged || Result == FPropertyNode::EditInlineNewValueChanged  )
-				{
-					RestoreExpandedItems( ObjectNode );
-					UpdatePropertyMap();
-					// Note this will invalidate all the external root nodes so there is no need to continue
-					ExternalRootPropertyNodes.Empty();
-					break;
-				}
-				else if(Result == FPropertyNode::ArraySizeChanged  )
-				{
-					RestoreExpandedItems( ObjectNode );
-					UpdateFilteredDetails();
-				}
-			}
-			else
-			{
-				// Remove the current node if it is no longer valid
-				ExternalRootPropertyNodes.RemoveAt( NodeIndex );
-				--NodeIndex;
-			}
-		}
-	}
-
-	if( ThumbnailPool.IsValid() )
-	{
-		ThumbnailPool->Tick( InDeltaTime );
-	}
-
-	if( DetailLayout.IsValid() )
-	{
-		DetailLayout->Tick( InDeltaTime );
-	}
-
-	if( !ColorPropertyNode.IsValid() && bHasOpenColorPicker )
-	{
-		// Destroy the color picker window if the color property node has become invalid
-		DestroyColorPicker();
-		bHasOpenColorPicker = false;
-	}
-
-	
-	if( FilteredNodesRequestingExpansionState.Num() > 0 )
-	{
-		// change expansion state on the nodes that request it
-		for( TMap<TSharedRef<IDetailTreeNode>, bool >::TConstIterator It(FilteredNodesRequestingExpansionState); It; ++It )
-		{
-			DetailTree->SetItemExpansion( It.Key(), It.Value() );
-		}
-
-		FilteredNodesRequestingExpansionState.Empty();
-	}
-}
-
-/** 
- * Creates the color picker window for this property view.
- *
- * @param Node				The slate property node to edit.
- * @param bUseAlpha			Whether or not alpha is supported
- */
-void SDetailsView::CreateColorPickerWindow(const TSharedRef< FPropertyEditor >& PropertyEditor, bool bUseAlpha)
-{
-	const TSharedRef< FPropertyNode > PinnedColorPropertyNode = PropertyEditor->GetPropertyNode();
-	ColorPropertyNode = PinnedColorPropertyNode;
-
-	UProperty* Property = PinnedColorPropertyNode->GetProperty();
-	check(Property);
-
-	FReadAddressList ReadAddresses;
-	PinnedColorPropertyNode->GetReadAddress( false, ReadAddresses, false );
-
-	TArray<FLinearColor*> LinearColor;
-	TArray<FColor*> DWORDColor;
-	for( int32 ColorIndex = 0; ColorIndex < ReadAddresses.Num(); ++ColorIndex ) 
-	{
-		const uint8* Addr = ReadAddresses.GetAddress( ColorIndex );
-		if( Addr )
-		{
-			if( Cast<UStructProperty>(Property)->Struct->GetFName() == NAME_Color )
-			{
-				DWORDColor.Add((FColor*)Addr);
-			}
-			else
-			{
-				check( Cast<UStructProperty>(Property)->Struct->GetFName() == NAME_LinearColor );
-				LinearColor.Add((FLinearColor*)Addr);
-			}
-		}
-	}
-
-	bHasOpenColorPicker = true;
-
-	FColorPickerArgs PickerArgs;
-	PickerArgs.ParentWidget = AsShared();
-	PickerArgs.bUseAlpha = bUseAlpha;
-	PickerArgs.DisplayGamma = TAttribute<float>::Create( TAttribute<float>::FGetter::CreateUObject(GEngine, &UEngine::GetDisplayGamma) );
-	PickerArgs.ColorArray = &DWORDColor;
-	PickerArgs.LinearColorArray = &LinearColor;
-	PickerArgs.OnColorCommitted = FOnLinearColorValueChanged::CreateSP( this, &SDetailsView::SetColorPropertyFromColorPicker);
-	PickerArgs.OnColorPickerWindowClosed = FOnWindowClosed::CreateSP( this, &SDetailsView::OnColorPickerWindowClosed );
-
-	OpenColorPicker(PickerArgs);
-}
-
 void SDetailsView::SetOnObjectArrayChanged(FOnObjectArrayChanged OnObjectArrayChangedDelegate)
 {
 	OnObjectArrayChanged = OnObjectArrayChangedDelegate;
@@ -1288,6 +574,11 @@ UClass* SDetailsView::GetBaseClass()
 	return NULL;
 }
 
+UStruct* SDetailsView::GetBaseStruct() const
+{
+	return RootPropertyNode.IsValid() ? RootPropertyNode->GetObjectBaseClass() : NULL;
+}
+
 void SDetailsView::RegisterInstancedCustomPropertyLayout( UClass* Class, FOnGetDetailCustomizationInstance DetailLayoutDelegate )
 {
 	check( Class );
@@ -1307,23 +598,6 @@ void SDetailsView::UnregisterInstancedCustomPropertyLayout( UClass* Class )
 	InstancedClassToDetailLayoutMap.Remove( Class );
 }
 
-bool SDetailsView::SupportsKeyboardFocus() const
-{
-	return DetailsViewArgs.bSearchInitialKeyFocus && SearchBox->SupportsKeyboardFocus() && GetFilterBoxVisibility() == EVisibility::Visible;
-}
-
-FReply SDetailsView::OnKeyboardFocusReceived( const FGeometry& MyGeometry, const FKeyboardFocusEvent& InKeyboardFocusEvent )
-{
-	FReply Reply = FReply::Handled();
-
-	if( InKeyboardFocusEvent.GetCause() != EKeyboardFocusCause::Cleared )
-	{
-		Reply.SetKeyboardFocus( SearchBox.ToSharedRef(), InKeyboardFocusEvent.GetCause() );
-	}
-
-	return Reply;
-}
-
 void SDetailsView::AddExternalRootPropertyNode( TSharedRef<FObjectPropertyNode> ExternalRootNode )
 {
 	ExternalRootPropertyNodes.Add( ExternalRootNode );
@@ -1334,4 +608,9 @@ bool SDetailsView::IsCategoryHiddenByClass( FName CategoryName ) const
 	return RootPropertyNode->GetHiddenCategories().Contains( CategoryName );
 }
 
+bool SDetailsView::IsConnected() const
+{
+	return RootPropertyNode.IsValid() && (RootPropertyNode->GetNumObjects() > 0);
+}
+
 #undef LOCTEXT_NAMESPACE
diff --git a/Engine/Source/Editor/PropertyEditor/Private/SDetailsView.h b/Engine/Source/Editor/PropertyEditor/Private/SDetailsView.h
--- a/Engine/Source/Editor/PropertyEditor/Private/SDetailsView.h
+++ b/Engine/Source/Editor/PropertyEditor/Private/SDetailsView.h
@@ -52,14 +52,6 @@ public:
 	 */
 	void RemoveInvalidActors( const TSet<AActor*>& ValidActors );
 
-	/**
-	 * Creates the color picker window for this property view.
-	 *
-	 * @param PropertyEditor				The slate property node to edit.
-	 * @param bUseAlpha			Whether or not alpha is supported
-	 */
-	virtual void CreateColorPickerWindow(const TSharedRef< class FPropertyEditor >& PropertyEditor, bool bUseAlpha) override;
-
 	/** Sets the callback for when the property view changes */
 	virtual void SetOnObjectArrayChanged( FOnObjectArrayChanged OnObjectArrayChangedDelegate);
 
@@ -89,12 +81,7 @@ public:
 	/** Gets the base class being viewed */
 	const UClass* GetBaseClass() const override;
 	UClass* GetBaseClass() override;
-
-	// SWidget interface
-	virtual bool SupportsKeyboardFocus() const override;
-	virtual FReply OnKeyboardFocusReceived( const FGeometry& MyGeometry, const FKeyboardFocusEvent& InKeyboardFocusEvent ) override;
-	virtual void Tick( const FGeometry& AllottedGeometry, const double InCurrentTime, const float InDeltaTime ) override;
-	// End of SWidget interface
+	UStruct* GetBaseStruct() const override;
 
 	/**
 	 * Adds an external property root node to the list of root nodes that the details new needs to manage
@@ -108,6 +95,12 @@ public:
 	 */
 	bool IsCategoryHiddenByClass(FName CategoryName) const override;
 
+	virtual bool IsConnected() const override;
+
+	virtual TSharedPtr<FComplexPropertyNode> GetRootNode() override
+	{
+		return RootPropertyNode;
+	}
 private:
 	void RegisterInstancedCustomPropertyLayout( UClass* Class, FOnGetDetailCustomizationInstance DetailLayoutDelegate );
 	void UnregisterInstancedCustomPropertyLayout( UClass* Class );
@@ -124,48 +117,18 @@ private:
 	 */
 	bool ShouldSetNewObjects( const TArray< TWeakObjectPtr< UObject > >& InObjects ) const;
 
-	/** Updates the property map for access when customizing the details view.  Generates default layout for properties */
-	void UpdatePropertyMap();
-
-	/** 
-	 * Recursively updates children of property nodes. Generates default layout for properties 
-	 * 
-	 * @param InNode	The parent node to get children from
-	 * @param The detail layout builder that will be used for customization of this property map
-	 * @param CurCategory The current category name
-	 */
-	void UpdatePropertyMapRecursive( FPropertyNode& InNode, FDetailLayoutBuilderImpl& DetailLayout, FName CurCategory, FObjectPropertyNode* CurObjectNode );
-
 	/** Called before during SetObjectArray before we change the objects being observed */
 	void PreSetObject();
 
 	/** Called at the end of SetObjectArray after we change the objects being observed */
 	void PostSetObject();
 	
-	/**
-	 * Queries a layout for a specific class
-	 */
-	void QueryLayoutForClass( FDetailLayoutBuilderImpl& CustomDetailLayout, UStruct* Class );
-
-	/**
-	 * Calls a delegate for each registered class that has properties visible to get any custom detail layouts 
-	 */
-	void QueryCustomDetailLayout( class FDetailLayoutBuilderImpl& CustomDetailLayout );
-
-	/**
-	 * Updates the details with the passed in filter                                                              
-	 */
-	virtual void UpdateFilteredDetails() override;
-
 	/** Called when the filter button is clicked */
 	void OnFilterButtonClicked();
 
 	/** Called to get the visibility of the actor name area */
 	EVisibility GetActorNameAreaVisibility() const;
 
-	/** Called to get the visibility of the filter box */
-	EVisibility GetFilterBoxVisibility() const;
-
 	/** Returns the name of the image used for the icon on the locked button */
 	const FSlateBrush* OnGetLockButtonImageResource() const;
 	/**
@@ -173,28 +136,12 @@ private:
 	 */
 	FReply OnOpenRawPropertyEditorClicked();
 
-	/**
-	 * Called when a color property is changed from a color picker
-	 */
-	void SetColorPropertyFromColorPicker(FLinearColor NewColor);
-
-	/** Saves the expansion state of property nodes for the selected object set */
-	void SaveExpandedItems();
-
-	/** 
-	 * Restores the expansion state of property nodes for the selected object set
-	 *
-	 * @param InitialStartNode The starting node if any.  If one is not supplied the expansion state is restored from the root node
-	 */
-	void RestoreExpandedItems( TSharedPtr<FPropertyNode> InitialStartNode = NULL );
-
 private:
 	/** Information about the current set of selected actors */
 	FSelectedActorInfo SelectedActorInfo;
 	/** Selected objects for this detail view.  */
 	TArray< TWeakObjectPtr<UObject> > SelectedObjects;
-	/** Root tree node that needs to be destroyed when safe */
-	TSharedPtr<FObjectPropertyNode> RootNodePendingKill;
+
 	/** 
 	 * Selected actors for this detail view.  Note that this is not necessarily the same editor selected actor set.  If this detail view is locked
 	 * It will only contain actors from when it was locked 
@@ -202,8 +149,6 @@ private:
 	TArray< TWeakObjectPtr<AActor> > SelectedActors;
 	/** The root property node of the property tree for a specific set of UObjects */
 	TSharedPtr<FObjectPropertyNode> RootPropertyNode;
-	/** External property nodes which need to validated each tick */
-	TArray< TWeakPtr<FObjectPropertyNode> > ExternalRootPropertyNodes;
 	/** Callback to send when the property view changes */
 	FOnObjectArrayChanged OnObjectArrayChanged;
 	/** True if at least one viewed object is a CDO (blueprint editing) */
diff --git a/Engine/Source/Editor/PropertyEditor/Private/SDetailsViewBase.cpp b/Engine/Source/Editor/PropertyEditor/Private/SDetailsViewBase.cpp
--- a/Engine/Source/Editor/PropertyEditor/Private/SDetailsViewBase.cpp
+++ b/Engine/Source/Editor/PropertyEditor/Private/SDetailsViewBase.cpp
@@ -109,8 +109,93 @@ void SDetailsViewBase::EnqueueDeferredAction(FSimpleDelegate& DeferredAction)
 	DeferredActions.AddUnique(DeferredAction);
 }
 
+/**
+* Creates the color picker window for this property view.
+*
+* @param Node				The slate property node to edit.
+* @param bUseAlpha			Whether or not alpha is supported
+*/
+void SDetailsViewBase::CreateColorPickerWindow(const TSharedRef< FPropertyEditor >& PropertyEditor, bool bUseAlpha)
+{
+	const TSharedRef< FPropertyNode > PinnedColorPropertyNode = PropertyEditor->GetPropertyNode();
+	ColorPropertyNode = PinnedColorPropertyNode;
+
+	UProperty* Property = PinnedColorPropertyNode->GetProperty();
+	check(Property);
+
+	FReadAddressList ReadAddresses;
+	PinnedColorPropertyNode->GetReadAddress(false, ReadAddresses, false);
+
+	TArray<FLinearColor*> LinearColor;
+	TArray<FColor*> DWORDColor;
+	for (int32 ColorIndex = 0; ColorIndex < ReadAddresses.Num(); ++ColorIndex)
+	{
+		const uint8* Addr = ReadAddresses.GetAddress(ColorIndex);
+		if (Addr)
+		{
+			if (Cast<UStructProperty>(Property)->Struct->GetFName() == NAME_Color)
+			{
+				DWORDColor.Add((FColor*)Addr);
+			}
+			else
+			{
+				check(Cast<UStructProperty>(Property)->Struct->GetFName() == NAME_LinearColor);
+				LinearColor.Add((FLinearColor*)Addr);
+			}
+		}
+	}
+
+	bHasOpenColorPicker = true;
+
+	FColorPickerArgs PickerArgs;
+	PickerArgs.ParentWidget = AsShared();
+	PickerArgs.bUseAlpha = bUseAlpha;
+	PickerArgs.DisplayGamma = TAttribute<float>::Create(TAttribute<float>::FGetter::CreateUObject(GEngine, &UEngine::GetDisplayGamma));
+	PickerArgs.ColorArray = &DWORDColor;
+	PickerArgs.LinearColorArray = &LinearColor;
+	PickerArgs.OnColorCommitted = FOnLinearColorValueChanged::CreateSP(this, &SDetailsViewBase::SetColorPropertyFromColorPicker);
+	PickerArgs.OnColorPickerWindowClosed = FOnWindowClosed::CreateSP(this, &SDetailsViewBase::OnColorPickerWindowClosed);
+
+	OpenColorPicker(PickerArgs);
+}
+
+
+void SDetailsViewBase::SetColorPropertyFromColorPicker(FLinearColor NewColor)
+{
+	const TSharedPtr< FPropertyNode > PinnedColorPropertyNode = ColorPropertyNode.Pin();
+	if (ensure(PinnedColorPropertyNode.IsValid()))
+	{
+		UProperty* Property = PinnedColorPropertyNode->GetProperty();
+		check(Property);
+
+		FObjectPropertyNode* ObjectNode = PinnedColorPropertyNode->FindObjectItemParent();
+
+		if (ObjectNode && ObjectNode->GetNumObjects())
+		{
+			FScopedTransaction Transaction(NSLOCTEXT("UnrealEd", "SetColorProperty", "Set Color Property"));
+
+			PinnedColorPropertyNode->NotifyPreChange(Property, GetNotifyHook());
+
+			FPropertyChangedEvent ChangeEvent(Property, false, EPropertyChangeType::ValueSet);
+			PinnedColorPropertyNode->NotifyPostChange(ChangeEvent, GetNotifyHook());
+		}
+	}
+}
+
 void SDetailsViewBase::OnColorPickerWindowClosed(const TSharedRef<SWindow>& Window)
 {
+	const TSharedPtr< FPropertyNode > PinnedColorPropertyNode = ColorPropertyNode.Pin();
+	if (ensure(PinnedColorPropertyNode.IsValid()))
+	{
+		UProperty* Property = PinnedColorPropertyNode->GetProperty();
+		if (Property && PropertyUtilities.IsValid())
+		{
+			FPropertyChangedEvent ChangeEvent(Property, false, EPropertyChangeType::ArrayAdd);
+			PinnedColorPropertyNode->FixPropertiesInEvent(ChangeEvent);
+			PropertyUtilities->NotifyFinishedChangingProperties(ChangeEvent);
+		}
+	}
+
 	// A color picker window is no longer open
 	bHasOpenColorPicker = false;
 	ColorPropertyNode.Reset();
@@ -242,4 +327,671 @@ void SDetailsViewBase::FilterView(const FString& InFilterText)
 	CurrentFilter.FilterStrings = CurrentFilterStrings;
 
 	UpdateFilteredDetails();
-}
\ No newline at end of file
+}
+
+void SDetailsViewBase::QueryLayoutForClass(FDetailLayoutBuilderImpl& CustomDetailLayout, UStruct* Class)
+{
+	CustomDetailLayout.SetCurrentCustomizationClass(CastChecked<UClass>(Class), NAME_None);
+
+	FPropertyEditorModule& ParentPlugin = FModuleManager::GetModuleChecked<FPropertyEditorModule>("PropertyEditor");
+	FCustomDetailLayoutNameMap& GlobalCustomLayoutNameMap = ParentPlugin.ClassNameToDetailLayoutNameMap;
+
+	// Check the instanced map first
+	FDetailLayoutCallback* Callback = InstancedClassToDetailLayoutMap.Find(TWeakObjectPtr<UStruct>(Class));
+
+	if (!Callback)
+	{
+		// callback wasn't found in the per instance map, try the global instances instead
+		Callback = GlobalCustomLayoutNameMap.Find(Class->GetFName());
+	}
+
+	if (Callback && Callback->DetailLayoutDelegate.IsBound())
+	{
+		// Create a new instance of the custom detail layout for the current class
+		TSharedRef<IDetailCustomization> CustomizationInstance = Callback->DetailLayoutDelegate.Execute();
+
+		// Ask for details immediately
+		CustomizationInstance->CustomizeDetails(CustomDetailLayout);
+
+		// Save the instance from destruction until we refresh
+		CustomizationClassInstances.Add(CustomizationInstance);
+	}
+}
+
+void SDetailsViewBase::QueryCustomDetailLayout(FDetailLayoutBuilderImpl& CustomDetailLayout)
+{
+	FPropertyEditorModule& ParentPlugin = FModuleManager::GetModuleChecked<FPropertyEditorModule>("PropertyEditor");
+
+	// Get the registered classes that customize details
+	FCustomDetailLayoutNameMap& GlobalCustomLayoutNameMap = ParentPlugin.ClassNameToDetailLayoutNameMap;
+
+	UStruct* BaseStruct = GetBaseStruct();
+
+	// All the current customization instances need to be deleted when it is safe
+	CustomizationClassInstancesPendingDelete = CustomizationClassInstances;
+
+	CustomizationClassInstances.Empty();
+
+	//Ask for generic details not specific to an object being viewed 
+	if (GenericLayoutDelegate.IsBound())
+	{
+		// Create a new instance of the custom detail layout for the current class
+		TSharedRef<IDetailCustomization> CustomizationInstance = GenericLayoutDelegate.Execute();
+
+		// Ask for details immediately
+		CustomizationInstance->CustomizeDetails(CustomDetailLayout);
+
+		// Save the instance from destruction until we refresh
+		CustomizationClassInstances.Add(CustomizationInstance);
+	}
+
+
+	// Sort them by query order.  @todo not good enough
+	struct FCompareFDetailLayoutCallback
+	{
+		FORCEINLINE bool operator()(const FDetailLayoutCallback& A, const FDetailLayoutCallback& B) const
+		{
+			return A.Order < B.Order;
+		}
+	};
+
+	TMap< TWeakObjectPtr<UStruct>, FDetailLayoutCallback*> FinalCallbackMap;
+
+	for (auto ClassIt = ClassesWithProperties.CreateConstIterator(); ClassIt; ++ClassIt)
+	{
+		// Check the instanced map first
+		FDetailLayoutCallback* Callback = InstancedClassToDetailLayoutMap.Find(*ClassIt);
+
+		if (!Callback)
+		{
+			// callback wasn't found in the per instance map, try the global instances instead
+			Callback = GlobalCustomLayoutNameMap.Find((*ClassIt)->GetFName());
+		}
+
+		if (Callback)
+		{
+			FinalCallbackMap.Add(*ClassIt, Callback);
+		}
+	}
+
+
+	FinalCallbackMap.ValueSort(FCompareFDetailLayoutCallback());
+
+	TSet<UStruct*> QueriedClasses;
+
+	if (FinalCallbackMap.Num() > 0)
+	{
+		// Ask each class that we have properties for to customize its layout
+		for (auto LayoutIt(FinalCallbackMap.CreateConstIterator()); LayoutIt; ++LayoutIt)
+		{
+			const TWeakObjectPtr<UStruct> WeakClass = LayoutIt.Key();
+
+			if (WeakClass.IsValid())
+			{
+				UStruct* Class = WeakClass.Get();
+
+				FClassInstanceToPropertyMap& InstancedPropertyMap = ClassToPropertyMap.FindChecked(Class->GetFName());
+				for (FClassInstanceToPropertyMap::TIterator InstanceIt(InstancedPropertyMap); InstanceIt; ++InstanceIt)
+				{
+					CustomDetailLayout.SetCurrentCustomizationClass(CastChecked<UClass>(Class), InstanceIt.Key());
+
+					const FOnGetDetailCustomizationInstance& DetailDelegate = LayoutIt.Value()->DetailLayoutDelegate;
+
+					if (DetailDelegate.IsBound())
+					{
+						QueriedClasses.Add(Class);
+
+						// Create a new instance of the custom detail layout for the current class
+						TSharedRef<IDetailCustomization> CustomizationInstance = DetailDelegate.Execute();
+
+						// Ask for details immediately
+						CustomizationInstance->CustomizeDetails(CustomDetailLayout);
+
+						// Save the instance from destruction until we refresh
+						CustomizationClassInstances.Add(CustomizationInstance);
+					}
+				}
+			}
+		}
+	}
+
+	// Ensure that the base class and its parents are always queried
+	TSet<UStruct*> ParentClassesToQuery;
+	if (BaseStruct && !QueriedClasses.Contains(BaseStruct))
+	{
+		ParentClassesToQuery.Add(BaseStruct);
+		ClassesWithProperties.Add(BaseStruct);
+	}
+
+	// Find base classes of queried classes that were not queried and add them to the query list
+	// this supports cases where a parent class has no properties but still wants to add customization
+	for (auto QueriedClassIt = ClassesWithProperties.CreateConstIterator(); QueriedClassIt; ++QueriedClassIt)
+	{
+		UStruct* ParentStruct = (*QueriedClassIt)->GetSuperStruct();
+
+		while (ParentStruct && ParentStruct->IsA(UClass::StaticClass()) && !QueriedClasses.Contains(ParentStruct) && !ClassesWithProperties.Contains(ParentStruct))
+		{
+			ParentClassesToQuery.Add(ParentStruct);
+			ParentStruct = ParentStruct->GetSuperStruct();
+
+		}
+	}
+
+	// Query extra base classes
+	for (auto ParentIt = ParentClassesToQuery.CreateConstIterator(); ParentIt; ++ParentIt)
+	{
+		if (Cast<UClass>(*ParentIt))
+		{
+			QueryLayoutForClass(CustomDetailLayout, *ParentIt);
+		}
+	}
+}
+
+EVisibility SDetailsViewBase::GetFilterBoxVisibility() const
+{
+	// Visible if we allow search and we have anything to search otherwise collapsed so it doesn't take up room
+	return (DetailsViewArgs.bAllowSearch && IsConnected()) ? EVisibility::Visible : EVisibility::Collapsed;
+}
+
+bool SDetailsViewBase::SupportsKeyboardFocus() const
+{
+	return DetailsViewArgs.bSearchInitialKeyFocus && SearchBox->SupportsKeyboardFocus() && GetFilterBoxVisibility() == EVisibility::Visible;
+}
+
+FReply SDetailsViewBase::OnKeyboardFocusReceived(const FGeometry& MyGeometry, const FKeyboardFocusEvent& InKeyboardFocusEvent)
+{
+	FReply Reply = FReply::Handled();
+
+	if (InKeyboardFocusEvent.GetCause() != EKeyboardFocusCause::Cleared)
+	{
+		Reply.SetKeyboardFocus(SearchBox.ToSharedRef(), InKeyboardFocusEvent.GetCause());
+	}
+
+	return Reply;
+}
+
+/** Ticks the property view.  This function performs a data consistency check */
+void SDetailsViewBase::Tick(const FGeometry& AllottedGeometry, const double InCurrentTime, const float InDeltaTime)
+{
+	for (int32 i = 0; i < CustomizationClassInstancesPendingDelete.Num(); ++i)
+	{
+		ensure(CustomizationClassInstancesPendingDelete[i].IsUnique());
+	}
+
+	if (RootNodePendingKill.IsValid())
+	{
+		RootNodePendingKill->Disconnect();
+		RootNodePendingKill.Reset();
+	}
+
+	// Empty all the customization instances that need to be deleted
+	CustomizationClassInstancesPendingDelete.Empty();
+
+	auto RootPropertyNode = GetRootNode();
+	check(RootPropertyNode.IsValid());
+
+	// Purge any objects that are marked pending kill from the object list
+	if (auto ObjectRoot = RootPropertyNode->AsObjectNode())
+	{
+		ObjectRoot->PurgeKilledObjects();
+	}
+
+	if (DeferredActions.Num() > 0)
+	{
+		// Any deferred actions are likely to cause the node  tree to be at least partially rebuilt
+		// Save the expansion state of existing nodes so we can expand them later
+		SaveExpandedItems();
+
+		// Execute any deferred actions
+		for (int32 ActionIndex = 0; ActionIndex < DeferredActions.Num(); ++ActionIndex)
+		{
+			DeferredActions[ActionIndex].ExecuteIfBound();
+		}
+		DeferredActions.Empty();
+	}
+
+	bool bValidateExternalNodes = true;
+	FPropertyNode::DataValidationResult Result = RootPropertyNode->EnsureDataIsValid();
+	if (Result == FPropertyNode::PropertiesChanged || Result == FPropertyNode::EditInlineNewValueChanged)
+	{
+		RestoreExpandedItems();
+		UpdatePropertyMap();
+	}
+	else if (Result == FPropertyNode::ArraySizeChanged)
+	{
+		RestoreExpandedItems();
+		UpdateFilteredDetails();
+	}
+	else if (Result == FPropertyNode::ObjectInvalid)
+	{
+		ForceRefresh();
+
+		// All objects are being reset, no need to validate external nodes
+		bValidateExternalNodes = false;
+	}
+
+	if (bValidateExternalNodes)
+	{
+		for (int32 NodeIndex = 0; NodeIndex < ExternalRootPropertyNodes.Num(); ++NodeIndex)
+		{
+			TSharedPtr<FObjectPropertyNode> ObjectNode = ExternalRootPropertyNodes[NodeIndex].Pin();
+
+			if (ObjectNode.IsValid())
+			{
+				Result = ObjectNode->EnsureDataIsValid();
+				if (Result == FPropertyNode::PropertiesChanged || Result == FPropertyNode::EditInlineNewValueChanged)
+				{
+					RestoreExpandedItems(ObjectNode);
+					UpdatePropertyMap();
+					// Note this will invalidate all the external root nodes so there is no need to continue
+					ExternalRootPropertyNodes.Empty();
+					break;
+				}
+				else if (Result == FPropertyNode::ArraySizeChanged)
+				{
+					RestoreExpandedItems(ObjectNode);
+					UpdateFilteredDetails();
+				}
+			}
+			else
+			{
+				// Remove the current node if it is no longer valid
+				ExternalRootPropertyNodes.RemoveAt(NodeIndex);
+				--NodeIndex;
+			}
+		}
+	}
+
+	if (ThumbnailPool.IsValid())
+	{
+		ThumbnailPool->Tick(InDeltaTime);
+	}
+
+	if (DetailLayout.IsValid())
+	{
+		DetailLayout->Tick(InDeltaTime);
+	}
+
+	if (!ColorPropertyNode.IsValid() && bHasOpenColorPicker)
+	{
+		// Destroy the color picker window if the color property node has become invalid
+		DestroyColorPicker();
+		bHasOpenColorPicker = false;
+	}
+
+
+	if (FilteredNodesRequestingExpansionState.Num() > 0)
+	{
+		// change expansion state on the nodes that request it
+		for (TMap<TSharedRef<IDetailTreeNode>, bool >::TConstIterator It(FilteredNodesRequestingExpansionState); It; ++It)
+		{
+			DetailTree->SetItemExpansion(It.Key(), It.Value());
+		}
+
+		FilteredNodesRequestingExpansionState.Empty();
+	}
+}
+
+/**
+* Recursively gets expanded items for a node
+*
+* @param InPropertyNode			The node to get expanded items from
+* @param OutExpandedItems	List of expanded items that were found
+*/
+void GetExpandedItems(TSharedPtr<FPropertyNode> InPropertyNode, TArray<FString>& OutExpandedItems)
+{
+	if (InPropertyNode->HasNodeFlags(EPropertyNodeFlags::Expanded))
+	{
+		const bool bWithArrayIndex = true;
+		FString Path;
+		Path.Empty(128);
+		InPropertyNode->GetQualifiedName(Path, bWithArrayIndex);
+
+		OutExpandedItems.Add(Path);
+	}
+
+	for (int32 ChildIndex = 0; ChildIndex < InPropertyNode->GetNumChildNodes(); ++ChildIndex)
+	{
+		GetExpandedItems(InPropertyNode->GetChildNode(ChildIndex), OutExpandedItems);
+	}
+
+}
+
+/**
+* Recursively sets expanded items for a node
+*
+* @param InNode			The node to set expanded items on
+* @param OutExpandedItems	List of expanded items to set
+*/
+void SetExpandedItems(TSharedPtr<FPropertyNode> InPropertyNode, const TArray<FString>& InExpandedItems)
+{
+	if (InExpandedItems.Num() > 0)
+	{
+		const bool bWithArrayIndex = true;
+		FString Path;
+		Path.Empty(128);
+		InPropertyNode->GetQualifiedName(Path, bWithArrayIndex);
+
+		for (int32 ItemIndex = 0; ItemIndex < InExpandedItems.Num(); ++ItemIndex)
+		{
+			if (InExpandedItems[ItemIndex] == Path)
+			{
+				InPropertyNode->SetNodeFlags(EPropertyNodeFlags::Expanded, true);
+				break;
+			}
+		}
+
+		for (int32 NodeIndex = 0; NodeIndex < InPropertyNode->GetNumChildNodes(); ++NodeIndex)
+		{
+			SetExpandedItems(InPropertyNode->GetChildNode(NodeIndex), InExpandedItems);
+		}
+	}
+}
+
+void SDetailsViewBase::SaveExpandedItems()
+{
+	auto RootPropertyNode = GetRootNode();
+	check(RootPropertyNode.IsValid());
+	UStruct* BestBaseStruct = RootPropertyNode->GetBaseStructure();
+
+	TArray<FString> ExpandedPropertyItems;
+	GetExpandedItems(RootPropertyNode, ExpandedPropertyItems);
+
+	TArray<FString> ExpandedCustomItems = ExpandedDetailNodes.Array();
+
+	// Expanded custom items may have spaces but SetSingleLineArray doesnt support spaces (treats it as another element in the array)
+	// Append a '|' after each element instead
+	FString ExpandedCustomItemsString;
+	for (auto It = ExpandedDetailNodes.CreateConstIterator(); It; ++It)
+	{
+		ExpandedCustomItemsString += *It;
+		ExpandedCustomItemsString += TEXT(",");
+	}
+
+	//while a valid class, and we're either the same as the base class (for multiple actors being selected and base class is AActor) OR we're not down to AActor yet)
+	for (UStruct* Struct = BestBaseStruct; Struct && ((BestBaseStruct == Struct) || (Struct != AActor::StaticClass())); Struct = Struct->GetSuperStruct())
+	{
+		if (RootPropertyNode->GetNumChildNodes() > 0 && ExpandedPropertyItems.Num() > 0)
+		{
+			GConfig->SetSingleLineArray(TEXT("DetailPropertyExpansion"), *Struct->GetName(), ExpandedPropertyItems, GEditorUserSettingsIni);
+		}
+	}
+
+	if (DetailLayout.IsValid() && BestBaseStruct && !ExpandedCustomItemsString.IsEmpty())
+	{
+		GConfig->SetString(TEXT("DetailCustomWidgetExpansion"), *BestBaseStruct->GetName(), *ExpandedCustomItemsString, GEditorUserSettingsIni);
+	}
+}
+
+void SDetailsViewBase::RestoreExpandedItems(TSharedPtr<FPropertyNode> InitialStartNode)
+{
+	auto RootPropertyNode = GetRootNode();
+	check(RootPropertyNode.IsValid());
+	TSharedPtr<FPropertyNode> StartNode = InitialStartNode;
+	if (!StartNode.IsValid())
+	{
+		StartNode = RootPropertyNode;
+	}
+
+	ExpandedDetailNodes.Empty();
+
+	TArray<FString> ExpandedPropertyItems;
+	FString ExpandedCustomItems;
+	UStruct* BestBaseStruct = RootPropertyNode->GetBaseStructure();
+	//while a valid class, and we're either the same as the base class (for multiple actors being selected and base class is AActor) OR we're not down to AActor yet)
+	for (UStruct* Struct = BestBaseStruct; Struct && ((BestBaseStruct == Struct) || (Struct != AActor::StaticClass())); Struct = Struct->GetSuperStruct())
+	{
+		GConfig->GetSingleLineArray(TEXT("DetailPropertyExpansion"), *Struct->GetName(), ExpandedPropertyItems, GEditorUserSettingsIni);
+		SetExpandedItems(StartNode, ExpandedPropertyItems);
+	}
+
+	if (BestBaseStruct)
+	{
+		GConfig->GetString(TEXT("DetailCustomWidgetExpansion"), *BestBaseStruct->GetName(), ExpandedCustomItems, GEditorUserSettingsIni);
+		TArray<FString> ExpandedCustomItemsArray;
+		ExpandedCustomItems.ParseIntoArray(&ExpandedCustomItemsArray, TEXT(","), true);
+
+		ExpandedDetailNodes.Append(ExpandedCustomItemsArray);
+	}
+}
+
+void SDetailsViewBase::UpdateFilteredDetails()
+{
+	auto RootPropertyNode = GetRootNode();
+	if (RootPropertyNode.IsValid())
+	{
+		RootPropertyNode->FilterNodes(CurrentFilter.FilterStrings);
+		RootPropertyNode->ProcessSeenFlags(true);
+
+		for (int32 NodeIndex = 0; NodeIndex < ExternalRootPropertyNodes.Num(); ++NodeIndex)
+		{
+			TSharedPtr<FObjectPropertyNode> ObjectNode = ExternalRootPropertyNodes[NodeIndex].Pin();
+
+			if (ObjectNode.IsValid())
+			{
+				ObjectNode->FilterNodes(CurrentFilter.FilterStrings);
+				ObjectNode->ProcessSeenFlags(true);
+			}
+		}
+
+		if (DetailLayout.IsValid())
+		{
+			DetailLayout->FilterDetailLayout(CurrentFilter);
+		}
+
+		RootTreeNodes = DetailLayout->GetRootTreeNodes();
+	}
+	DetailTree->RequestTreeRefresh();
+}
+
+/**
+* Determines whether or not a property should be visible in the default generated detail layout
+*
+* @param PropertyNode	The property node to check
+* @param ParentNode	The parent property node to check
+* @return true if the property should be visible
+*/
+static bool IsVisibleStandaloneProperty(const FPropertyNode& PropertyNode, const FPropertyNode& ParentNode)
+{
+	const UProperty* Property = PropertyNode.GetProperty();
+	const UArrayProperty* ParentArrayProperty = Cast<const UArrayProperty>(ParentNode.GetProperty());
+	bool bIsVisibleStandalone = false;
+	if (Property)
+	{
+		if (Property->IsA(UObjectPropertyBase::StaticClass()))
+		{
+			// Do not add this child node to the current map if its a single object property in a category (serves no purpose for UI)
+			bIsVisibleStandalone = !ParentArrayProperty && (PropertyNode.GetNumChildNodes() == 0 || PropertyNode.GetNumChildNodes() > 1);
+		}
+		else if (Property->IsA(UArrayProperty::StaticClass()) || Property->ArrayDim > 1 && PropertyNode.GetArrayIndex() == INDEX_NONE)
+		{
+			// Base array properties are always visible
+			bIsVisibleStandalone = true;
+		}
+		else
+		{
+			bIsVisibleStandalone = true;
+		}
+
+	}
+
+	return bIsVisibleStandalone;
+}
+
+void SDetailsViewBase::UpdatePropertyMapRecursive(FPropertyNode& InNode, FDetailLayoutBuilderImpl& InDetailLayout, FName CurCategory, FComplexPropertyNode* CurObjectNode)
+{
+	UProperty* ParentProperty = InNode.GetProperty();
+	UStructProperty* ParentStructProp = Cast<UStructProperty>(ParentProperty);
+
+	for (int32 ChildIndex = 0; ChildIndex < InNode.GetNumChildNodes(); ++ChildIndex)
+	{
+		TSharedPtr<FPropertyNode> ChildNodePtr = InNode.GetChildNode(ChildIndex);
+		FPropertyNode& ChildNode = *ChildNodePtr;
+		UProperty* Property = ChildNode.GetProperty();
+
+		{
+			FObjectPropertyNode* ObjNode = ChildNode.AsObjectNode();
+			FCategoryPropertyNode* CategoryNode = ChildNode.AsCategoryNode();
+			if (ObjNode)
+			{
+				// Currently object property nodes do not provide any useful information other than being a container for its children.  We do not draw anything for them.
+				// When we encounter object property nodes, add their children instead of adding them to the tree.
+				UpdatePropertyMapRecursive(ChildNode, InDetailLayout, CurCategory, ObjNode);
+			}
+			else if (CategoryNode)
+			{
+				// For category nodes, we just set the current category and recurse through the children
+				UpdatePropertyMapRecursive(ChildNode, InDetailLayout, CategoryNode->GetCategoryName(), CurObjectNode);
+			}
+			else
+			{
+				// Whether or not the property can be visible in the default detail layout
+				bool bVisibleByDefault = IsVisibleStandaloneProperty(ChildNode, InNode);
+
+				// Whether or not the property is a struct
+				UStructProperty* StructProperty = Cast<UStructProperty>(Property);
+
+				bool bIsStruct = StructProperty != NULL;
+
+				static FName ShowOnlyInners("ShowOnlyInnerProperties");
+
+				bool bIsChildOfCustomizedStruct = false;
+				bool bIsCustomizedStruct = false;
+
+				const UStruct* Struct = StructProperty ? StructProperty->Struct : NULL;
+				const UStruct* ParentStruct = ParentStructProp ? ParentStructProp->Struct : NULL;
+				if (Struct || ParentStruct)
+				{
+					FPropertyEditorModule& ParentPlugin = FModuleManager::GetModuleChecked<FPropertyEditorModule>("PropertyEditor");
+					if (Struct)
+					{
+						bIsCustomizedStruct = ParentPlugin.IsCustomizedStruct(Struct);
+					}
+
+					if (ParentStruct)
+					{
+						bIsChildOfCustomizedStruct = ParentPlugin.IsCustomizedStruct(ParentStruct);
+					}
+				}
+
+				// Whether or not to push out struct properties to their own categories or show them inside an expandable struct 
+				bool bPushOutStructProps = bIsStruct && !bIsCustomizedStruct && !ParentStructProp && Property->HasMetaData(ShowOnlyInners);
+
+				// Is the property edit inline new 
+				const bool bIsEditInlineNew = SPropertyEditorEditInline::Supports(&ChildNode, ChildNode.GetArrayIndex());
+
+				// Is this a property of an array
+				bool bIsChildOfArray = PropertyEditorHelpers::IsChildOfArray(ChildNode);
+
+				// Edit inline new properties should be visible by default
+				bVisibleByDefault |= bIsEditInlineNew;
+
+				// Children of arrays are not visible directly,
+				bVisibleByDefault &= !bIsChildOfArray;
+
+				// See if the user requested specific property visibility 
+				const bool bIsUserVisible = IsPropertyVisible(Property);
+
+				// Inners of customized in structs should not be taken into consideration for customizing.  They are not designed to be individually customized when their parent is already customized
+				if (!bIsChildOfCustomizedStruct)
+				{
+					// Add any object classes with properties so we can ask them for custom property layouts later
+					ClassesWithProperties.Add(Property->GetOwnerStruct());
+				}
+
+				// If there is no outer object then the class is the object root and there is only one instance
+				FName InstanceName = NAME_None;
+				if (CurObjectNode && CurObjectNode->GetParentNode())
+				{
+					InstanceName = CurObjectNode->GetParentNode()->GetProperty()->GetFName();
+				}
+				else if (ParentStructProp)
+				{
+					InstanceName = ParentStructProp->GetFName();
+				}
+
+				// Do not add children of customized in struct properties or arrays
+				if (!bIsChildOfCustomizedStruct && !bIsChildOfArray)
+				{
+					// Get the class property map
+					FClassInstanceToPropertyMap& ClassInstanceMap = ClassToPropertyMap.FindOrAdd(Property->GetOwnerStruct()->GetFName());
+
+					FPropertyNodeMap& PropertyNodeMap = ClassInstanceMap.FindOrAdd(InstanceName);
+
+					if (!PropertyNodeMap.ParentProperty)
+					{
+						PropertyNodeMap.ParentProperty = CurObjectNode;
+					}
+					else
+					{
+						ensure(PropertyNodeMap.ParentProperty == CurObjectNode);
+					}
+
+					checkSlow(!PropertyNodeMap.Contains(Property->GetFName()));
+
+					PropertyNodeMap.Add(Property->GetFName(), ChildNodePtr);
+				}
+
+				if (bVisibleByDefault && bIsUserVisible && !bPushOutStructProps)
+				{
+					FName CategoryName = CurCategory;
+
+					// For properties inside a struct, add them to their own category unless they just take the name of the parent struct.  
+					// In that case push them to the parent category
+					FName PropertyCatagoryName = FObjectEditorUtils::GetCategoryFName(Property);
+					if (!ParentStructProp || (PropertyCatagoryName != ParentStructProp->Struct->GetFName()))
+					{
+						CategoryName = PropertyCatagoryName;
+					}
+
+					// Add a property to the default category
+					FDetailCategoryImpl& CategoryImpl = InDetailLayout.DefaultCategory(CategoryName);
+					CategoryImpl.AddPropertyNode(ChildNodePtr.ToSharedRef(), InstanceName);
+				}
+
+				bool bRecurseIntoChildren =
+					!bIsChildOfCustomizedStruct // Don't recurse into built in struct children, we already know what they are and how to display them
+					&&  !bIsCustomizedStruct // Don't recurse into customized structs
+					&&	!bIsChildOfArray // Do not recurse into arrays, the children are drawn by the array property parent
+					&&	!bIsEditInlineNew // Edit inline new children are not supported for customization yet
+					&&	bIsUserVisible // Properties must be allowed to be visible by a user if they are not then their children are not visible either
+					&& (!bIsStruct || bPushOutStructProps); //  Only recurse into struct properties if they are going to be displayed as standalone properties in categories instead of inside an expandable area inside a category
+
+				if (bRecurseIntoChildren)
+				{
+					// Built in struct properties or children of arras 
+					UpdatePropertyMapRecursive(ChildNode, InDetailLayout, CurCategory, CurObjectNode);
+				}
+			}
+		}
+	}
+}
+
+void SDetailsViewBase::UpdatePropertyMap()
+{
+	// Reset everything
+	ClassToPropertyMap.Empty();
+	ClassesWithProperties.Empty();
+
+	// We need to be able to create a new detail layout and properly clean up the old one in the process
+	check(!DetailLayout.IsValid() || DetailLayout.IsUnique());
+
+	RootTreeNodes.Empty();
+
+	DetailLayout = MakeShareable(new FDetailLayoutBuilderImpl(ClassToPropertyMap, PropertyUtilities.ToSharedRef(), SharedThis(this)));
+
+
+	auto RootPropertyNode = GetRootNode();
+	check(RootPropertyNode.IsValid());
+	// Currently object property nodes do not provide any useful information other than being a container for its children.  We do not draw anything for them.
+	// When we encounter object property nodes, add their children instead of adding them to the tree.
+	UpdatePropertyMapRecursive(*RootPropertyNode, *DetailLayout, NAME_None, RootPropertyNode.Get());
+
+	// Ask for custom detail layouts
+	QueryCustomDetailLayout(*DetailLayout);
+
+	DetailLayout->GenerateDetailLayout();
+
+	UpdateFilteredDetails();
+}
+
diff --git a/Engine/Source/Editor/PropertyEditor/Private/SDetailsViewBase.h b/Engine/Source/Editor/PropertyEditor/Private/SDetailsViewBase.h
--- a/Engine/Source/Editor/PropertyEditor/Private/SDetailsViewBase.h
+++ b/Engine/Source/Editor/PropertyEditor/Private/SDetailsViewBase.h
@@ -14,11 +14,11 @@ class FDetailLayoutBuilderImpl;
 struct FPropertyNodeMap
 {
 	FPropertyNodeMap()
-		: ParentObjectProperty(NULL)
+		: ParentProperty(NULL)
 	{}
 
 	/** Object property node which contains the properties in the node map */
-	FObjectPropertyNode* ParentObjectProperty;
+	FPropertyNode* ParentProperty;
 
 	/** Property name to property node map */
 	TMap<FName, TSharedPtr<FPropertyNode> > PropertyNameToNode;
@@ -101,6 +101,8 @@ public:
 		return OnFinishedChangingPropertiesDelegate; 
 	}
 
+	virtual bool IsConnected() const = 0;
+
 	/**
 	 * Called when the open color picker window associated with this details view is closed
 	 */
@@ -154,6 +156,16 @@ public:
 	 */
 	void SetRootExpansionStates(const bool bExpand, const bool bRecurse);
 
+	/**
+	 * Queries a layout for a specific class
+	 */
+	void QueryLayoutForClass(FDetailLayoutBuilderImpl& CustomDetailLayout, UStruct* Class);
+
+	/**
+	 * Calls a delegate for each registered class that has properties visible to get any custom detail layouts
+	 */
+	void QueryCustomDetailLayout(class FDetailLayoutBuilderImpl& CustomDetailLayout);
+
 	/**
 	 * Refreshes the detail's treeview
 	 */
@@ -218,8 +230,52 @@ public:
 
 	virtual ~SDetailsViewBase();
 
+	// SWidget interface
+	virtual bool SupportsKeyboardFocus() const override;
+	virtual FReply OnKeyboardFocusReceived(const FGeometry& MyGeometry, const FKeyboardFocusEvent& InKeyboardFocusEvent) override;
+	virtual void Tick(const FGeometry& AllottedGeometry, const double InCurrentTime, const float InDeltaTime) override;
+	// End of SWidget interface
+
+	/** Saves the expansion state of property nodes for the selected object set */
+	void SaveExpandedItems();
+
+	/**
+	* Restores the expansion state of property nodes for the selected object set
+	*
+	* @param InitialStartNode The starting node if any.  If one is not supplied the expansion state is restored from the root node
+	*/
+	void RestoreExpandedItems(TSharedPtr<FPropertyNode> InitialStartNode = NULL);
+
+	virtual TSharedPtr<FComplexPropertyNode> GetRootNode() = 0;
+
+	/**
+	 * Creates the color picker window for this property view.
+	 *
+	 * @param PropertyEditor				The slate property node to edit.
+	 * @param bUseAlpha			Whether or not alpha is supported
+	 */
+	void CreateColorPickerWindow(const TSharedRef< class FPropertyEditor >& PropertyEditor, bool bUseAlpha) override;
+
 protected:
 
+	/**
+	 * Called when a color property is changed from a color picker
+	 */
+	void SetColorPropertyFromColorPicker(FLinearColor NewColor);
+
+	/** Updates the property map for access when customizing the details view.  Generates default layout for properties */
+	void UpdatePropertyMap();
+
+	/** 
+	 * Recursively updates children of property nodes. Generates default layout for properties 
+	 * 
+	 * @param InNode	The parent node to get children from
+	 * @param The detail layout builder that will be used for customization of this property map
+	 * @param CurCategory The current category name
+	 */
+	void UpdatePropertyMapRecursive( FPropertyNode& InNode, FDetailLayoutBuilderImpl& DetailLayout, FName CurCategory, FComplexPropertyNode* CurObjectNode );
+
+
 	/** Called to get the visibility of the tree view */
 	EVisibility GetTreeVisibility() const;
 
@@ -270,7 +326,10 @@ protected:
 	/** Called when show all advanced is clicked */
 	void OnShowAllAdvancedClicked();
 
-	virtual void UpdateFilteredDetails() = 0;
+	/**
+	* Updates the details with the passed in filter
+	*/
+	void UpdateFilteredDetails();
 
 	/** Called when the filter text changes.  This filters specific property nodes out of view */
 	void OnFilterTextChanged(const FText& InFilterText);
@@ -282,6 +341,9 @@ protected:
 	 */
 	void FilterView( const FString& InFilterText );
 
+	/** Called to get the visibility of the filter box */
+	EVisibility GetFilterBoxVisibility() const;
+
 protected:
 	/** The user defined args for the details view */
 	FDetailsViewArgs DetailsViewArgs;
@@ -337,4 +399,10 @@ protected:
 	FOnGetDetailCustomizationInstance GenericLayoutDelegate;
 	/** Actions that should be executed next tick */
 	TArray<FSimpleDelegate> DeferredActions;
+
+	/** Root tree node that needs to be destroyed when safe */
+	TSharedPtr<FComplexPropertyNode> RootNodePendingKill;
+
+	/** External property nodes which need to validated each tick */
+	TArray< TWeakPtr<FObjectPropertyNode> > ExternalRootPropertyNodes;
 };
\ No newline at end of file
diff --git a/Engine/Source/Editor/PropertyEditor/Private/UserInterface/PropertyEditor/SPropertyEditorEditInline.cpp b/Engine/Source/Editor/PropertyEditor/Private/UserInterface/PropertyEditor/SPropertyEditorEditInline.cpp
--- a/Engine/Source/Editor/PropertyEditor/Private/UserInterface/PropertyEditor/SPropertyEditorEditInline.cpp
+++ b/Engine/Source/Editor/PropertyEditor/Private/UserInterface/PropertyEditor/SPropertyEditorEditInline.cpp
@@ -109,11 +109,11 @@ bool SPropertyEditorEditInline::Supports( const FPropertyNode* InTreeNode, int32
 
 	bool bCanEditInline = false;
 
-	if( InTreeNode->HasNodeFlags(EPropertyNodeFlags::EditInline) != 0 )
+	const FObjectPropertyNode* ObjectPropertyNode = InTreeNode->FindObjectItemParent();
+	if (ObjectPropertyNode && (InTreeNode->HasNodeFlags(EPropertyNodeFlags::EditInline) != 0))
 	{
 		// Assume objects with the editinline flag can be editinline.  If any of the parent object is a default subobject then we cannot allow editinline
 		bCanEditInline = true;
-		const FObjectPropertyNode* ObjectPropertyNode = InTreeNode->FindObjectItemParent();
 		for( int32 ObjectIndex = 0; ObjectIndex < ObjectPropertyNode->GetNumObjects(); ++ObjectIndex )
 		{
 			const UObject* Object = ObjectPropertyNode->GetUObject(ObjectIndex);
diff --git a/Engine/Source/Editor/PropertyEditor/Private/UserInterface/PropertyEditor/SPropertyEditorNumeric.h b/Engine/Source/Editor/PropertyEditor/Private/UserInterface/PropertyEditor/SPropertyEditorNumeric.h
--- a/Engine/Source/Editor/PropertyEditor/Private/UserInterface/PropertyEditor/SPropertyEditorNumeric.h
+++ b/Engine/Source/Editor/PropertyEditor/Private/UserInterface/PropertyEditor/SPropertyEditorNumeric.h
@@ -77,11 +77,15 @@ public:
 			UE_LOG(LogPropertyNode, Warning, TEXT("UI Min (%s) >= UI Max (%s) for Ranged Numeric"), *UIMinString, *UIMaxString );
 		}
 
+		FObjectPropertyNode* ObjectPropertyNode = PropertyNode->FindObjectItemParent();
+		const bool bAllowSpin = (!ObjectPropertyNode || (1 == ObjectPropertyNode->GetNumObjects()))
+			&& !PropertyNode->GetProperty()->GetBoolMetaData("NoSpinbox");
+
 		ChildSlot
 			[
 				SAssignNew( PrimaryWidget, SNumericEntryBox<NumericType> )
 				// Only allow spinning if we have a single value
-				.AllowSpin( PropertyNode->FindObjectItemParent()->GetNumObjects() == 1 && !PropertyNode->GetProperty()->GetBoolMetaData("NoSpinbox") )
+				.AllowSpin(bAllowSpin)
 				.Value( this, &SPropertyEditorNumeric<NumericType>::OnGetValue )
 				.Font( InArgs._Font )
 				.MinValue(MinValue)
diff --git a/Engine/Source/Editor/PropertyEditor/Public/IDetailsView.h b/Engine/Source/Editor/PropertyEditor/Public/IDetailsView.h
--- a/Engine/Source/Editor/PropertyEditor/Public/IDetailsView.h
+++ b/Engine/Source/Editor/PropertyEditor/Public/IDetailsView.h
@@ -77,6 +77,9 @@ public:
 	/** Gets the base class being viewed */
 	virtual const UClass* GetBaseClass() const = 0;
 	
+	/** Gets the base struct being viewed */
+	virtual UStruct* GetBaseStruct() const = 0;
+
 	/**
 	 * Registers a custom detail layout delegate for a specific class in this instance of the details view only
 	 *
