Fix for Node Expansion in StructureDetailView

diff --git a/Engine/Source/Editor/PropertyEditor/Private/SDetailsViewBase.cpp b/Engine/Source/Editor/PropertyEditor/Private/SDetailsViewBase.cpp
--- a/Engine/Source/Editor/PropertyEditor/Private/SDetailsViewBase.cpp
+++ b/Engine/Source/Editor/PropertyEditor/Private/SDetailsViewBase.cpp
@@ -711,15 +711,36 @@ void SDetailsViewBase::SaveExpandedItems()
 	//while a valid class, and we're either the same as the base class (for multiple actors being selected and base class is AActor) OR we're not down to AActor yet)
 	for (UStruct* Struct = BestBaseStruct; Struct && ((BestBaseStruct == Struct) || (Struct != AActor::StaticClass())); Struct = Struct->GetSuperStruct())
 	{
-		if (RootPropertyNode->GetNumChildNodes() > 0 && ExpandedPropertyItems.Num() > 0)
+		if (RootPropertyNode->GetNumChildNodes() > 0)
 		{
-			GConfig->SetSingleLineArray(TEXT("DetailPropertyExpansion"), *Struct->GetName(), ExpandedPropertyItems, GEditorUserSettingsIni);
+			bool bShouldSave = ExpandedPropertyItems.Num() > 0;
+			if (!bShouldSave)
+			{
+				TArray<FString> DummyExpandedPropertyItems;
+				GConfig->GetSingleLineArray(TEXT("DetailPropertyExpansion"), *Struct->GetName(), DummyExpandedPropertyItems, GEditorUserSettingsIni);
+				bShouldSave = DummyExpandedPropertyItems.Num() > 0;
+			}
+
+			if (bShouldSave)
+			{
+				GConfig->SetSingleLineArray(TEXT("DetailPropertyExpansion"), *Struct->GetName(), ExpandedPropertyItems, GEditorUserSettingsIni);
+			}
 		}
 	}
 
-	if (DetailLayout.IsValid() && BestBaseStruct && !ExpandedCustomItemsString.IsEmpty())
+	if (DetailLayout.IsValid() && BestBaseStruct)
 	{
-		GConfig->SetString(TEXT("DetailCustomWidgetExpansion"), *BestBaseStruct->GetName(), *ExpandedCustomItemsString, GEditorUserSettingsIni);
+		bool bShouldSave = !ExpandedCustomItemsString.IsEmpty();
+		if (!bShouldSave)
+		{
+			FString DummyExpandedCustomItemsString;
+			GConfig->GetString(TEXT("DetailCustomWidgetExpansion"), *BestBaseStruct->GetName(), DummyExpandedCustomItemsString, GEditorUserSettingsIni);
+			bShouldSave = !DummyExpandedCustomItemsString.IsEmpty();
+		}
+		if (bShouldSave)
+		{
+			GConfig->SetString(TEXT("DetailCustomWidgetExpansion"), *BestBaseStruct->GetName(), *ExpandedCustomItemsString, GEditorUserSettingsIni);
+		}
 	}
 }
 
