Improve the display name for property transactions

diff --git a/Engine/Source/Editor/DetailCustomizations/Private/MathStructCustomizations.cpp b/Engine/Source/Editor/DetailCustomizations/Private/MathStructCustomizations.cpp
--- a/Engine/Source/Editor/DetailCustomizations/Private/MathStructCustomizations.cpp
+++ b/Engine/Source/Editor/DetailCustomizations/Private/MathStructCustomizations.cpp
@@ -587,7 +587,7 @@ void FColorStructCustomization::OnColorPickerInteractiveBegin()
 {
 	bIsInteractive = true;
 
-	GEditor->BeginTransaction( NSLOCTEXT("FColorStructCustomization", "SetColorProperty", "Set Color Property") );
+	GEditor->BeginTransaction( FText::Format( NSLOCTEXT("FColorStructCustomization", "SetColorProperty", "Edit {0}"), FText::FromString( StructPropertyHandle->GetPropertyDisplayName() ) ) );
 }
 
 void FColorStructCustomization::OnColorPickerInteractiveEnd()
diff --git a/Engine/Source/Editor/PropertyEditor/Private/PropertyHandleImpl.cpp b/Engine/Source/Editor/PropertyEditor/Private/PropertyHandleImpl.cpp
--- a/Engine/Source/Editor/PropertyEditor/Private/PropertyHandleImpl.cpp
+++ b/Engine/Source/Editor/PropertyEditor/Private/PropertyHandleImpl.cpp
@@ -389,7 +389,7 @@ FPropertyAccess::Result FPropertyValueImpl::ImportText( const TArray<FObjectBase
 				// Begin a transaction only if we need to call PreChange
 				if( GEditor && bTransactable )
 				{
-					GEditor->BeginTransaction( TEXT("PropertyEditor"), NSLOCTEXT("UnrealEd", "PropertyWindowEditProperties", "Edit Properties"), NodeProperty );
+					GEditor->BeginTransaction( TEXT("PropertyEditor"), FText::Format( NSLOCTEXT("PropertyEditor", "EditPropertyTransaction", "Edit {0}"), FText::FromString( InPropertyNode->GetDisplayName() ) ), NodeProperty );
 				}
 				
 				InPropertyNode->NotifyPreChange( NodeProperty, NotifyHook );
diff --git a/Engine/Source/Editor/PropertyEditor/Private/UserInterface/PropertyEditor/SPropertyEditorNumeric.h b/Engine/Source/Editor/PropertyEditor/Private/UserInterface/PropertyEditor/SPropertyEditorNumeric.h
--- a/Engine/Source/Editor/PropertyEditor/Private/UserInterface/PropertyEditor/SPropertyEditorNumeric.h
+++ b/Engine/Source/Editor/PropertyEditor/Private/UserInterface/PropertyEditor/SPropertyEditorNumeric.h
@@ -222,7 +222,7 @@ private:
 	{
 		bIsUsingSlider = true;
 
-		GEditor->BeginTransaction( TEXT("PropertyEditor"), NSLOCTEXT("UnrealEd", "SetNumberProperty", "Set Number Property"), PropertyEditor->GetPropertyHandle()->GetProperty() );
+		GEditor->BeginTransaction( TEXT("PropertyEditor"), FText::Format( NSLOCTEXT("PropertyEditor", "SetNumericPropertyTransaction", "Edit {0}"), FText::FromString( PropertyEditor->GetDisplayName() ) ), PropertyEditor->GetPropertyHandle()->GetProperty() );
 	}
 
 
