Some static analysis fixes.

From here: https://udn.unrealengine.com/questions/210807/static-code-analysis.html

#fix Redundant comparison of RootEdNode removed.
#fix Incorrect operator used to test FPropertyTable::SelectionUnit flags.
#fix Missing comparison used for testing FSlateColor::ColorUseRule

diff --git a/Engine/Source/Editor/BehaviorTreeEditor/Private/BehaviorTreeGraph.cpp b/Engine/Source/Editor/BehaviorTreeEditor/Private/BehaviorTreeGraph.cpp
--- a/Engine/Source/Editor/BehaviorTreeEditor/Private/BehaviorTreeGraph.cpp
+++ b/Engine/Source/Editor/BehaviorTreeEditor/Private/BehaviorTreeGraph.cpp
@@ -412,7 +412,7 @@ namespace BTGraphHelpers
 
 	void CreateChildren(UBehaviorTree* BTAsset, UBTCompositeNode* RootNode, const UBehaviorTreeGraphNode* RootEdNode, uint16* ExecutionIndex, uint8 TreeDepth)
 	{
-		if (RootEdNode == NULL || RootEdNode == NULL)
+		if (RootEdNode == NULL)
 		{
 			return;
 		}
diff --git a/Engine/Source/Editor/PropertyEditor/Private/Presentation/PropertyTable/PropertyTable.cpp b/Engine/Source/Editor/PropertyEditor/Private/Presentation/PropertyTable/PropertyTable.cpp
--- a/Engine/Source/Editor/PropertyEditor/Private/Presentation/PropertyTable/PropertyTable.cpp
+++ b/Engine/Source/Editor/PropertyEditor/Private/Presentation/PropertyTable/PropertyTable.cpp
@@ -296,12 +296,12 @@ void FPropertyTable::SetOrientation( EPropertyTableOrientation::Type InOrientati
 
 bool FPropertyTable::CanSelectCells() const
 {
-	return ( SelectionUnit | EPropertyTableSelectionUnit::Cell ) != 0;
+	return ( SelectionUnit & EPropertyTableSelectionUnit::Cell ) != 0;
 }
 
 bool FPropertyTable::CanSelectRows() const
 {
-	return ( SelectionUnit | EPropertyTableSelectionUnit::Row ) != 0;
+	return ( SelectionUnit & EPropertyTableSelectionUnit::Row ) != 0;
 }
 
 void FPropertyTable::SetSelectionMode( const ESelectionMode::Type Mode )
diff --git a/Engine/Source/Runtime/SlateCore/Public/Styling/SlateColor.h b/Engine/Source/Runtime/SlateCore/Public/Styling/SlateColor.h
--- a/Engine/Source/Runtime/SlateCore/Public/Styling/SlateColor.h
+++ b/Engine/Source/Runtime/SlateCore/Public/Styling/SlateColor.h
@@ -136,7 +136,7 @@ public:
 	 */
 	bool IsColorSpecified( ) const
 	{
-		return (ColorUseRule == ESlateColorStylingMode::UseColor_Specified) || (ESlateColorStylingMode::UseColor_Specified_Link);
+		return (ColorUseRule == ESlateColorStylingMode::UseColor_Specified) || (ColorUseRule == ESlateColorStylingMode::UseColor_Specified_Link);
 	}
 
 public:
