Expanded  byte properties so you can actually see all three numbers inside an array

diff --git a/Engine/Source/Editor/PropertyEditor/Private/UserInterface/PropertyEditor/SPropertyEditorNumeric.h b/Engine/Source/Editor/PropertyEditor/Private/UserInterface/PropertyEditor/SPropertyEditorNumeric.h
--- a/Engine/Source/Editor/PropertyEditor/Private/UserInterface/PropertyEditor/SPropertyEditorNumeric.h
+++ b/Engine/Source/Editor/PropertyEditor/Private/UserInterface/PropertyEditor/SPropertyEditorNumeric.h
@@ -121,8 +121,8 @@ public:
 
 		if( bIsNonEnumByte )
 		{
-			OutMinDesiredWidth = 60.0f;
-			OutMaxDesiredWidth = 60.0f;
+			OutMinDesiredWidth = 75.0f;
+			OutMaxDesiredWidth = 75.0f;
 		}
 		else
 		{
