BP, UI: in details panel for bp nodes with customizable pins (dealing with structures) no "reset to default" option is displayed for "Show Pin for Property"

diff --git a/Engine/Source/Editor/DetailCustomizations/Private/SkeletalControlNodeDetails.cpp b/Engine/Source/Editor/DetailCustomizations/Private/SkeletalControlNodeDetails.cpp
--- a/Engine/Source/Editor/DetailCustomizations/Private/SkeletalControlNodeDetails.cpp
+++ b/Engine/Source/Editor/DetailCustomizations/Private/SkeletalControlNodeDetails.cpp
@@ -20,7 +20,9 @@ void FSkeletalControlNodeDetails::CustomizeDetails(class IDetailLayoutBuilder& D
 	TSharedRef<IPropertyHandle> AvailablePins = DetailBuilder.GetProperty("ShowPinForProperties");
 
 	//@TODO: Shouldn't show this if the available pins array is empty!
-	TSharedRef<FDetailArrayBuilder> AvailablePinsBuilder = MakeShareable( new FDetailArrayBuilder( AvailablePins ) );
+	const bool bGenerateHeader = true;
+	const bool bDisplayResetToDefault = false;
+	TSharedRef<FDetailArrayBuilder> AvailablePinsBuilder = MakeShareable( new FDetailArrayBuilder( AvailablePins, bGenerateHeader, bDisplayResetToDefault ) );
 	AvailablePinsBuilder->OnGenerateArrayElementWidget( FOnGenerateArrayElementWidget::CreateSP( this, &FSkeletalControlNodeDetails::OnGenerateElementForPropertyPin ) );
 
 	const bool bForAdvanced = false;
diff --git a/Engine/Source/Editor/PropertyEditor/Public/PropertyCustomizationHelpers.h b/Engine/Source/Editor/PropertyEditor/Public/PropertyCustomizationHelpers.h
--- a/Engine/Source/Editor/PropertyEditor/Public/PropertyCustomizationHelpers.h
+++ b/Engine/Source/Editor/PropertyEditor/Public/PropertyCustomizationHelpers.h
@@ -183,10 +183,11 @@ class FDetailArrayBuilder : public IDetailCustomNodeBuilder
 {
 public:
 
-	FDetailArrayBuilder( TSharedRef<IPropertyHandle> InBaseProperty, bool InGenerateHeader = true )
+	FDetailArrayBuilder( TSharedRef<IPropertyHandle> InBaseProperty, bool InGenerateHeader = true, bool InDisplayResetToDefault = true)
 		: ArrayProperty( InBaseProperty->AsArray() )
 		, BaseProperty( InBaseProperty )
 		, bGenerateHeader( InGenerateHeader)
+		, bDisplayResetToDefault(InDisplayResetToDefault)
 	{
 		check( ArrayProperty.IsValid() );
 
@@ -228,30 +229,34 @@ public:
 	{
 		if (bGenerateHeader)
 		{
-			TSharedPtr<SResetToDefaultMenu> ResetToDefaultMenu;
-			bool bDisplayResetToDefault = false;
+			const bool bDisplayResetToDefaultInNameContent = false;
+			TSharedPtr<SHorizontalBox> ContentHorizontalBox;
 			NodeRow
-				.FilterString(DisplayName.Len() > 0 ? DisplayName : BaseProperty->GetPropertyDisplayName())
-				.NameContent()
+			.FilterString(DisplayName.Len() > 0 ? DisplayName : BaseProperty->GetPropertyDisplayName())
+			.NameContent()
+			[
+				BaseProperty->CreatePropertyNameWidget(DisplayName, bDisplayResetToDefaultInNameContent)
+			]
+			.ValueContent()
+			[
+				SAssignNew(ContentHorizontalBox, SHorizontalBox)
+				+ SHorizontalBox::Slot()
 				[
-					BaseProperty->CreatePropertyNameWidget(DisplayName, bDisplayResetToDefault)
+					BaseProperty->CreatePropertyValueWidget()
 				]
-			.ValueContent()
+			];
+
+			if (bDisplayResetToDefault)
+			{
+				TSharedPtr<SResetToDefaultMenu> ResetToDefaultMenu;
+				ContentHorizontalBox->AddSlot()
+				.AutoWidth()
+				.Padding(FMargin(2.0f, 0.0f, 0.0f, 0.0f))
 				[
-					SNew(SHorizontalBox)
-					+ SHorizontalBox::Slot()
-					[
-						BaseProperty->CreatePropertyValueWidget()
-					]
-					+ SHorizontalBox::Slot()
-						.AutoWidth()
-						.Padding(FMargin(2.0f, 0.0f, 0.0f, 0.0f))
-						[
-							SAssignNew(ResetToDefaultMenu, SResetToDefaultMenu)
-						]
+					SAssignNew(ResetToDefaultMenu, SResetToDefaultMenu)
 				];
-
-			ResetToDefaultMenu->AddProperty(BaseProperty);
+				ResetToDefaultMenu->AddProperty(BaseProperty);
+			}
 		}
 	}
 
@@ -287,6 +292,7 @@ private:
 	TSharedRef<IPropertyHandle> BaseProperty;
 	FSimpleDelegate OnRebuildChildren;
 	bool bGenerateHeader;
+	bool bDisplayResetToDefault;
 };
 
 /**
