Removed "Use" and "Browse for Asset" buttons from actors array property customizations

diff --git a/Engine/Source/Editor/PropertyEditor/Private/PropertyEditorHelpers.cpp b/Engine/Source/Editor/PropertyEditor/Private/PropertyEditorHelpers.cpp
--- a/Engine/Source/Editor/PropertyEditor/Private/PropertyEditorHelpers.cpp
+++ b/Engine/Source/Editor/PropertyEditor/Private/PropertyEditorHelpers.cpp
@@ -508,7 +508,7 @@ namespace PropertyEditorHelpers
 					if( ObjProp )
 					{
 						UClass* ObjPropClass = ObjProp->PropertyClass;
-						if ( ObjPropClass->IsChildOf( AActor::StaticClass() ) )
+						if ( !ObjPropClass->IsChildOf( AActor::StaticClass() ) )
 						{
 							OutRequiredButtons.Add( EPropertyButton::Use );
 						}
diff --git a/Engine/Source/Editor/PropertyEditor/Private/UserInterface/PropertyEditor/SPropertyEditorAsset.cpp b/Engine/Source/Editor/PropertyEditor/Private/UserInterface/PropertyEditor/SPropertyEditorAsset.cpp
--- a/Engine/Source/Editor/PropertyEditor/Private/UserInterface/PropertyEditor/SPropertyEditorAsset.cpp
+++ b/Engine/Source/Editor/PropertyEditor/Private/UserInterface/PropertyEditor/SPropertyEditorAsset.cpp
@@ -181,19 +181,19 @@ void SPropertyEditorAsset::Construct( const FArguments& InArgs, const TSharedPtr
 		[
 			PropertyCustomizationHelpers::MakeUseSelectedButton( FSimpleDelegate::CreateSP( this, &SPropertyEditorAsset::OnUse ) )
 		];
-	}
-
-	ButtonBox->AddSlot()
-	.AutoWidth()
-	.Padding( 2.0f, 0.0f )
-	.VAlign(VAlign_Center)
-	[
-		PropertyCustomizationHelpers::MakeBrowseButton(
-			FSimpleDelegate::CreateSP( this, &SPropertyEditorAsset::OnBrowse ),
-			TAttribute<FText>( this, &SPropertyEditorAsset::GetOnBrowseToolTip )
-			)
-	];
 
+		ButtonBox->AddSlot()
+		.AutoWidth()
+		.Padding( 2.0f, 0.0f )
+		.VAlign(VAlign_Center)
+		[
+			PropertyCustomizationHelpers::MakeBrowseButton(
+				FSimpleDelegate::CreateSP( this, &SPropertyEditorAsset::OnBrowse ),
+				TAttribute<FText>( this, &SPropertyEditorAsset::GetOnBrowseToolTip )
+				)
+		];
+	}
+	
 	if(bIsActor)
 	{
 		ButtonBox->AddSlot()
