Added ability to remove "Duplicate" operation from array element property customizations (using meta flag "NoElementDuplicate")

diff --git a/Engine/Source/Editor/PropertyEditor/Private/PropertyCustomizationHelpers.cpp b/Engine/Source/Editor/PropertyEditor/Private/PropertyCustomizationHelpers.cpp
--- a/Engine/Source/Editor/PropertyEditor/Private/PropertyCustomizationHelpers.cpp
+++ b/Engine/Source/Editor/PropertyEditor/Private/PropertyCustomizationHelpers.cpp
@@ -162,8 +162,12 @@ namespace PropertyCustomizationHelpers
 			FUIAction DeleteAction( OnDeleteClicked );
 			MenuContentBuilder.AddMenuEntry( LOCTEXT( "DeleteButtonLabel", "Delete"), FText::GetEmpty(), FSlateIcon(), DeleteAction );
 
-			FUIAction DuplicateAction( OnDuplicateClicked );
-			MenuContentBuilder.AddMenuEntry( LOCTEXT( "DuplicateButtonLabel", "Duplicate"), FText::GetEmpty(), FSlateIcon(), DuplicateAction );
+			// Duplicate operation is optional
+			if (OnDuplicateClicked.IsBound())
+			{
+				FUIAction DuplicateAction( OnDuplicateClicked );
+				MenuContentBuilder.AddMenuEntry( LOCTEXT( "DuplicateButtonLabel", "Duplicate"), FText::GetEmpty(), FSlateIcon(), DuplicateAction );
+			}
 		}
 
 		return
diff --git a/Engine/Source/Editor/PropertyEditor/Private/PropertyEditorHelpers.cpp b/Engine/Source/Editor/PropertyEditor/Private/PropertyEditorHelpers.cpp
--- a/Engine/Source/Editor/PropertyEditor/Private/PropertyEditorHelpers.cpp
+++ b/Engine/Source/Editor/PropertyEditor/Private/PropertyEditorHelpers.cpp
@@ -608,7 +608,14 @@ namespace PropertyEditorHelpers
 		{
 			if( PropertyNode->HasNodeFlags(EPropertyNodeFlags::SingleSelectOnly) && !(OuterArrayProp->PropertyFlags & CPF_EditFixedSize) )
 			{
-				OutRequiredButtons.Add( EPropertyButton::Insert_Delete_Duplicate );
+				if (OuterArrayProp->HasMetaData(TEXT("NoElementDuplicate")))
+				{
+					OutRequiredButtons.Add( EPropertyButton::Insert_Delete );
+				}
+				else
+				{
+					OutRequiredButtons.Add( EPropertyButton::Insert_Delete_Duplicate );
+				}
 			}
 		}
 
@@ -742,11 +749,16 @@ namespace PropertyEditorHelpers
 			NewButton = PropertyCustomizationHelpers::MakeEmptyButton( FSimpleDelegate::CreateSP( PropertyEditor, &FPropertyEditor::EmptyArray ) );
 			break;
 
+		case EPropertyButton::Insert_Delete:
 		case EPropertyButton::Insert_Delete_Duplicate:
 			{
 				FExecuteAction InsertAction = FExecuteAction::CreateSP( PropertyEditor, &FPropertyEditor::InsertItem );
 				FExecuteAction DeleteAction = FExecuteAction::CreateSP( PropertyEditor, &FPropertyEditor::DeleteItem );
-				FExecuteAction DuplicateAction = FExecuteAction::CreateSP( PropertyEditor, &FPropertyEditor::DuplicateItem );
+				FExecuteAction DuplicateAction;
+				if (ButtonType == EPropertyButton::Insert_Delete_Duplicate)
+				{
+					DuplicateAction = FExecuteAction::CreateSP( PropertyEditor, &FPropertyEditor::DuplicateItem );
+				}
 
 				NewButton = PropertyCustomizationHelpers::MakeInsertDeleteDuplicateButton( InsertAction, DeleteAction, DuplicateAction );
 
diff --git a/Engine/Source/Editor/PropertyEditor/Private/PropertyEditorHelpers.h b/Engine/Source/Editor/PropertyEditor/Private/PropertyEditorHelpers.h
--- a/Engine/Source/Editor/PropertyEditor/Private/PropertyEditorHelpers.h
+++ b/Engine/Source/Editor/PropertyEditor/Private/PropertyEditorHelpers.h
@@ -14,6 +14,7 @@ namespace EPropertyButton
 		Add,
 		Empty,
 		Insert_Delete_Duplicate,
+		Insert_Delete,
 		Insert,
 		Delete,
 		Duplicate,
