Removed "Used Selected in CB" button from all arrays customizations (actors, assets)

diff --git a/Engine/Source/Editor/PropertyEditor/Private/PropertyEditorHelpers.cpp b/Engine/Source/Editor/PropertyEditor/Private/PropertyEditorHelpers.cpp
--- a/Engine/Source/Editor/PropertyEditor/Private/PropertyEditorHelpers.cpp
+++ b/Engine/Source/Editor/PropertyEditor/Private/PropertyEditorHelpers.cpp
@@ -500,21 +500,6 @@ namespace PropertyEditorHelpers
 		{
 			if( !(NodeProperty->PropertyFlags & CPF_EditFixedSize) )
 			{
-				const UArrayProperty* ArrayProp = Cast<const UArrayProperty>( NodeProperty );
-				if (ArrayProp)
-				{
-					//if this array supports actors
-					UObjectPropertyBase* ObjProp = Cast<UObjectPropertyBase>( ArrayProp->Inner );
-					if( ObjProp )
-					{
-						UClass* ObjPropClass = ObjProp->PropertyClass;
-						if ( !ObjPropClass->IsChildOf( AActor::StaticClass() ) )
-						{
-							OutRequiredButtons.Add( EPropertyButton::Use );
-						}
-					}
-				}
-
 				OutRequiredButtons.Add( EPropertyButton::Add );
 				OutRequiredButtons.Add( EPropertyButton::Empty );
 			}
