Temporary fix to get 4.3 out the door, not painting the structure details view if it is pending refesh (to keep from painting deleted properties).
(TTP #340949)

diff --git a/Engine/Source/Editor/PropertyEditor/Private/SStructureDetailsView.cpp b/Engine/Source/Editor/PropertyEditor/Private/SStructureDetailsView.cpp
--- a/Engine/Source/Editor/PropertyEditor/Private/SStructureDetailsView.cpp
+++ b/Engine/Source/Editor/PropertyEditor/Private/SStructureDetailsView.cpp
@@ -203,6 +203,23 @@ void SStructureDetailsView::SetStructureData(TSharedPtr<FStructOnScope> InStruct
 	UpdatePropertyMap();
 }
 
+int32 SStructureDetailsView::OnPaint(FGeometry const& AllottedGeometry, FSlateRect const& MyClippingRect, FSlateWindowElementList& OutDrawElements, int32 LayerId, FWidgetStyle const& InWidgetStyle, bool bParentEnabled) const
+{
+	int32 LayerIdOut = LayerId;
+	// if the details table is still in need of a refresh, then wait for it to
+	// update before we paint (the tree view is probably holding onto some out
+	// of date properties, and painting those rows could cause a crash)
+	//
+	// @TODO: this is a temp fix to get 4.3 out the door, we need to find a
+	//        better, lower level solution (this could happen for other detail panels)
+	if (!DetailTree->IsPendingRefresh())
+	{
+		LayerIdOut = SDetailsViewBase::OnPaint(AllottedGeometry, MyClippingRect, OutDrawElements, LayerId, InWidgetStyle, bParentEnabled);
+	}
+	
+	return LayerIdOut;
+}
+
 void SStructureDetailsView::ForceRefresh()
 {
 	SetStructureData(StructData);
diff --git a/Engine/Source/Editor/PropertyEditor/Private/SStructureDetailsView.h b/Engine/Source/Editor/PropertyEditor/Private/SStructureDetailsView.h
--- a/Engine/Source/Editor/PropertyEditor/Private/SStructureDetailsView.h
+++ b/Engine/Source/Editor/PropertyEditor/Private/SStructureDetailsView.h
@@ -86,6 +86,10 @@ public:
 	virtual void SetObject(UObject* InObject, bool bForceRefresh = false) override{}
 
 	virtual TSharedPtr<class FComplexPropertyNode> GetRootNode() override;
+
+	// Begin SWidget interface
+	virtual int32 OnPaint(FGeometry const& AllottedGeometry, FSlateRect const& MyClippingRect, FSlateWindowElementList& OutDrawElements, int32 LayerId, FWidgetStyle const& InWidgetStyle, bool bParentEnabled) const override;
+	// End SWidget interface
 protected:
 	virtual void CustomUpdatePropertyMap() override;
 };
