More robust StructurePropertyNode

diff --git a/Engine/Source/Editor/PropertyEditor/Private/StructurePropertyNode.h b/Engine/Source/Editor/PropertyEditor/Private/StructurePropertyNode.h
--- a/Engine/Source/Editor/PropertyEditor/Private/StructurePropertyNode.h
+++ b/Engine/Source/Editor/PropertyEditor/Private/StructurePropertyNode.h
@@ -21,6 +21,7 @@ public:
 
 	void SetStructure(TSharedPtr<FStructOnScope> InStructData)
 	{
+		ClearCachedReadAddresses(true);
 		DestroyTree();
 		StructData = InStructData;
 	}
