Add the ability to link to object class documentation throughout the UI

diff --git a/Engine/Source/Editor/ClassViewer/Private/SClassViewer.cpp b/Engine/Source/Editor/ClassViewer/Private/SClassViewer.cpp
--- a/Engine/Source/Editor/ClassViewer/Private/SClassViewer.cpp
+++ b/Engine/Source/Editor/ClassViewer/Private/SClassViewer.cpp
@@ -21,6 +21,7 @@
 
 #include "UnloadedBlueprintData.h"
 
+#include "EditorClassUtils.h"
 #include "IDocumentation.h"
 
 #include "PropertyEditorModule.h"
@@ -1339,7 +1340,7 @@ public:
 					UPackage*  Package  = Class->GetOutermost();
 					UMetaData* MetaData = Package->GetMetaData();
 
-					ToolTip = IDocumentation::Get()->CreateToolTip(Class->GetToolTipText(), nullptr, "Shared/Classes", Class->GetName());
+					ToolTip = FEditorClassUtils::GetTooltip(Class);
 				}
 
 				return ToolTip;
diff --git a/Engine/Source/Editor/Documentation/Private/Documentation.cpp b/Engine/Source/Editor/Documentation/Private/Documentation.cpp
--- a/Engine/Source/Editor/Documentation/Private/Documentation.cpp
+++ b/Engine/Source/Editor/Documentation/Private/Documentation.cpp
@@ -115,9 +115,10 @@ bool FDocumentation::Open(const FString& Link, const TSharedRef<FCulture, ESPMod
 	return !DocumentationUrl.IsEmpty();
 }
 
-TSharedRef< SWidget > FDocumentation::CreateAnchor( const FString& Link, const FString& PreviewLink, const FString& PreviewExcerptName ) const
+TSharedRef< SWidget > FDocumentation::CreateAnchor( const TAttribute<FString>& Link, const FString& PreviewLink, const FString& PreviewExcerptName ) const
 {
-	return SNew( SDocumentationAnchor, Link )
+	return SNew( SDocumentationAnchor )
+		.Link(Link)
 		.PreviewLink(PreviewLink)
 		.PreviewExcerptName(PreviewExcerptName);
 }
diff --git a/Engine/Source/Editor/Documentation/Private/Documentation.h b/Engine/Source/Editor/Documentation/Private/Documentation.h
--- a/Engine/Source/Editor/Documentation/Private/Documentation.h
+++ b/Engine/Source/Editor/Documentation/Private/Documentation.h
@@ -23,7 +23,7 @@ public:
 
 	virtual bool Open(const FString& Link, const TSharedRef<FCulture, ESPMode::ThreadSafe>& Culture) const override;
 
-	virtual TSharedRef< SWidget > CreateAnchor( const FString& Link, const FString& PreviewLink = FString(), const FString& PreviewExcerptName = FString() ) const override;
+	virtual TSharedRef< SWidget > CreateAnchor( const TAttribute<FString>& Link, const FString& PreviewLink = FString(), const FString& PreviewExcerptName = FString() ) const override;
 
 	virtual TSharedRef< IDocumentationPage > GetPage( const FString& Link, const TSharedPtr< FParserConfiguration >& Config, const FDocumentationStyle& Style = FDocumentationStyle() ) override;
 
diff --git a/Engine/Source/Editor/Documentation/Private/DocumentationPage.cpp b/Engine/Source/Editor/Documentation/Private/DocumentationPage.cpp
--- a/Engine/Source/Editor/Documentation/Private/DocumentationPage.cpp
+++ b/Engine/Source/Editor/Documentation/Private/DocumentationPage.cpp
@@ -40,6 +40,20 @@ int32 FDocumentationPage::GetNumExcerpts() const
 	return StoredExcerpts.Num();
 }
 
+bool FDocumentationPage::GetExcerpt(const FString& ExcerptName, FExcerpt& Excerpt)
+{
+	for (const FExcerpt& StoredExcerpt : StoredExcerpts)
+	{
+		if (StoredExcerpt.Name == ExcerptName)
+		{
+			Excerpt = StoredExcerpt;
+			return true;
+		}
+	}
+
+	return false;
+}
+
 void FDocumentationPage::GetExcerpts( /*OUT*/ TArray< FExcerpt >& Excerpts )
 {
 	Excerpts.Empty();
diff --git a/Engine/Source/Editor/Documentation/Private/DocumentationPage.h b/Engine/Source/Editor/Documentation/Private/DocumentationPage.h
--- a/Engine/Source/Editor/Documentation/Private/DocumentationPage.h
+++ b/Engine/Source/Editor/Documentation/Private/DocumentationPage.h
@@ -16,6 +16,7 @@ public:
 	~FDocumentationPage();
 	virtual bool HasExcerpt( const FString& ExcerptName ) override;
 	virtual int32 GetNumExcerpts() const override;
+	virtual bool GetExcerpt( const FString& ExcerptName, FExcerpt& Excerpt) override;
 	virtual void GetExcerpts( /*OUT*/ TArray< FExcerpt >& Excerpts ) override;
 	virtual bool GetExcerptContent( FExcerpt& Excerpt ) override;
 
diff --git a/Engine/Source/Editor/Documentation/Private/SDocumentationAnchor.cpp b/Engine/Source/Editor/Documentation/Private/SDocumentationAnchor.cpp
--- a/Engine/Source/Editor/Documentation/Private/SDocumentationAnchor.cpp
+++ b/Engine/Source/Editor/Documentation/Private/SDocumentationAnchor.cpp
@@ -5,9 +5,9 @@
 #include "DocumentationLink.h"
 #include "IDocumentation.h"
 
-void SDocumentationAnchor::Construct(const FArguments& InArgs, const FString& InLink )
+void SDocumentationAnchor::Construct(const FArguments& InArgs )
 {
-	Link = InLink;
+	Link = InArgs._Link;
 	
 	struct StringToText_Local
 	{
@@ -76,6 +76,6 @@ const FSlateBrush* SDocumentationAnchor::GetButtonImage() const
 
 FReply SDocumentationAnchor::OnClicked() const
 {
-	IDocumentation::Get()->Open( Link );
+	IDocumentation::Get()->Open( Link.Get() );
 	return FReply::Handled();
 }
diff --git a/Engine/Source/Editor/Documentation/Private/SDocumentationAnchor.h b/Engine/Source/Editor/Documentation/Private/SDocumentationAnchor.h
--- a/Engine/Source/Editor/Documentation/Private/SDocumentationAnchor.h
+++ b/Engine/Source/Editor/Documentation/Private/SDocumentationAnchor.h
@@ -8,9 +8,14 @@ public:
 	{}
 		SLATE_ARGUMENT( FString, PreviewLink )
 		SLATE_ARGUMENT( FString, PreviewExcerptName )
+
+		/** The string for the link to follow when clicked  */
+		SLATE_ATTRIBUTE( FString, Link )
+
 	SLATE_END_ARGS()
 
-	void Construct(const FArguments& InArgs, const FString& InLink );
+
+	void Construct(const FArguments& InArgs);
 
 private:
 	
@@ -20,7 +25,7 @@ private:
 
 private:
 
-	FString Link;
+	TAttribute<FString> Link;
 	TSharedPtr< SButton > Button;
 	TSharedPtr< SImage > ButtonImage;
 	const FSlateBrush* Default; 
diff --git a/Engine/Source/Editor/Documentation/Public/IDocumentation.h b/Engine/Source/Editor/Documentation/Public/IDocumentation.h
--- a/Engine/Source/Editor/Documentation/Public/IDocumentation.h
+++ b/Engine/Source/Editor/Documentation/Public/IDocumentation.h
@@ -147,7 +147,7 @@ public:
 
 	virtual bool Open(const FString& Link, const TSharedRef<FCulture, ESPMode::ThreadSafe>& Culture) const = 0;
 
-	virtual TSharedRef< class SWidget > CreateAnchor( const FString& Link, const FString& PreviewLink = FString(), const FString& PreviewExcerptName = FString() ) const = 0;
+	virtual TSharedRef< class SWidget > CreateAnchor( const TAttribute<FString>& Link, const FString& PreviewLink = FString(), const FString& PreviewExcerptName = FString() ) const = 0;
 
 	virtual TSharedRef< class IDocumentationPage > GetPage( const FString& Link, const TSharedPtr< FParserConfiguration >& Config, const FDocumentationStyle& Style = FDocumentationStyle() ) = 0;
 
diff --git a/Engine/Source/Editor/Documentation/Public/IDocumentationPage.h b/Engine/Source/Editor/Documentation/Public/IDocumentationPage.h
--- a/Engine/Source/Editor/Documentation/Public/IDocumentationPage.h
+++ b/Engine/Source/Editor/Documentation/Public/IDocumentationPage.h
@@ -5,6 +5,10 @@
 
 struct FExcerpt
 {
+	FExcerpt()
+		: LineNumber( INDEX_NONE )
+	{}
+
 	FExcerpt( const FString& InName, const TSharedPtr< SWidget >& InContent, const TMap< FString, FString >& InVariables, int32 InLineNumber )
 		: Name( InName )
 		, Content( InContent )
@@ -23,6 +27,7 @@ class IDocumentationPage
 public:
 	virtual bool HasExcerpt( const FString& ExcerptName ) = 0;
 	virtual int32 GetNumExcerpts() const = 0;
+	virtual bool GetExcerpt( const FString& ExcerptName, FExcerpt& Excerpt) = 0;
 	virtual void GetExcerpts( /*OUT*/ TArray< FExcerpt >& Excerpts ) = 0;
 	virtual bool GetExcerptContent( FExcerpt& Excerpt ) = 0;
 
diff --git a/Engine/Source/Editor/GameProjectGeneration/Private/SNewClassDialog.cpp b/Engine/Source/Editor/GameProjectGeneration/Private/SNewClassDialog.cpp
--- a/Engine/Source/Editor/GameProjectGeneration/Private/SNewClassDialog.cpp
+++ b/Engine/Source/Editor/GameProjectGeneration/Private/SNewClassDialog.cpp
@@ -8,6 +8,7 @@
 #include "Editor/ClassViewer/Private/SClassViewer.h"
 #include "DesktopPlatformModule.h"
 #include "Editor/Documentation/Public/IDocumentation.h"
+#include "EditorClassUtils.h"
 
 #define LOCTEXT_NAMESPACE "GameProjectGeneration"
 
@@ -90,6 +91,9 @@ void SNewClassDialog::Construct( const FArguments& InArgs )
 
 	ClassViewer = StaticCastSharedRef<SClassViewer>(FModuleManager::LoadModuleChecked<FClassViewerModule>("ClassViewer").CreateClassViewer(Options, FOnClassPicked::CreateSP(this, &SNewClassDialog::OnAdvancedClassSelected)));
 
+	TSharedRef<SWidget> DocWidget = IDocumentation::Get()->CreateAnchor(TAttribute<FString>(this, &SNewClassDialog::GetSelectedParentDocLink));
+	DocWidget->SetVisibility(TAttribute<EVisibility>(this, &SNewClassDialog::GetDocLinkVisibility));
+
 	const float EditableTextHeight = 26.0f;
 
 	ChildSlot
@@ -244,8 +248,20 @@ void SNewClassDialog::Construct( const FArguments& InArgs )
 							.VAlign(VAlign_Center)
 							.Padding(0, 0, 12, 0)
 							[
-								SNew(STextBlock)
-								.Text( this, &SNewClassDialog::GetSelectedParentClassName )
+								SNew(SHorizontalBox)
+
+								+SHorizontalBox::Slot()
+								.AutoWidth()
+								[
+									SNew(STextBlock)
+									.Text( this, &SNewClassDialog::GetSelectedParentClassName )
+								]
+
+								+SHorizontalBox::Slot()
+								.AutoWidth()
+								[
+									DocWidget
+								]
 							]
 
 							+ SVerticalBox::Slot()
@@ -687,6 +703,16 @@ FString SNewClassDialog::GetSelectedParentClassFilename() const
 	return FString();
 }
 
+EVisibility SNewClassDialog::GetDocLinkVisibility() const
+{
+	return (ParentClassInfo.BaseClass == nullptr || FEditorClassUtils::GetDocumentationLink(ParentClassInfo.BaseClass).IsEmpty() ? EVisibility::Hidden : EVisibility::Visible);
+}
+
+FString SNewClassDialog::GetSelectedParentDocLink() const
+{
+	return FEditorClassUtils::GetDocumentationLink(ParentClassInfo.BaseClass);
+}
+
 void SNewClassDialog::OnEditCodeClicked()
 {
 	const FString ClassHeaderPath = GetClassHeaderPath(ParentClassInfo.BaseClass);
diff --git a/Engine/Source/Editor/GameProjectGeneration/Private/SNewClassDialog.h b/Engine/Source/Editor/GameProjectGeneration/Private/SNewClassDialog.h
--- a/Engine/Source/Editor/GameProjectGeneration/Private/SNewClassDialog.h
+++ b/Engine/Source/Editor/GameProjectGeneration/Private/SNewClassDialog.h
@@ -38,6 +38,12 @@ private:
 	/** Whether the hyper link to go to source should be visible */
 	EVisibility GetSourceHyperlinkVisibility() const;
 
+	/** Gets the currently select parent class's doc link */
+	FString GetSelectedParentDocLink() const;
+
+	/** Whether the document link anchor should be visible */
+	EVisibility GetDocLinkVisibility() const;
+	
 	/** Handler for when the seleted parent class's filename is clicked */
 	void OnEditCodeClicked();
 
diff --git a/Engine/Source/Editor/LevelEditor/Private/LevelEditor.cpp b/Engine/Source/Editor/LevelEditor/Private/LevelEditor.cpp
--- a/Engine/Source/Editor/LevelEditor/Private/LevelEditor.cpp
+++ b/Engine/Source/Editor/LevelEditor/Private/LevelEditor.cpp
@@ -570,6 +570,11 @@ void FLevelEditorModule::BindGlobalLevelEditorCommands()
 		FExecuteAction::CreateStatic( &FLevelEditorActionCallbacks::GoToCodeForActor_Clicked )
 		);
 
+	ActionList.MapAction( 
+		Commands.GoToDocsForActor, 
+		FExecuteAction::CreateStatic( &FLevelEditorActionCallbacks::GoToDocsForActor_Clicked )
+		);
+
 	ActionList.MapAction( 
 		FGenericCommands::Get().Duplicate, 
 		FExecuteAction::CreateStatic( &FLevelEditorActionCallbacks::ExecuteExecCommand, FString( TEXT("DUPLICATE") ) ),
diff --git a/Engine/Source/Editor/LevelEditor/Private/LevelEditorActions.cpp b/Engine/Source/Editor/LevelEditor/Private/LevelEditorActions.cpp
--- a/Engine/Source/Editor/LevelEditor/Private/LevelEditorActions.cpp
+++ b/Engine/Source/Editor/LevelEditor/Private/LevelEditorActions.cpp
@@ -43,6 +43,7 @@
 #include "IAnalyticsProvider.h"
 #include "ReferenceViewer.h"
 #include "Developer/MeshUtilities/Public/MeshUtilities.h"
+#include "EditorClassUtils.h"
 
 #include "EditorActorFolders.h"
 #include "ActorPickerMode.h"
@@ -975,6 +976,19 @@ void FLevelEditorActionCallbacks::GoToCodeForActor_Clicked()
 	}
 }
 
+void FLevelEditorActionCallbacks::GoToDocsForActor_Clicked()
+{
+	const auto& SelectedActorInfo = AssetSelectionUtils::GetSelectedActorInfo();
+	if( SelectedActorInfo.SelectionClass != nullptr )
+	{
+		FString DocumentationLink = FEditorClassUtils::GetDocumentationLink(SelectedActorInfo.SelectionClass);
+		if (!DocumentationLink.IsEmpty())
+		{
+			IDocumentation::Get()->Open( DocumentationLink );
+		}
+	}
+}
+
 void FLevelEditorActionCallbacks::FindInContentBrowser_Clicked()
 {
 	GEditor->SyncToContentBrowser();
@@ -2637,6 +2651,7 @@ void FLevelEditorCommands::RegisterCommands()
 
 	UI_COMMAND( SnapCameraToActor, "Snap View to Actor", "Snaps the view to the selected actors", EUserInterfaceActionType::Button, FInputGesture() );
 	UI_COMMAND( GoToCodeForActor, "Go to C++ Code for Actor", "Opens a code editing IDE and navigates to the source file associated with the seleced actor", EUserInterfaceActionType::Button, FInputGesture() );
+	UI_COMMAND( GoToDocsForActor, "Go to Documentation for Actor", "Opens documentation for the Actor in the default web browser", EUserInterfaceActionType::Button, FInputGesture() );
 
 	UI_COMMAND( PasteHere, "Paste Here", "Pastes the actor at the click location", EUserInterfaceActionType::Button, FInputGesture() );
 
diff --git a/Engine/Source/Editor/LevelEditor/Private/LevelEditorContextMenu.cpp b/Engine/Source/Editor/LevelEditor/Private/LevelEditorContextMenu.cpp
--- a/Engine/Source/Editor/LevelEditor/Private/LevelEditorContextMenu.cpp
+++ b/Engine/Source/Editor/LevelEditor/Private/LevelEditorContextMenu.cpp
@@ -27,6 +27,7 @@
 #include "GlobalEditorCommonCommands.h"
 #include "LevelEditorCreateActorMenu.h"
 #include "SourceCodeNavigation.h"
+#include "EditorClassUtils.h"
 
 #define LOCTEXT_NAMESPACE "LevelViewportContextMenu"
 
@@ -235,19 +236,35 @@ void FLevelEditorContextMenu::FillMenu( FMenuBuilder& MenuBuilder, TWeakPtr<SLev
 	MenuBuilder.EndSection();
 
 	// Go to C++ Code
-	if( SelectionInfo.SelectionClass != NULL && FSourceCodeNavigation::IsCompilerAvailable() )
+	if( SelectionInfo.SelectionClass != NULL )
 	{
-		FString ClassHeaderPath;
-		if( FSourceCodeNavigation::FindClassHeaderPath( SelectionInfo.SelectionClass, ClassHeaderPath ) && IFileManager::Get().FileSize( *ClassHeaderPath ) != INDEX_NONE )
+		if ( FSourceCodeNavigation::IsCompilerAvailable())
 		{
-			const FString CodeFileName = FPaths::GetCleanFilename( *ClassHeaderPath );
+			FString ClassHeaderPath;
+			if( FSourceCodeNavigation::FindClassHeaderPath( SelectionInfo.SelectionClass, ClassHeaderPath ) && IFileManager::Get().FileSize( *ClassHeaderPath ) != INDEX_NONE )
+			{
+				const FString CodeFileName = FPaths::GetCleanFilename( *ClassHeaderPath );
 
-			MenuBuilder.BeginSection( "ActorCode", LOCTEXT("ActorCodeHeading", "C++") );
+				MenuBuilder.BeginSection( "ActorCode", LOCTEXT("ActorCodeHeading", "C++") );
+				{
+					MenuBuilder.AddMenuEntry( FLevelEditorCommands::Get().GoToCodeForActor,
+						NAME_None, 
+						FText::Format( LOCTEXT("GoToCodeForActor", "Open {0}"), FText::FromString( CodeFileName ) ),
+						FText::Format( LOCTEXT("GoToCodeForActor_ToolTip", "Opens the header file for this actor ({0}) in a code editing program"), FText::FromString( CodeFileName ) ) );
+				}
+				MenuBuilder.EndSection();
+			}
+		}
+
+		const FString DocumentationLink = FEditorClassUtils::GetDocumentationLink(SelectionInfo.SelectionClass);
+		if (!DocumentationLink.IsEmpty())
+		{
+			MenuBuilder.BeginSection( "ActorDocumentation", LOCTEXT("ActorDocsHeading", "Documentation") );
 			{
-				MenuBuilder.AddMenuEntry( FLevelEditorCommands::Get().GoToCodeForActor,
+				MenuBuilder.AddMenuEntry( FLevelEditorCommands::Get().GoToDocsForActor,
 					NAME_None, 
-					FText::Format( LOCTEXT("GoToCodeForActor", "Open {0}"), FText::FromString( CodeFileName ) ),
-					FText::Format( LOCTEXT("GoToCodeForActor_ToolTip", "Opens the header file for this actor ({0}) in a code editing program"), FText::FromString( CodeFileName ) ) );
+					LOCTEXT("GoToDocsForActor", "Full documentation"),
+					LOCTEXT("GoToDocsForActor_ToolTip", "Click to open documentation for this actor"));
 			}
 			MenuBuilder.EndSection();
 		}
diff --git a/Engine/Source/Editor/LevelEditor/Public/LevelEditorActions.h b/Engine/Source/Editor/LevelEditor/Public/LevelEditorActions.h
--- a/Engine/Source/Editor/LevelEditor/Public/LevelEditorActions.h
+++ b/Engine/Source/Editor/LevelEditor/Public/LevelEditorActions.h
@@ -116,6 +116,9 @@ public:
 	/** Goes to the source code for the selected actor's class. */
 	TSharedPtr< FUICommandInfo > GoToCodeForActor;
 
+	/** Goes to the documentation for the selected actor's class. */
+	TSharedPtr< FUICommandInfo > GoToDocsForActor;
+
 	/** Paste actor at click location*/
 	TSharedPtr< FUICommandInfo > PasteHere;
 
@@ -780,6 +783,8 @@ public:
 	/** Called when 'Go to Code for Actor' is clicked */
 	static void GoToCodeForActor_Clicked();
 
+	/** Called when 'Go to Documentation for Actor' is clicked */
+	static void GoToDocsForActor_Clicked();
 
 	/**
 	 * Called when the LockActorMovement command is executed
diff --git a/Engine/Source/Editor/PlacementMode/Private/SPlacementModeTools.cpp b/Engine/Source/Editor/PlacementMode/Private/SPlacementModeTools.cpp
--- a/Engine/Source/Editor/PlacementMode/Private/SPlacementModeTools.cpp
+++ b/Engine/Source/Editor/PlacementMode/Private/SPlacementModeTools.cpp
@@ -8,6 +8,7 @@
 #include "DragAndDrop/AssetDragDropOp.h"
 #include "IBspModeModule.h"
 #include "KismetEditorUtilities.h"
+#include "EditorClassUtils.h"
 
 #include "STutorialWrapper.h"
 
@@ -94,35 +95,50 @@ void SPlacementAssetEntry::Construct(const FArguments& InArgs, UActorFactory* In
 	TSharedPtr< SHorizontalBox > ActorType = SNew( SHorizontalBox );
 
 	const bool IsClass = AssetData.GetClass() == UClass::StaticClass();
-	const bool IsVolume = IsClass ? Cast<UClass>( AssetData.GetAsset() )->IsChildOf( AVolume::StaticClass() ) : false;
+	const bool IsVolume = IsClass ? CastChecked<UClass>( AssetData.GetAsset() )->IsChildOf( AVolume::StaticClass() ) : false;
 
-	AssetDisplayName = FText::FromName( AssetData.AssetName );
 	if ( IsClass )
 	{
 		AssetDisplayName = FText::FromString( FName::NameToDisplayString( AssetData.AssetName.ToString(), false ) );
 	}
-
-	FText ActorTypeDisplayName;
-	AActor* DefaultActor = nullptr;
-	if ( IsClass && Cast<UClass>( AssetData.GetAsset() )->IsChildOf( AActor::StaticClass() ) )
+	else
 	{
-		DefaultActor = Cast<AActor>( Cast<UClass>( AssetData.GetAsset() )->ClassDefaultObject );
-		ActorTypeDisplayName = FText::FromString( FName::NameToDisplayString( DefaultActor->GetClass()->GetName(), false ) );
+		AssetDisplayName = FText::FromName( AssetData.AssetName );
 	}
 
+	TSharedPtr<IToolTip> ToolTip;
+	FText ActorTypeDisplayName;
+	AActor* DefaultActor = nullptr;
 	if ( FactoryToUse != nullptr )
 	{
 		DefaultActor = FactoryToUse->GetDefaultActor( AssetData );
 		ActorTypeDisplayName = FactoryToUse->GetDisplayName();
+		ToolTip = FEditorClassUtils::GetTooltip(DefaultActor->GetClass());
 	}
-
-	AssetDisplayName = ( IsClass && !IsVolume && !ActorTypeDisplayName.IsEmpty() ) ? ActorTypeDisplayName : AssetDisplayName;
-
-	if ( !InArgs._LabelOverride.IsEmpty() )
+	else if ( IsClass && CastChecked<UClass>( AssetData.GetAsset() )->IsChildOf( AActor::StaticClass() ) )
+	{
+		DefaultActor = CastChecked<AActor>( CastChecked<UClass>( AssetData.GetAsset() )->ClassDefaultObject );
+		ActorTypeDisplayName = FText::FromString( FName::NameToDisplayString( DefaultActor->GetClass()->GetName(), false ) );
+		ToolTip = FEditorClassUtils::GetTooltip(DefaultActor->GetClass());
+	}
+	
+	if ( InArgs._LabelOverride.IsEmpty() )
+	{
+		if (IsClass && !IsVolume && !ActorTypeDisplayName.IsEmpty())
+		{
+			AssetDisplayName = ActorTypeDisplayName;
+		}
+	}
+	else
 	{
 		AssetDisplayName = InArgs._LabelOverride;
 	}
 
+	if (!ToolTip.IsValid())
+	{
+		ToolTip = FSlateApplicationBase::Get().MakeToolTip(AssetDisplayName);
+	}
+	
 	const FButtonStyle& ButtonStyle = FEditorStyle::GetWidgetStyle<FButtonStyle>( "PlacementBrowser.Asset" );
 
 	NormalImage = &ButtonStyle.Normal;
@@ -134,6 +150,7 @@ void SPlacementAssetEntry::Construct(const FArguments& InArgs, UActorFactory* In
 		SNew( SBorder )
 		.BorderImage( this, &SPlacementAssetEntry::GetBorder )
 		.Cursor( EMouseCursor::GrabHand )
+		.ToolTip( ToolTip )
 		[
 			SNew( SHorizontalBox )
 
@@ -145,7 +162,6 @@ void SPlacementAssetEntry::Construct(const FArguments& InArgs, UActorFactory* In
 				SNew( SBorder )
 				.Padding( 5 )
 				.BorderImage( FEditorStyle::GetBrush( "ContentBrowser.ThumbnailShadow" ) )
-				.ToolTipText( ActorTypeDisplayName )
 				[
 					SNew( SBox )
 					.WidthOverride( 35 )
diff --git a/Engine/Source/Editor/PropertyEditor/Private/SDetailNameArea.cpp b/Engine/Source/Editor/PropertyEditor/Private/SDetailNameArea.cpp
--- a/Engine/Source/Editor/PropertyEditor/Private/SDetailNameArea.cpp
+++ b/Engine/Source/Editor/PropertyEditor/Private/SDetailNameArea.cpp
@@ -5,7 +5,7 @@
 #include "SDetailNameArea.h"
 #include "ClassIconFinder.h"
 #include "Editor/EditorWidgets/Public/EditorWidgets.h"
-#include "SourceCodeNavigation.h"
+#include "EditorClassUtils.h"
 
 #define LOCTEXT_NAMESPACE "SDetailsView"
 
@@ -91,7 +91,7 @@ TSharedRef< SWidget > SDetailNameArea::BuildObjectNameArea( const TArray< TWeakO
 		[
 			SNew(SImage)
 			.Image(ActorIcon)
-			.ToolTipText( FText::FromString( BaseClass->GetName() ) )
+			.ToolTip(FEditorClassUtils::GetTooltip(BaseClass))
 		];
 	}
 
@@ -124,8 +124,8 @@ TSharedRef< SWidget > SDetailNameArea::BuildObjectNameArea( const TArray< TWeakO
 		if( bShowLockButton )
 		{
 			ObjectNameArea->AddSlot()
-				.AutoWidth()
 				.HAlign(HAlign_Right)
+				.FillWidth(1.0f)
 				[
 					SNew( SButton )
 					.ButtonStyle( FEditorStyle::Get(), "NoBorder" )
@@ -161,21 +161,6 @@ TSharedRef< SWidget > SDetailNameArea::BuildObjectNameArea( const TArray< TWeakO
 	return ObjectNameArea;
 }
 
-void SDetailNameArea::OnEditBlueprintClicked( TWeakObjectPtr<UBlueprint> InBlueprint, TWeakObjectPtr<UObject> InAsset )
-{
-	if (UBlueprint* Blueprint = InBlueprint.Get())
-	{
-		// Set the object being debugged if given an actor reference (if we don't do this before we edit the object the editor wont know we are debugging something)
-		if (UObject* Asset = InAsset.Get())
-		{
-			check(Asset->GetClass()->ClassGeneratedBy == Blueprint);
-			Blueprint->SetObjectBeingDebugged(Asset);
-		}
-		// Open the blueprint
-		GEditor->EditObject( Blueprint );
-	}
-}
-
 void SDetailNameArea::BuildObjectNameAreaSelectionLabel( TSharedRef< SHorizontalBox > SelectionLabelBox, const TWeakObjectPtr<UObject> ObjectWeakPtr, const int32 NumSelectedObjects ) 
 {
 	check( NumSelectedObjects > 1 || ObjectWeakPtr.IsValid() );
@@ -185,64 +170,23 @@ void SDetailNameArea::BuildObjectNameAreaSelectionLabel( TSharedRef< SHorizontal
 		UClass* ObjectClass = ObjectWeakPtr.Get()->GetClass();
 		if( ObjectClass != nullptr )
 		{
-			if (UBlueprint* Blueprint = Cast<UBlueprint>(ObjectClass->ClassGeneratedBy))
-			{
-				TWeakObjectPtr<UBlueprint> BlueprintPtr = Blueprint;
-
-				SelectionLabelBox->AddSlot()
-					.VAlign( VAlign_Center )
-					.HAlign( HAlign_Left )
-					.Padding( 6.0f, 1.0f, 0.0f, 0.0f )
-					.FillWidth( 1.0f )
-					[
-						SNew(SHyperlink)
-							.Style(FEditorStyle::Get(), "EditBPHyperlink")
-							.TextStyle(FEditorStyle::Get(), "DetailsView.EditBlueprintHyperlinkStyle")
-							.OnNavigate(this, &SDetailNameArea::OnEditBlueprintClicked, BlueprintPtr, ObjectWeakPtr)
-							.Text(FText::Format(LOCTEXT("EditBlueprint", "Edit {0}"), FText::FromString( Blueprint->GetName() ) ))
-							.ToolTipText(LOCTEXT("EditBlueprint_ToolTip", "Click to edit the blueprint"))
-					];
-			}		
-			else if( FSourceCodeNavigation::IsCompilerAvailable() )
-			{
-				FString ClassHeaderPath;
-				if( FSourceCodeNavigation::FindClassHeaderPath( ObjectClass, ClassHeaderPath ) && IFileManager::Get().FileSize( *ClassHeaderPath ) != INDEX_NONE )
-				{
-					struct Local
-					{
-						static void OnEditCodeClicked( FString InClassHeaderPath )
-						{
-							FString AbsoluteHeaderPath = IFileManager::Get().ConvertToAbsolutePathForExternalAppForRead(*InClassHeaderPath);
-							FSourceCodeNavigation::OpenSourceFile( AbsoluteHeaderPath );
-						}
-					};
+			SelectionLabelBox->AddSlot()
+				.AutoWidth()
+				.VAlign( VAlign_Center )
+				.HAlign( HAlign_Left )
+				.Padding( 1.0f, 1.0f, 0.0f, 0.0f )
+				[
+					FEditorClassUtils::GetDocumentationLinkWidget(ObjectClass)
+				];
 
-					SelectionLabelBox->AddSlot()
-						.VAlign( VAlign_Center )
-						.HAlign( HAlign_Left )
-						.Padding( 6.0f, 1.0f, 0.0f, 0.0f )
-						.FillWidth( 1.0f )
-						[
-							SNew(SHyperlink)
-								.Style(FCoreStyle::Get(), "Hyperlink")
-								.TextStyle(FEditorStyle::Get(), "DetailsView.GoToCodeHyperlinkStyle")
-								.OnNavigate_Static(&Local::OnEditCodeClicked, ClassHeaderPath)
-//								.Text(LOCTEXT("GoToCode", "(C++)" ))
-								.Text(FText::Format(LOCTEXT("GoToCode", "{0}" ), FText::FromString(FPaths::GetCleanFilename( *ClassHeaderPath ) ) ) )
-								.ToolTipText(LOCTEXT("GoToCode_ToolTip", "Click to open this source file in a text editor"))
-						];
-				}
-			}
-			else
-			{
-				// Fill the empty space
-				SelectionLabelBox->AddSlot();
-			}
-		}
-		else
-		{
-			// Fill the empty space
-			SelectionLabelBox->AddSlot();
+			SelectionLabelBox->AddSlot()
+				.AutoWidth()
+				.VAlign( VAlign_Center )
+				.HAlign( HAlign_Left )
+				.Padding( 6.0f, 1.0f, 0.0f, 0.0f )
+				[
+					FEditorClassUtils::GetSourceLink(ObjectClass, ObjectWeakPtr)
+				];
 		}
 	}
 	else
diff --git a/Engine/Source/Editor/SceneOutliner/Private/SceneOutlinerActorInfoColumn.cpp b/Engine/Source/Editor/SceneOutliner/Private/SceneOutlinerActorInfoColumn.cpp
--- a/Engine/Source/Editor/SceneOutliner/Private/SceneOutlinerActorInfoColumn.cpp
+++ b/Engine/Source/Editor/SceneOutliner/Private/SceneOutlinerActorInfoColumn.cpp
@@ -1,9 +1,10 @@
 // Copyright 1998-2014 Epic Games, Inc. All Rights Reserved.
 
 #include "SceneOutlinerPrivatePCH.h"
-#include "SourceCodeNavigation.h"
 #include "SceneOutlinerTreeItems.h"
 
+#include "EditorClassUtils.h"
+
 #define LOCTEXT_NAMESPACE "SceneOutlinerActorInfoColumn"
 
 TArray< TSharedPtr< ECustomColumnMode::Type > > FSceneOutlinerActorInfoColumn::ModeOptions;
@@ -71,60 +72,11 @@ TSharedRef< SWidget > FSceneOutlinerActorInfoColumn::ConstructClassHyperlink( co
 
 	UClass* ObjectClass = Actor->GetClass();
 
-	if ( UBlueprint* Blueprint = Cast<UBlueprint>( ObjectClass->ClassGeneratedBy ) )
-	{
-		TWeakObjectPtr<UBlueprint> BlueprintPtr = Blueprint;
+	TSharedPtr<SWidget> SourceLink = FEditorClassUtils::GetSourceLink(ObjectClass, InActor);
 
-		struct Local
-		{
-			static void OnEditBlueprintClicked( TWeakObjectPtr<UBlueprint> InBlueprint, TWeakObjectPtr<AActor> InAsset )
-			{
-				if ( UBlueprint* GeneratedByBlueprint = InBlueprint.Get() )
-				{
-					// Set the object being debugged if given an actor reference (if we don't do this before we edit the object the editor wont know we are debugging something)
-					if ( UObject* Asset = InAsset.Get() )
-					{
-						check(Asset->GetClass()->ClassGeneratedBy == GeneratedByBlueprint);
-						GeneratedByBlueprint->SetObjectBeingDebugged(Asset);
-					}
-					// Open the blueprint
-					GEditor->EditObject(GeneratedByBlueprint);
-				}
-			}
-		};
-
-		return SNew( SHyperlink )
-			.Visibility( this, &FSceneOutlinerActorInfoColumn::GetColumnDataVisibility, true )
-			.Style( FEditorStyle::Get(), "HoverOnlyHyperlink" )
-			.TextStyle( FEditorStyle::Get(), "SceneOutliner.EditBlueprintHyperlinkStyle" )
-			.OnNavigate_Static( &Local::OnEditBlueprintClicked, BlueprintPtr, InActor )
-			//.Text( FText::FromString( Blueprint->GetName() ) )
-			.Text( FText::FromString( ObjectClass->GetName() ) )
-			.ToolTipText( LOCTEXT( "EditBlueprint_ToolTip", "Click to edit the blueprint" ) );
-	}
-	else if ( FSourceCodeNavigation::IsCompilerAvailable() )
+	if (SourceLink.IsValid())
 	{
-		FString ClassHeaderPath;
-		if ( FSourceCodeNavigation::FindClassHeaderPath( ObjectClass, ClassHeaderPath ) && IFileManager::Get().FileSize( *ClassHeaderPath ) != INDEX_NONE )
-		{
-			struct Local
-			{
-				static void OnEditCodeClicked( FString InClassHeaderPath )
-				{
-					FString AbsoluteHeaderPath = IFileManager::Get().ConvertToAbsolutePathForExternalAppForRead(*InClassHeaderPath);
-					FSourceCodeNavigation::OpenSourceFile( AbsoluteHeaderPath );
-				}
-			};
-
-			return SNew( SHyperlink )
-				.Visibility( this, &FSceneOutlinerActorInfoColumn::GetColumnDataVisibility, true )
-				.Style( FEditorStyle::Get(), "HoverOnlyHyperlink" )
-				.TextStyle( FEditorStyle::Get(), "SceneOutliner.GoToCodeHyperlinkStyle" )
-				.OnNavigate_Static( &Local::OnEditCodeClicked, ClassHeaderPath )
-				//.Text( FText::FromString( FPaths::GetCleanFilename( *ClassHeaderPath ) ) )
-				.Text( FText::FromString( ObjectClass->GetName() ) )
-				.ToolTipText( LOCTEXT( "GoToCode_ToolTip", "Click to open this source file in a text editor" ) );
-		}
+		return SourceLink.ToSharedRef();
 	}
 
 	return SNew( STextBlock )
diff --git a/Engine/Source/Editor/UnrealEd/Private/EditorClassUtils.cpp b/Engine/Source/Editor/UnrealEd/Private/EditorClassUtils.cpp
new file mode 100644
--- /dev/null
+++ b/Engine/Source/Editor/UnrealEd/Private/EditorClassUtils.cpp
@@ -0,0 +1,109 @@
+// Copyright 1998-2014 Epic Games, Inc. All Rights Reserved.
+
+#include "UnrealEd.h"
+#include "EditorClassUtils.h"
+
+#include "IDocumentation.h"
+#include "SourceCodeNavigation.h"
+
+static const FString ClassDocsPage(TEXT("Shared/Classes"));
+
+TSharedRef<SToolTip> FEditorClassUtils::GetTooltip(const UClass* Class)
+{
+	return IDocumentation::Get()->CreateToolTip(Class->GetToolTipText(), nullptr, ClassDocsPage, Class->GetName());
+}
+
+FString FEditorClassUtils::GetDocumentationLink(const UClass* Class)
+{
+	FString DocumentationLink;
+
+	TSharedRef<IDocumentation> Documentation = IDocumentation::Get();
+	if (Documentation->PageExists(ClassDocsPage))
+	{
+		TSharedRef<IDocumentationPage> ClassDocs = Documentation->GetPage(ClassDocsPage, NULL);
+
+		FExcerpt Excerpt;
+		if (ClassDocs->GetExcerpt(Class->GetName(), Excerpt))
+		{
+			FString* FullDocumentationLink = Excerpt.Variables.Find( TEXT("ToolTipFullLink") );
+			if (FullDocumentationLink)
+			{
+				DocumentationLink = *FullDocumentationLink;
+			}
+		}
+	}
+
+	return DocumentationLink;
+}
+
+TSharedRef<SWidget> FEditorClassUtils::GetDocumentationLinkWidget(const UClass* Class)
+{
+	TSharedRef<SWidget> DocLinkWidget = SNullWidget::NullWidget;
+	const FString DocumentationLink = GetDocumentationLink(Class);
+
+	if (!DocumentationLink.IsEmpty())
+	{
+		DocLinkWidget = IDocumentation::Get()->CreateAnchor(DocumentationLink);
+	}
+
+	return DocLinkWidget;
+}
+
+TSharedRef<SWidget> FEditorClassUtils::GetSourceLink(const UClass* Class, const TWeakObjectPtr<UObject> ObjectWeakPtr)
+{
+	TSharedRef<SWidget> SourceHyperlink = SNullWidget::NullWidget;
+
+	if (UBlueprint* Blueprint = Cast<UBlueprint>(Class->ClassGeneratedBy))
+	{
+		struct Local
+		{
+			static void OnEditBlueprintClicked( TWeakObjectPtr<UBlueprint> InBlueprint, TWeakObjectPtr<UObject> InAsset )
+			{
+				if (UBlueprint* Blueprint = InBlueprint.Get())
+				{
+					// Set the object being debugged if given an actor reference (if we don't do this before we edit the object the editor wont know we are debugging something)
+					if (UObject* Asset = InAsset.Get())
+					{
+						check(Asset->GetClass()->ClassGeneratedBy == Blueprint);
+						Blueprint->SetObjectBeingDebugged(Asset);
+					}
+					// Open the blueprint
+					GEditor->EditObject( Blueprint );
+				}
+			}
+		};
+
+		TWeakObjectPtr<UBlueprint> BlueprintPtr = Blueprint;
+
+		SAssignNew(SourceHyperlink, SHyperlink)
+			.Style(FEditorStyle::Get(), "EditBPHyperlink")
+			.TextStyle(FEditorStyle::Get(), "DetailsView.EditBlueprintHyperlinkStyle")
+			.OnNavigate_Static(&Local::OnEditBlueprintClicked, BlueprintPtr, ObjectWeakPtr)
+			.Text(FText::Format(NSLOCTEXT("SourceHyperlink", "EditBlueprint", "Edit {0}"), FText::FromString( Blueprint->GetName() ) ))
+			.ToolTipText(NSLOCTEXT("SourceHyperlink", "EditBlueprint_ToolTip", "Click to edit the blueprint"));
+	}
+	else if( FSourceCodeNavigation::IsCompilerAvailable() )
+	{
+		FString ClassHeaderPath;
+		if( FSourceCodeNavigation::FindClassHeaderPath( Class, ClassHeaderPath ) && IFileManager::Get().FileSize( *ClassHeaderPath ) != INDEX_NONE )
+		{
+			struct Local
+			{
+				static void OnEditCodeClicked( FString InClassHeaderPath )
+				{
+					FString AbsoluteHeaderPath = IFileManager::Get().ConvertToAbsolutePathForExternalAppForRead(*InClassHeaderPath);
+					FSourceCodeNavigation::OpenSourceFile( AbsoluteHeaderPath );
+				}
+			};
+
+			SAssignNew(SourceHyperlink, SHyperlink)
+				.Style(FCoreStyle::Get(), "Hyperlink")
+				.TextStyle(FEditorStyle::Get(), "DetailsView.GoToCodeHyperlinkStyle")
+				.OnNavigate_Static(&Local::OnEditCodeClicked, ClassHeaderPath)
+				.Text(FText::Format(NSLOCTEXT("SourceHyperlink", "GoToCode", "{0}" ), FText::FromString(FPaths::GetCleanFilename( *ClassHeaderPath ) ) ) )
+				.ToolTipText(NSLOCTEXT("SourceHyperlink", "GoToCode_ToolTip", "Click to open this source file in a text editor"));
+		}
+	}
+
+	return SourceHyperlink;
+}
\ No newline at end of file
diff --git a/Engine/Source/Editor/UnrealEd/Private/SClassPickerDialog.cpp b/Engine/Source/Editor/UnrealEd/Private/SClassPickerDialog.cpp
--- a/Engine/Source/Editor/UnrealEd/Private/SClassPickerDialog.cpp
+++ b/Engine/Source/Editor/UnrealEd/Private/SClassPickerDialog.cpp
@@ -4,6 +4,7 @@
 #include "UnrealEd.h"
 #include "Editor/ClassViewer/Private/SClassViewer.h"
 #include "SClassPickerDialog.h"
+#include "EditorClassUtils.h"
 
 #define LOCTEXT_NAMESPACE "SClassPicker"
 
@@ -164,6 +165,7 @@ TSharedRef<ITableRow> SClassPickerDialog::GenerateListRow(TSharedPtr<FClassPicke
 			[
 				SNew(SButton)
 				.OnClicked(this, &SClassPickerDialog::OnDefaultClassPicked, ItemClass)
+				.ToolTip(FEditorClassUtils::GetTooltip(ItemClass))
 				.Content()
 				[
 					SNew(SHorizontalBox)
@@ -183,16 +185,21 @@ TSharedRef<ITableRow> SClassPickerDialog::GenerateListRow(TSharedPtr<FClassPicke
 						SNew(STextBlock)
 						.Text(Obj->GetName())
 					]
+
 				]
 			]
 			+SHorizontalBox::Slot()
 			.Padding(10.0f, 0.0f)
-			.MaxWidth(500.0f)
 			[
 				SNew(STextBlock)
 				.Text(FString::Printf(*NSLOCTEXT("EditorFactories", "FClassPickerDefaults", "%s").ToString(), *Obj->GetDescription()))
 				.AutoWrapText(true)
 			]
+			+SHorizontalBox::Slot()
+			.AutoWidth()
+			[
+				FEditorClassUtils::GetDocumentationLinkWidget(ItemClass)
+			]
 		]
 		+SVerticalBox::Slot()
 		.AutoHeight()
diff --git a/Engine/Source/Editor/UnrealEd/Public/EditorClassUtils.h b/Engine/Source/Editor/UnrealEd/Public/EditorClassUtils.h
new file mode 100644
--- /dev/null
+++ b/Engine/Source/Editor/UnrealEd/Public/EditorClassUtils.h
@@ -0,0 +1,43 @@
+// Copyright 1998-2014 Epic Games, Inc. All Rights Reserved.
+
+#pragma once
+
+class UClass;
+class SToolTip;
+
+namespace FEditorClassUtils
+{
+
+	/**
+	 * Gets the tooltip to display for a given class
+	 *
+	 * @param	InClass		Class we want to build a tooltip for
+	 * @return				Shared reference to the constructed tooltip
+	 */
+	UNREALED_API TSharedRef<SToolTip> GetTooltip(const UClass* Class);
+
+	/**
+	 * Returns the link path to the documentation for a given class
+	 *
+	 * @param	Class		Class we want to build a link for
+	 * @return				The path to the documentation for the class
+	 */
+	UNREALED_API FString GetDocumentationLink(const UClass* Class);
+
+	/**
+	 * Creates a link widget to the documentation for a given class
+	 *
+	 * @param	Class		Class we want to build a link for
+	 * @return				Shared pointer to the constructed tooltip
+	 */
+	UNREALED_API TSharedRef<SWidget> GetDocumentationLinkWidget(const UClass* Class);
+
+	/**
+	 * Creates a link to the source code or blueprint for a given class
+	 *
+	 * @param	Class			Class we want to build a link for
+	 * @param	ObjectWeakPtr	Optional object to set blueprint debugging to in the case we are choosing a blueprint
+	 * @return					Shared pointer to the constructed tooltip
+	 */
+	UNREALED_API TSharedRef<SWidget> GetSourceLink(const UClass* Class, const TWeakObjectPtr<UObject> ObjectWeakPtr);
+};
\ No newline at end of file
