UDS: CRASH: Editor crashes when quickly switching from a number based default value to text based ttp#340404

Hopefuly the change will properly fix also:
MAC: USER DEFINED STRUCTURE: Crash: Changing a variable type away from text or string ttp#340949

diff --git a/Engine/Source/Editor/PropertyEditor/Private/SStructureDetailsView.cpp b/Engine/Source/Editor/PropertyEditor/Private/SStructureDetailsView.cpp
--- a/Engine/Source/Editor/PropertyEditor/Private/SStructureDetailsView.cpp
+++ b/Engine/Source/Editor/PropertyEditor/Private/SStructureDetailsView.cpp
@@ -171,6 +171,7 @@ void SStructureDetailsView::SetStructureData(TSharedPtr<FStructOnScope> InStruct
 {
 	//PRE SET
 	SaveExpandedItems();
+	RootNode->SetStructure(NULL);
 	RootNodePendingKill = RootNode;
 	RootNode = MakeShareable(new FStructurePropertyNode);
 
