ttp#343033 CRASH: LIVE: Adding a variable of Struct with defined Array will crash on save

diff --git a/Engine/Source/Editor/PropertyEditor/Private/PropertyEditorModule.cpp b/Engine/Source/Editor/PropertyEditor/Private/PropertyEditorModule.cpp
--- a/Engine/Source/Editor/PropertyEditor/Private/PropertyEditorModule.cpp
+++ b/Engine/Source/Editor/PropertyEditor/Private/PropertyEditorModule.cpp
@@ -648,7 +648,12 @@ TSharedRef<class IStructureDetailsView> FPropertyEditorModule::CreateStructureDe
 		{
 			static bool CanShow( const FPropertyAndParent& PropertyAndParent )
 			{
-				return !PropertyAndParent.Property.IsA<UObjectPropertyBase>() && !PropertyAndParent.Property.IsA<UInterfaceProperty>();
+				const auto ArrayProperty = Cast<UArrayProperty>(&PropertyAndParent.Property);
+				const auto TrueProperty = ArrayProperty ? ArrayProperty->Inner : &PropertyAndParent.Property;
+				const bool bCanShow = TrueProperty 
+					&& !TrueProperty->IsA<UObjectPropertyBase>()
+					&& !TrueProperty->IsA<UInterfaceProperty>();
+				return bCanShow;
 			}
 		};
 
