PropertyEditor: Added ability to commit changes to FDateTime properties

diff --git a/Engine/Source/Editor/PropertyEditor/Private/UserInterface/PropertyEditor/SPropertyEditorDateTime.cpp b/Engine/Source/Editor/PropertyEditor/Private/UserInterface/PropertyEditor/SPropertyEditorDateTime.cpp
--- a/Engine/Source/Editor/PropertyEditor/Private/UserInterface/PropertyEditor/SPropertyEditorDateTime.cpp
+++ b/Engine/Source/Editor/PropertyEditor/Private/UserInterface/PropertyEditor/SPropertyEditorDateTime.cpp
@@ -17,6 +17,10 @@ void SPropertyEditorDateTime::Construct( const FArguments& InArgs, const TShared
 		SAssignNew( PrimaryWidget, SEditableTextBox )
 		.Text( InPropertyEditor, &FPropertyEditor::GetValueAsText )
 		.Font( InArgs._Font )
+		.SelectAllTextWhenFocused(true)
+		.ClearKeyboardFocusOnCommit(false)
+		.OnTextCommitted(this, &SPropertyEditorDateTime::HandleTextCommitted)
+		.SelectAllTextOnCommit(true)
 		.IsReadOnly(InPropertyEditor->IsEditConst())
 	];
 
@@ -27,6 +31,7 @@ void SPropertyEditorDateTime::Construct( const FArguments& InArgs, const TShared
 	}
 }
 
+
 bool SPropertyEditorDateTime::Supports( const TSharedRef< FPropertyEditor >& InPropertyEditor )
 {
 	const TSharedRef< FPropertyNode > PropertyNode = InPropertyEditor->GetPropertyNode();
@@ -43,4 +48,13 @@ bool SPropertyEditorDateTime::Supports( const TSharedRef< FPropertyEditor >& InP
 	}
 
 	return false;
-}
\ No newline at end of file
+}
+
+
+void SPropertyEditorDateTime::HandleTextCommitted( const FText& NewText, ETextCommit::Type /*CommitInfo*/ )
+{
+	const TSharedRef<FPropertyNode> PropertyNode = PropertyEditor->GetPropertyNode();
+	const TSharedRef<IPropertyHandle> PropertyHandle = PropertyEditor->GetPropertyHandle();
+
+	PropertyHandle->SetValueFromFormattedString(NewText.ToString());
+}
diff --git a/Engine/Source/Editor/PropertyEditor/Private/UserInterface/PropertyEditor/SPropertyEditorDateTime.h b/Engine/Source/Editor/PropertyEditor/Private/UserInterface/PropertyEditor/SPropertyEditorDateTime.h
--- a/Engine/Source/Editor/PropertyEditor/Private/UserInterface/PropertyEditor/SPropertyEditorDateTime.h
+++ b/Engine/Source/Editor/PropertyEditor/Private/UserInterface/PropertyEditor/SPropertyEditorDateTime.h
@@ -19,6 +19,9 @@ public:
 
 	void GetDesiredWidth( float& OutMinDesiredWidth, float& OutMaxDesiredWidth );
 
+private:
+
+	virtual void HandleTextCommitted( const FText& NewText, ETextCommit::Type CommitInfo );
 
 private:
 
