Don't show "_C" in names of BPGC. Pins and editinlinenew widget.

diff --git a/Engine/Source/Editor/GraphEditor/Private/KismetPins/SGraphPinObject.cpp b/Engine/Source/Editor/GraphEditor/Private/KismetPins/SGraphPinObject.cpp
--- a/Engine/Source/Editor/GraphEditor/Private/KismetPins/SGraphPinObject.cpp
+++ b/Engine/Source/Editor/GraphEditor/Private/KismetPins/SGraphPinObject.cpp
@@ -274,7 +274,11 @@ FText SGraphPinObject::OnGetComboTextValue() const
 	if (GraphPinObj != NULL)
 	{
 		UObject* DefaultObject = GraphPinObj->DefaultObject;
-		if (DefaultObject != NULL)
+		if (UField* Field = Cast<UField>(DefaultObject))
+		{
+			Value = Field->GetDisplayNameText();
+		}
+		else if (DefaultObject != NULL)
 		{
 			Value = FText::FromString(DefaultObject->GetName());
 		}
diff --git a/Engine/Source/Editor/PropertyEditor/Private/UserInterface/PropertyEditor/SPropertyEditorEditInline.cpp b/Engine/Source/Editor/PropertyEditor/Private/UserInterface/PropertyEditor/SPropertyEditorEditInline.cpp
--- a/Engine/Source/Editor/PropertyEditor/Private/UserInterface/PropertyEditor/SPropertyEditorEditInline.cpp
+++ b/Engine/Source/Editor/PropertyEditor/Private/UserInterface/PropertyEditor/SPropertyEditorEditInline.cpp
@@ -66,7 +66,6 @@ public:
 void SPropertyEditorEditInline::Construct( const FArguments& InArgs, const TSharedRef< class FPropertyEditor >& InPropertyEditor )
 {
 	PropertyEditor = InPropertyEditor;
-	bUseDisplayNames = PropertyEditor->GetProperty()->GetBoolMetaData(TEXT("UseDisplayNames"));
 
 	ChildSlot
 	[
@@ -97,24 +96,17 @@ void SPropertyEditorEditInline::Construct( const FArguments& InArgs, const TShar
 	];
 }
 
-FString SPropertyEditorEditInline::GetDisplayValueAsString() const
+FText SPropertyEditorEditInline::GetDisplayValueAsString() const
 {
 	UObject* CurrentValue = NULL;
 	FPropertyAccess::Result Result = PropertyEditor->GetPropertyHandle()->GetValue( CurrentValue );
 	if( Result == FPropertyAccess::Success && CurrentValue != NULL )
 	{
-		if(bUseDisplayNames)
-		{
-			return CurrentValue->GetClass()->GetDisplayNameText().ToString();
-		}
-		else
-		{
-			return CurrentValue->GetName();
-		}
+		return CurrentValue->GetClass()->GetDisplayNameText();
 	}
 	else
 	{
-		return PropertyEditor->GetValueAsString();
+		return PropertyEditor->GetValueAsText();
 	}
 }
 
@@ -182,7 +174,7 @@ TSharedRef<SWidget> SPropertyEditorEditInline::GenerateClassPicker()
 {
 	FClassViewerInitializationOptions Options;
 	Options.bShowUnloadedBlueprints = true;
-	Options.bShowDisplayNames = bUseDisplayNames;
+	Options.bShowDisplayNames = true;
 
 	TSharedPtr<FPropertyEditorInlineClassFilter> ClassFilter = MakeShareable( new FPropertyEditorInlineClassFilter );
 	Options.ClassFilter = ClassFilter;
diff --git a/Engine/Source/Editor/PropertyEditor/Private/UserInterface/PropertyEditor/SPropertyEditorEditInline.h b/Engine/Source/Editor/PropertyEditor/Private/UserInterface/PropertyEditor/SPropertyEditorEditInline.h
--- a/Engine/Source/Editor/PropertyEditor/Private/UserInterface/PropertyEditor/SPropertyEditorEditInline.h
+++ b/Engine/Source/Editor/PropertyEditor/Private/UserInterface/PropertyEditor/SPropertyEditorEditInline.h
@@ -25,7 +25,7 @@ private:
 	/**
 	 * @return The current display value for the combo box as a string
 	 */
-	FString GetDisplayValueAsString() const;
+	FText GetDisplayValueAsString() const;
 
 	/**
 	 * @return The current display value's icon, if any. Returns nullptr if we have no valid value.
@@ -63,7 +63,4 @@ private:
 	TSharedPtr<class FPropertyEditor > PropertyEditor;
 
 	TSharedPtr<class SComboButton> ComboButton;
-
-	/** Whether to use display names or literal class names */
-	bool bUseDisplayNames;
 };
diff --git a/Engine/Source/Runtime/AIModule/Classes/BehaviorTree/BlackboardData.h b/Engine/Source/Runtime/AIModule/Classes/BehaviorTree/BlackboardData.h
--- a/Engine/Source/Runtime/AIModule/Classes/BehaviorTree/BlackboardData.h
+++ b/Engine/Source/Runtime/AIModule/Classes/BehaviorTree/BlackboardData.h
@@ -16,7 +16,7 @@ struct FBlackboardEntry
 	FName EntryName;
 
 	/** key type and additional properties */
-	UPROPERTY(EditAnywhere, EditInline, Category=Blackboard, meta=(UseDisplayNames="true"))
+	UPROPERTY(EditAnywhere, EditInline, Category=Blackboard)
 	UBlackboardKeyType* KeyType;
 };
 
