Property Editor - Exposing a utility function to get the Edit Condition UProperty for another property if it exists.

diff --git a/Engine/Source/Editor/PropertyEditor/Private/Presentation/PropertyEditor/PropertyEditor.cpp b/Engine/Source/Editor/PropertyEditor/Private/Presentation/PropertyEditor/PropertyEditor.cpp
--- a/Engine/Source/Editor/PropertyEditor/Private/Presentation/PropertyEditor/PropertyEditor.cpp
+++ b/Engine/Source/Editor/PropertyEditor/Private/Presentation/PropertyEditor/PropertyEditor.cpp
@@ -509,7 +509,7 @@ bool FPropertyEditor::GetEditConditionPropertyAddress( UBoolProperty*& Condition
 {
 	bool bResult = false;
 	bool bNegate = false;
-	UBoolProperty* EditConditionProperty = GetEditConditionProperty( InPropertyNode.GetProperty(), bNegate);
+	UBoolProperty* EditConditionProperty = PropertyCustomizationHelpers::GetEditConditionProperty(InPropertyNode.GetProperty(), bNegate);
 	if ( EditConditionProperty != NULL )
 	{
 		FPropertyNode* ParentNode = InPropertyNode.GetParentNode();
@@ -569,7 +569,7 @@ bool FPropertyEditor::SupportsEditConditionToggle( UProperty* InProperty )
 	if (!InProperty->HasMetaData(TEXT("HideEditConditionToggle")))
 	{
 		bool bNegateValue = false;
-		UBoolProperty* ConditionalProperty = GetEditConditionProperty( InProperty, bNegateValue );
+		UBoolProperty* ConditionalProperty = PropertyCustomizationHelpers::GetEditConditionProperty( InProperty, bNegateValue );
 		if( ConditionalProperty != NULL )
 		{
 			bShowEditConditionToggle = true;
@@ -586,35 +586,6 @@ bool FPropertyEditor::SupportsEditConditionToggle( UProperty* InProperty )
 	return bShowEditConditionToggle;
 }
 
-UBoolProperty* FPropertyEditor::GetEditConditionProperty( const UProperty* InProperty, bool& bNegate ) 
-{
-	UBoolProperty* EditConditionProperty = NULL;
-	bNegate = false;
-
-	if ( InProperty != NULL )
-	{
-		// find the name of the property that should be used to determine whether this property should be editable
-		FString ConditionPropertyName = InProperty->GetMetaData(TEXT("EditCondition"));
-
-		// Support negated edit conditions whose syntax is !BoolProperty
-		if (ConditionPropertyName.StartsWith(FString(TEXT("!"))))
-		{
-			bNegate = true;
-			// Chop off the negation from the property name
-			ConditionPropertyName = ConditionPropertyName.Right(ConditionPropertyName.Len() - 1);
-		}
-
-		// for now, only support boolean conditions, and only allow use of another property within the same struct as the conditional property
-		if ( ConditionPropertyName.Len() > 0 && !ConditionPropertyName.Contains(TEXT(".")) )
-		{
-			UStruct* Scope = InProperty->GetOwnerStruct();
-			EditConditionProperty = FindField<UBoolProperty>(Scope, *ConditionPropertyName);
-		}
-	}
-
-	return EditConditionProperty;
-}
-
 void FPropertyEditor::SyncToObjectsInNode( const TWeakPtr< FPropertyNode >& WeakPropertyNode )
 {
 #if WITH_EDITOR
diff --git a/Engine/Source/Editor/PropertyEditor/Private/Presentation/PropertyEditor/PropertyEditor.h b/Engine/Source/Editor/PropertyEditor/Private/Presentation/PropertyEditor/PropertyEditor.h
--- a/Engine/Source/Editor/PropertyEditor/Private/Presentation/PropertyEditor/PropertyEditor.h
+++ b/Engine/Source/Editor/PropertyEditor/Private/Presentation/PropertyEditor/PropertyEditor.h
@@ -141,8 +141,6 @@ private:
 
 	static bool SupportsEditConditionToggle( UProperty* InProperty );
 
-	static UBoolProperty* GetEditConditionProperty( const UProperty* InProperty, bool& bNegate ) ;
-
 private:
 
 	TArray< FPropertyConditionInfo > PropertyEditConditions;
diff --git a/Engine/Source/Editor/PropertyEditor/Private/PropertyCustomizationHelpers.cpp b/Engine/Source/Editor/PropertyEditor/Private/PropertyCustomizationHelpers.cpp
--- a/Engine/Source/Editor/PropertyEditor/Private/PropertyCustomizationHelpers.cpp
+++ b/Engine/Source/Editor/PropertyEditor/Private/PropertyCustomizationHelpers.cpp
@@ -244,6 +244,35 @@ namespace PropertyCustomizationHelpers
 			.OnShouldFilterActor( OnShouldFilterActor )
 			.OnActorSelected( OnActorSelectedFromPicker );
 	}
+
+	UBoolProperty* GetEditConditionProperty(const UProperty* InProperty, bool& bNegate)
+	{
+		UBoolProperty* EditConditionProperty = NULL;
+		bNegate = false;
+
+		if ( InProperty != NULL )
+		{
+			// find the name of the property that should be used to determine whether this property should be editable
+			FString ConditionPropertyName = InProperty->GetMetaData(TEXT("EditCondition"));
+
+			// Support negated edit conditions whose syntax is !BoolProperty
+			if ( ConditionPropertyName.StartsWith(FString(TEXT("!"))) )
+			{
+				bNegate = true;
+				// Chop off the negation from the property name
+				ConditionPropertyName = ConditionPropertyName.Right(ConditionPropertyName.Len() - 1);
+			}
+
+			// for now, only support boolean conditions, and only allow use of another property within the same struct as the conditional property
+			if ( ConditionPropertyName.Len() > 0 && !ConditionPropertyName.Contains(TEXT(".")) )
+			{
+				UStruct* Scope = InProperty->GetOwnerStruct();
+				EditConditionProperty = FindField<UBoolProperty>(Scope, *ConditionPropertyName);
+			}
+		}
+
+		return EditConditionProperty;
+	}
 }
 
 void SObjectPropertyEntryBox::Construct( const FArguments& InArgs )
diff --git a/Engine/Source/Editor/PropertyEditor/Public/PropertyCustomizationHelpers.h b/Engine/Source/Editor/PropertyEditor/Public/PropertyCustomizationHelpers.h
--- a/Engine/Source/Editor/PropertyEditor/Public/PropertyCustomizationHelpers.h
+++ b/Engine/Source/Editor/PropertyEditor/Public/PropertyCustomizationHelpers.h
@@ -3,9 +3,13 @@
 #pragma once
 
 #include "IFilter.h"
+#include "IDetailCustomNodeBuilder.h"
+
 #include "FilterCollection.h"
 #include "SResetToDefaultMenu.h"
 #include "ActorPickerMode.h"
+#include "PropertyHandle.h"
+#include "DetailWidgetRow.h"
 
 namespace SceneOutliner { struct FOutlinerFilters; }
 
@@ -29,6 +33,9 @@ namespace PropertyCustomizationHelpers
 	PROPERTYEDITOR_API TSharedRef<SWidget> MakeActorPickerAnchorButton( FOnGetActorFilters OnGetActorFilters, FOnActorSelected OnActorSelectedFromPicker );
 	PROPERTYEDITOR_API TSharedRef<SWidget> MakeActorPickerWithMenu( AActor* const InitialActor, const bool AllowClear, const TSharedPtr< SceneOutliner::FOutlinerFilters >& ActorFilters, FOnActorSelected OnSet, FSimpleDelegate OnClose, FSimpleDelegate OnUseSelected );
 	PROPERTYEDITOR_API TSharedRef<SWidget> MakeInteractiveActorPicker( FOnGetAllowedClasses OnGetAllowedClasses, FOnShouldFilterActor OnShouldFilterActor, FOnActorSelected OnActorSelectedFromPicker );
+
+	/** @return the UBoolProperty edit condition property if one exists. */
+	PROPERTYEDITOR_API class UBoolProperty* GetEditConditionProperty(const class UProperty* InProperty, bool& bNegate);
 }
 
 /** Delegate used to get a generic object */
