#ue4 build fix

diff --git a/Engine/Source/Editor/PropertyEditor/Private/PropertyNode.h b/Engine/Source/Editor/PropertyEditor/Private/PropertyNode.h
--- a/Engine/Source/Editor/PropertyEditor/Private/PropertyNode.h
+++ b/Engine/Source/Editor/PropertyEditor/Private/PropertyNode.h
@@ -653,7 +653,7 @@ public:
 	bool IsHighlighted() const { return bIsHighlighted; }
 
 	/** Sets the highlighted flag to the requested value */
-	bool SetIsHighlighted( bool bInIsHighlighted ) { bIsHighlighted = bInIsHighlighted; }
+	void SetIsHighlighted( bool bInIsHighlighted ) { bIsHighlighted = bInIsHighlighted; }
 
 protected:
 
