Fix crash accessing property text

diff --git a/Engine/Source/Editor/PropertyEditor/Private/PropertyHandleImpl.cpp b/Engine/Source/Editor/PropertyEditor/Private/PropertyHandleImpl.cpp
--- a/Engine/Source/Editor/PropertyEditor/Private/PropertyHandleImpl.cpp
+++ b/Engine/Source/Editor/PropertyEditor/Private/PropertyHandleImpl.cpp
@@ -100,7 +100,7 @@ FPropertyAccess::Result FPropertyValueImpl::GetPropertyValueText( FText& OutText
 	FReadAddressList ReadAddresses;
 	bool bAllValuesTheSame = InPropertyNode->GetReadAddress( !!InPropertyNode->HasNodeFlags(EPropertyNodeFlags::SingleSelectOnly), ReadAddresses, false, true );
 
-	if( ReadAddresses.Num() > 0 && bAllValuesTheSame || ReadAddresses.Num() == 1 ) 
+	if( ReadAddresses.Num() > 0 && InPropertyNode->GetProperty() != nullptr && (bAllValuesTheSame || ReadAddresses.Num() == 1) ) 
 	{
 		ValueAddress = ReadAddresses.GetAddress(0);
 
