Added instanced property type customization

diff --git a/Engine/Source/Editor/PropertyEditor/Private/CustomChildBuilder.cpp b/Engine/Source/Editor/PropertyEditor/Private/CustomChildBuilder.cpp
--- a/Engine/Source/Editor/PropertyEditor/Private/CustomChildBuilder.cpp
+++ b/Engine/Source/Editor/PropertyEditor/Private/CustomChildBuilder.cpp
@@ -91,7 +91,12 @@ TSharedRef<SWidget> FCustomChildrenBuilder::GenerateStructValueWidget( TSharedRe
 	UStructProperty* StructProperty = CastChecked<UStructProperty>( StructPropertyHandle->GetProperty() );
 
 	FPropertyEditorModule& PropertyEditorModule = FModuleManager::GetModuleChecked<FPropertyEditorModule>("PropertyEditor");
-	FPropertyTypeLayoutCallback LayoutCallback = PropertyEditorModule.GetPropertyTypeCustomization(StructProperty, *StructPropertyHandle);
+	
+	IDetailsViewPrivate& DetailsView = ParentCategory.Pin()->GetDetailsView();
+	
+	TSharedRef<IDetailsView> DetailsViewPtr = StaticCastSharedRef<IDetailsView>( DetailsView.AsShared() );
+	
+	FPropertyTypeLayoutCallback LayoutCallback = PropertyEditorModule.GetPropertyTypeCustomization(StructProperty, *StructPropertyHandle, DetailsViewPtr );
 	if (LayoutCallback.IsValid())
 	{
 		TSharedRef<IPropertyTypeCustomization> CustomStructInterface = LayoutCallback.GetCustomizationInstance();
diff --git a/Engine/Source/Editor/PropertyEditor/Private/DetailPropertyRow.cpp b/Engine/Source/Editor/PropertyEditor/Private/DetailPropertyRow.cpp
--- a/Engine/Source/Editor/PropertyEditor/Private/DetailPropertyRow.cpp
+++ b/Engine/Source/Editor/PropertyEditor/Private/DetailPropertyRow.cpp
@@ -35,7 +35,12 @@ FDetailPropertyRow::FDetailPropertyRow(TSharedPtr<FPropertyNode> InPropertyNode,
 		{
 			static FName PropertyEditor("PropertyEditor");
 			FPropertyEditorModule& PropertyEditorModule = FModuleManager::GetModuleChecked<FPropertyEditorModule>(PropertyEditor);
-			FPropertyTypeLayoutCallback LayoutCallback = PropertyEditorModule.GetPropertyTypeCustomization(Property,*PropertyHandle);
+			
+			IDetailsViewPrivate& DetailsView = InParentCategory->GetDetailsView();
+			
+			TSharedRef<IDetailsView> DetailsViewPtr = StaticCastSharedRef<IDetailsView>( DetailsView.AsShared() );
+			
+			FPropertyTypeLayoutCallback LayoutCallback = PropertyEditorModule.GetPropertyTypeCustomization(Property,*PropertyHandle, DetailsViewPtr );
 			if (LayoutCallback.IsValid())
 			{
 				if (PropertyHandle->IsValidHandle())
diff --git a/Engine/Source/Editor/PropertyEditor/Private/PropertyEditorModule.cpp b/Engine/Source/Editor/PropertyEditor/Private/PropertyEditorModule.cpp
--- a/Engine/Source/Editor/PropertyEditor/Private/PropertyEditorModule.cpp
+++ b/Engine/Source/Editor/PropertyEditor/Private/PropertyEditorModule.cpp
@@ -29,7 +29,7 @@ IMPLEMENT_MODULE( FPropertyEditorModule, PropertyEditor );
 
 TSharedRef<IPropertyTypeCustomization> FPropertyTypeLayoutCallback::GetCustomizationInstance() const
 {
-	return PropertyTypeLayoutDelegate.IsBound() ? PropertyTypeLayoutDelegate.Execute() : StaticCastSharedRef<IPropertyTypeCustomization>(DeprecatedLayoutDelegate.Execute());
+	return PropertyTypeLayoutDelegate.Execute();
 }
 
 void FPropertyTypeLayoutCallbackList::Add( const FPropertyTypeLayoutCallback& NewCallback )
@@ -385,10 +385,6 @@ void FPropertyEditorModule::RegisterCustomClassLayout( FName ClassName, FOnGetDe
 	}
 }
 
-void FPropertyEditorModule::RegisterCustomPropertyLayout( FName ClassName, FOnGetDetailCustomizationInstance DetailLayoutDelegate )
-{
-	RegisterCustomClassLayout( ClassName, DetailLayoutDelegate );
-}
 
 void FPropertyEditorModule::UnregisterCustomClassLayout( FName ClassName )
 {
@@ -398,69 +394,74 @@ void FPropertyEditorModule::UnregisterCustomClassLayout( FName ClassName )
 	}
 }
 
-void FPropertyEditorModule::UnregisterCustomPropertyLayout( FName ClassName )
-{
-	UnregisterCustomClassLayout(ClassName);
-}
 
-void FPropertyEditorModule::RegisterStructPropertyLayout( FName StructTypeName, FOnGetStructCustomizationInstance StructLayoutDelegate )
+void FPropertyEditorModule::RegisterCustomPropertyTypeLayout( FName PropertyTypeName, FOnGetPropertyTypeCustomizationInstance PropertyTypeLayoutDelegate, TSharedPtr<IPropertyTypeIdentifier> Identifier, TSharedPtr<IDetailsView> ForSpecificInstance )
 {
-	if (StructTypeName != NAME_None)
+	if( PropertyTypeName != NAME_None )
 	{
 		FPropertyTypeLayoutCallback Callback;
-		Callback.DeprecatedLayoutDelegate = StructLayoutDelegate;
+		Callback.PropertyTypeLayoutDelegate = PropertyTypeLayoutDelegate;
+		Callback.PropertyTypeIdentifier = Identifier;
 
-		FPropertyTypeLayoutCallbackList* LayoutCallbacks = PropertyTypeToLayoutMap.Find(StructTypeName);
-		if (LayoutCallbacks)
+		if( ForSpecificInstance.IsValid() )
 		{
-			LayoutCallbacks->Add(Callback);
+			FCustomPropertyTypeLayoutMap& PropertyTypeToLayoutMap = InstancePropertyTypeLayoutMap.FindOrAdd( ForSpecificInstance );
+			
+			FPropertyTypeLayoutCallbackList* LayoutCallbacks = PropertyTypeToLayoutMap.Find( PropertyTypeName );
+			if( LayoutCallbacks )
+			{
+				LayoutCallbacks->Add( Callback );
+			}
+			else
+			{
+				FPropertyTypeLayoutCallbackList NewLayoutCallbacks;
+				NewLayoutCallbacks.Add( Callback );
+				PropertyTypeToLayoutMap.Add( PropertyTypeName, NewLayoutCallbacks );
+			}
 		}
 		else
 		{
-			FPropertyTypeLayoutCallbackList NewLayoutCallbacks;
-			NewLayoutCallbacks.Add(Callback);
-			PropertyTypeToLayoutMap.Add(StructTypeName, NewLayoutCallbacks);
+			FPropertyTypeLayoutCallbackList* LayoutCallbacks = GlobalPropertyTypeToLayoutMap.Find( PropertyTypeName );
+			if( LayoutCallbacks )
+			{
+				LayoutCallbacks->Add( Callback );
+			}
+			else
+			{
+				FPropertyTypeLayoutCallbackList NewLayoutCallbacks;
+				NewLayoutCallbacks.Add( Callback );
+				GlobalPropertyTypeToLayoutMap.Add( PropertyTypeName, NewLayoutCallbacks );
+			}
 		}
 	}
 }
 
-void FPropertyEditorModule::UnregisterStructPropertyLayout( FName StructTypeName )
-{
-	UnregisterCustomPropertyTypeLayout( StructTypeName );
-}
-
-
-void FPropertyEditorModule::RegisterCustomPropertyTypeLayout( FName PropertyTypeName, FOnGetPropertyTypeCustomizationInstance PropertyTypeLayoutDelegate, TSharedPtr<IPropertyTypeIdentifier> Identifier )
+void FPropertyEditorModule::UnregisterCustomPropertyTypeLayout( FName PropertyTypeName, TSharedPtr<IPropertyTypeIdentifier> Identifier, TSharedPtr<IDetailsView> ForSpecificInstance )
 {
 	if( PropertyTypeName != NAME_None )
 	{
-		FPropertyTypeLayoutCallback Callback;
-		Callback.PropertyTypeLayoutDelegate = PropertyTypeLayoutDelegate;
-		Callback.PropertyTypeIdentifier = Identifier;
-
-		FPropertyTypeLayoutCallbackList* LayoutCallbacks = PropertyTypeToLayoutMap.Find( PropertyTypeName );
-		if( LayoutCallbacks )
+		if( ForSpecificInstance.IsValid() )
 		{
-			LayoutCallbacks->Add( Callback );
+			FCustomPropertyTypeLayoutMap* PropertyTypeToLayoutMap = InstancePropertyTypeLayoutMap.Find( ForSpecificInstance );
+			
+			if( PropertyTypeToLayoutMap )
+			{
+				FPropertyTypeLayoutCallbackList* LayoutCallbacks = PropertyTypeToLayoutMap->Find( PropertyTypeName );
+			
+				if( LayoutCallbacks )
+				{
+					LayoutCallbacks->Remove( Identifier );
+				}
+			}
 		}
 		else
 		{
-			FPropertyTypeLayoutCallbackList NewLayoutCallbacks;
-			NewLayoutCallbacks.Add( Callback );
-			PropertyTypeToLayoutMap.Add( PropertyTypeName, NewLayoutCallbacks );
-		}
-	}
-}
+			FPropertyTypeLayoutCallbackList* LayoutCallbacks = GlobalPropertyTypeToLayoutMap.Find( PropertyTypeName );
 
-void FPropertyEditorModule::UnregisterCustomPropertyTypeLayout( FName PropertyTypeName, TSharedPtr<IPropertyTypeIdentifier> Identifier )
-{
-	if( PropertyTypeName != NAME_None )
-	{
-		FPropertyTypeLayoutCallbackList* LayoutCallbacks = PropertyTypeToLayoutMap.Find( PropertyTypeName );
-
-		if( LayoutCallbacks )
-		{
-			LayoutCallbacks->Remove( Identifier );
+			if( LayoutCallbacks )
+			{
+				LayoutCallbacks->Remove( Identifier );
+			}
 		}
 	}
 }
@@ -580,16 +581,28 @@ void FPropertyEditorModule::RemoveDeletedObjects( TArray<UObject*>& DeletedObjec
 	}
 }
 
-bool FPropertyEditorModule::IsCustomizedStruct(const UStruct* Struct) const
+bool FPropertyEditorModule::IsCustomizedStruct(const UStruct* Struct, TSharedPtr<IDetailsView> DetailsViewInstance ) const
 {
+	bool bFound = false;
 	if (Struct && !Struct->IsA<UUserDefinedStruct>())
 	{
-		return PropertyTypeToLayoutMap.Contains(Struct->GetFName());
+		const FCustomPropertyTypeLayoutMap* Map = InstancePropertyTypeLayoutMap.Find( DetailsViewInstance );
+
+		if( Map )
+		{
+			bFound = Map->Contains( Struct->GetFName() );
+		}
+		
+		if( !bFound )
+		{
+			bFound = GlobalPropertyTypeToLayoutMap.Contains( Struct->GetFName() );
+		}
 	}
-	return false;
+	
+	return bFound;
 }
 
-FPropertyTypeLayoutCallback FPropertyEditorModule::GetPropertyTypeCustomization(const UProperty* Property, const IPropertyHandle& PropertyHandle)
+FPropertyTypeLayoutCallback FPropertyEditorModule::GetPropertyTypeCustomization(const UProperty* Property, const IPropertyHandle& PropertyHandle, TSharedPtr<IDetailsView> DetailsViewInstance )
 {
 	if( Property )
 	{
@@ -625,11 +638,22 @@ FPropertyTypeLayoutCallback FPropertyEditorModule::GetPropertyTypeCustomization(
 
 		if (PropertyTypeName != NAME_None)
 		{
-			FPropertyTypeLayoutCallbackList* LayoutCallbacks = PropertyTypeToLayoutMap.Find(PropertyTypeName);
-			if (LayoutCallbacks)
+			FCustomPropertyTypeLayoutMap* Map = InstancePropertyTypeLayoutMap.Find( DetailsViewInstance );
+			
+			FPropertyTypeLayoutCallbackList* LayoutCallbacks = nullptr;
+			if( Map )
+			{
+				LayoutCallbacks = Map->Find( PropertyTypeName );
+			}
+				
+			if( !LayoutCallbacks )
 			{
-				const FPropertyTypeLayoutCallback& Callback = LayoutCallbacks->Find(PropertyHandle);
-				return Callback;
+				FPropertyTypeLayoutCallbackList* LayoutCallbacks = GlobalPropertyTypeToLayoutMap.Find(PropertyTypeName);
+				if (LayoutCallbacks)
+				{
+					const FPropertyTypeLayoutCallback& Callback = LayoutCallbacks->Find(PropertyHandle);
+					return Callback;
+				}
 			}
 		}
 	}
diff --git a/Engine/Source/Editor/PropertyEditor/Private/SDetailsViewBase.cpp b/Engine/Source/Editor/PropertyEditor/Private/SDetailsViewBase.cpp
--- a/Engine/Source/Editor/PropertyEditor/Private/SDetailsViewBase.cpp
+++ b/Engine/Source/Editor/PropertyEditor/Private/SDetailsViewBase.cpp
@@ -944,12 +944,12 @@ void SDetailsViewBase::UpdatePropertyMapRecursive(FPropertyNode& InNode, FDetail
 					FPropertyEditorModule& ParentPlugin = FModuleManager::GetModuleChecked<FPropertyEditorModule>("PropertyEditor");
 					if (Struct)
 					{
-						bIsCustomizedStruct = ParentPlugin.IsCustomizedStruct(Struct);
+						bIsCustomizedStruct = ParentPlugin.IsCustomizedStruct(Struct, SharedThis( this ) );
 					}
 
 					if (ParentStruct)
 					{
-						bIsChildOfCustomizedStruct = ParentPlugin.IsCustomizedStruct(ParentStruct);
+						bIsChildOfCustomizedStruct = ParentPlugin.IsCustomizedStruct(ParentStruct, SharedThis( this ) );
 					}
 				}
 
diff --git a/Engine/Source/Editor/PropertyEditor/Public/PropertyEditorDelegates.h b/Engine/Source/Editor/PropertyEditor/Public/PropertyEditorDelegates.h
--- a/Engine/Source/Editor/PropertyEditor/Public/PropertyEditorDelegates.h
+++ b/Engine/Source/Editor/PropertyEditor/Public/PropertyEditorDelegates.h
@@ -54,6 +54,3 @@ DECLARE_DELEGATE_RetVal(bool, FIsPropertyEditingEnabled );
  */
 DECLARE_MULTICAST_DELEGATE_OneParam(FOnFinishedChangingProperties, const FPropertyChangedEvent&);
 
-/* Deprecated customization delegate support*/
-DECLARE_DELEGATE_RetVal( TSharedRef<class IStructCustomization>, FOnGetStructCustomizationInstance );
-
diff --git a/Engine/Source/Editor/PropertyEditor/Public/PropertyEditorModule.h b/Engine/Source/Editor/PropertyEditor/Public/PropertyEditorModule.h
--- a/Engine/Source/Editor/PropertyEditor/Public/PropertyEditorModule.h
+++ b/Engine/Source/Editor/PropertyEditor/Public/PropertyEditorModule.h
@@ -51,6 +51,7 @@ class SPropertyTreeViewImpl;
 class SWindow;
 class IPropertyTableCellPresenter;
 class IPropertyTypeCustomization;
+class IDetailsView;
 
 /**
  * Base class for adding an extra data to identify a custom property type
@@ -83,11 +84,9 @@ struct FPropertyTypeLayoutCallback
 {
 	FOnGetPropertyTypeCustomizationInstance PropertyTypeLayoutDelegate;
 
-	FOnGetStructCustomizationInstance DeprecatedLayoutDelegate;
-
 	TSharedPtr<IPropertyTypeIdentifier> PropertyTypeIdentifier;
 
-	bool IsValid() { return PropertyTypeLayoutDelegate.IsBound() || DeprecatedLayoutDelegate.IsBound(); }
+	bool IsValid() { return PropertyTypeLayoutDelegate.IsBound(); }
 
 	TSharedRef<IPropertyTypeCustomization> GetCustomizationInstance() const;
 };
@@ -114,6 +113,7 @@ typedef TMap< FName, FDetailLayoutCallback > FCustomDetailLayoutNameMap;
 /** This is a multimap as there many be more than one customization per property type */
 typedef TMap< FName, FPropertyTypeLayoutCallbackList > FCustomPropertyTypeLayoutMap;
 
+
 class FPropertyEditorModule : public IModuleInterface
 {
 	friend class SPropertyTreeView;
@@ -166,9 +166,6 @@ public:
 	 */
 	virtual void RegisterCustomClassLayout( FName ClassName, FOnGetDetailCustomizationInstance DetailLayoutDelegate );
 
-	DEPRECATED(4.3, "RegisterCustomPropertyLayout is deprecated use RegisterCustomClassLayout instead.")
-	virtual void RegisterCustomPropertyLayout( FName ClassName, FOnGetDetailCustomizationInstance DetailLayoutDelegate );
-
 	/**
 	 * Unregisters a custom detail layout delegate for a specific class name
 	 *
@@ -176,9 +173,6 @@ public:
 	 */
 	virtual void UnregisterCustomClassLayout( FName ClassName );
 
-	DEPRECATED(4.3, "UnregisterCustomPropertyLayout is deprecated use UnregisterCustomClassLayout instead.")
-	virtual void UnregisterCustomPropertyLayout( FName ClassName );
-
 	/**
 	 * Registers a property type customization
 	 * A property type is a specific UProperty type, a struct, or enum type
@@ -187,7 +181,7 @@ public:
 	 * @param PropertyTypeLayoutDelegate	The delegate to call when querying for a custom layout of the property type
 	 * @param Identifier			An identifier to use to differentiate between two customizations on the same type
 	 */
-	virtual void RegisterCustomPropertyTypeLayout( FName PropertyTypeName, FOnGetPropertyTypeCustomizationInstance PropertyTypeLayoutDelegate, TSharedPtr<IPropertyTypeIdentifier> Identifier = NULL );
+	virtual void RegisterCustomPropertyTypeLayout( FName PropertyTypeName, FOnGetPropertyTypeCustomizationInstance PropertyTypeLayoutDelegate, TSharedPtr<IPropertyTypeIdentifier> Identifier = nullptr, TSharedPtr<IDetailsView> ForSpecificInstance = nullptr );
 
 	/**
 	 * Unregisters a custom detail layout for a properrty type
@@ -195,24 +189,7 @@ public:
 	 * @param PropertyTypeName 	The name of the property type that was registered
 	 * @param Identifier 		An identifier to use to differentiate between two customizations on the same type
 	 */
-	virtual void UnregisterCustomPropertyTypeLayout( FName PropertyTypeName, TSharedPtr<IPropertyTypeIdentifier> InIdentifier = NULL );
-
-	/**
-	 * Registers a custom detail layout delegate for a specific class
-	 *
-	 * @param Class	The class the custom detail layout is for
-	 * @param StructLayoutDelegate	The delegate to call when querying for custom detail layouts for the struct properties
-	 */
-	DEPRECATED(4.3, "RegisterStructPropertyLayout is deprecated.  Use RegisterCustomPropertyTypeLayout instead.")
-	virtual void RegisterStructPropertyLayout( FName StructTypeName, FOnGetStructCustomizationInstance StructLayoutDelegate );
-
-	/**
-	 * Unregisters a custom detail layout delegate for a specific structure
-	 *
-	 * @param StructTypeName	The structure with the custom detail layout delegate to remove
-	 */
-	DEPRECATED(4.3, "UnregisterStructPropertyLayout is deprecated.  Use UnregisterCustomPropertyTypeLayout instead.")
-	virtual void UnregisterStructPropertyLayout( FName StructTypeName );
+	virtual void UnregisterCustomPropertyTypeLayout( FName PropertyTypeName, TSharedPtr<IPropertyTypeIdentifier> InIdentifier = nullptr, TSharedPtr<IDetailsView> ForSpecificInstance = nullptr );
 
 	/**
 	 * Customization modules should call this when that module has been unloaded, loaded, etc...
@@ -282,8 +259,8 @@ public:
 	virtual TSharedRef< FAssetEditorToolkit > CreatePropertyEditorToolkit( const EToolkitMode::Type Mode, const TSharedPtr< IToolkitHost >& InitToolkitHost, const TArray< UObject* >& ObjectsToEdit );
 	virtual TSharedRef< FAssetEditorToolkit > CreatePropertyEditorToolkit( const EToolkitMode::Type Mode, const TSharedPtr< IToolkitHost >& InitToolkitHost, const TArray< TWeakObjectPtr< UObject > >& ObjectsToEdit );
 
-	FPropertyTypeLayoutCallback GetPropertyTypeCustomization(const UProperty* InProperty,const IPropertyHandle& PropertyHandle);
-	bool IsCustomizedStruct(const UStruct* Struct) const;
+	FPropertyTypeLayoutCallback GetPropertyTypeCustomization(const UProperty* InProperty,const IPropertyHandle& PropertyHandle, TSharedPtr<IDetailsView> DetailsViewInstance );
+	bool IsCustomizedStruct(const UStruct* Struct, TSharedPtr<IDetailsView> DetailsViewInstance ) const;
 
 private:
 
@@ -310,5 +287,7 @@ private:
 	/** A mapping of class names to detail layout delegates, called when querying for custom detail layouts */
 	FCustomDetailLayoutNameMap ClassNameToDetailLayoutNameMap;
 	/** A mapping of property names to property type layout delegates, called when querying for custom property layouts */
-	FCustomPropertyTypeLayoutMap PropertyTypeToLayoutMap;
+	FCustomPropertyTypeLayoutMap GlobalPropertyTypeToLayoutMap;
+	/** Instanced property type customization mapped to a specific details view */
+	TMap< TWeakPtr<IDetailsView>, FCustomPropertyTypeLayoutMap > InstancePropertyTypeLayoutMap;
 };
