fixed TTP 345881 EDITOR: TextRenderActors do not reconize Capitalization as a change

diff --git a/Engine/Source/Editor/PropertyEditor/Private/PropertyHandleImpl.cpp b/Engine/Source/Editor/PropertyEditor/Private/PropertyHandleImpl.cpp
--- a/Engine/Source/Editor/PropertyEditor/Private/PropertyHandleImpl.cpp
+++ b/Engine/Source/Editor/PropertyEditor/Private/PropertyHandleImpl.cpp
@@ -381,7 +381,7 @@ FPropertyAccess::Result FPropertyValueImpl::ImportText( const TArray<FObjectBase
 			// Remove quotes from the original value because FName properties  
 			// are wrapped in quotes before getting here. This causes the 
 			// string comparison to fail even when the name is unchanged. 
-			if ( !bNotifiedPreChange && ( InValues[ObjectIndex].TrimQuotes() != PreviousValue || ( bFinished && bInteractiveChangeInProgress ) ) )
+			if ( !bNotifiedPreChange && ( FCString::Strcmp(*InValues[ObjectIndex].TrimQuotes(), *PreviousValue) != 0 || ( bFinished && bInteractiveChangeInProgress ) ) )
 			{
 				bNotifiedPreChange = true;
 				NotifiedObj = Cur.Object;
@@ -700,7 +700,8 @@ FPropertyAccess::Result FPropertyValueImpl::SetValueAsString( const FString& InV
 
 		FString PreviousValue;
 		GetValueAsString( PreviousValue );
-		if (ParentNode && (ParentNode->GetInstancesNum() == 1 || (Value.Len() && PreviousValue != Value)))
+
+		if (ParentNode && (ParentNode->GetInstancesNum() == 1 || (Value.Len() && FCString::Strcmp(*PreviousValue, *Value) != 0)))
 		{
 			ImportText( Value, PropertyNodePin.Get(), Flags );
 		}
