Editor - Checking parent nodes for edit inline properties to see if they are edit const.

diff --git a/Engine/Source/Editor/PropertyEditor/Private/PropertyNode.cpp b/Engine/Source/Editor/PropertyEditor/Private/PropertyNode.cpp
--- a/Engine/Source/Editor/PropertyEditor/Private/PropertyNode.cpp
+++ b/Engine/Source/Editor/PropertyEditor/Private/PropertyNode.cpp
@@ -555,6 +555,20 @@ bool FPropertyNode::IsEditConst() const
 						bIsEditConst = true;
 						break;
 					}
+					else
+					{
+						// If 
+						if ( const FPropertyNode* ObjectParentNode = ObjectPropertyNode->GetParentNode() )
+						{
+							if ( const UProperty* ParentProperty = ObjectParentNode->GetProperty() )
+							{
+								if ( ParentProperty->PropertyFlags & CPF_EditInline )
+								{
+									return ObjectParentNode->IsEditConst();
+								}
+							}
+						}
+					}
 				}
 			}
 		}
