Fixed broken IsEditConst node

diff --git a/Engine/Source/Editor/PropertyEditor/Private/PropertyNode.cpp b/Engine/Source/Editor/PropertyEditor/Private/PropertyNode.cpp
--- a/Engine/Source/Editor/PropertyEditor/Private/PropertyNode.cpp
+++ b/Engine/Source/Editor/PropertyEditor/Private/PropertyNode.cpp
@@ -555,7 +555,7 @@ bool FPropertyNode::IsEditConst() const
 						bIsEditConst = true;
 						break;
 					}
-					else
+					/*else
 					{
 						// If 
 						if ( const FPropertyNode* ObjectParentNode = ObjectPropertyNode->GetParentNode() )
@@ -568,7 +568,7 @@ bool FPropertyNode::IsEditConst() const
 								}
 							}
 						}
-					}
+					}*/
 				}
 			}
 		}
