Editor - Reporting the correct callback if an instance customization exist.

diff --git a/Engine/Source/Editor/PropertyEditor/Private/PropertyEditorModule.cpp b/Engine/Source/Editor/PropertyEditor/Private/PropertyEditorModule.cpp
--- a/Engine/Source/Editor/PropertyEditor/Private/PropertyEditorModule.cpp
+++ b/Engine/Source/Editor/PropertyEditor/Private/PropertyEditorModule.cpp
@@ -648,12 +648,13 @@ FPropertyTypeLayoutCallback FPropertyEditorModule::GetPropertyTypeCustomization(
 				
 			if( !LayoutCallbacks )
 			{
-				FPropertyTypeLayoutCallbackList* LayoutCallbacks = GlobalPropertyTypeToLayoutMap.Find(PropertyTypeName);
-				if (LayoutCallbacks)
-				{
-					const FPropertyTypeLayoutCallback& Callback = LayoutCallbacks->Find(PropertyHandle);
-					return Callback;
-				}
+				LayoutCallbacks = GlobalPropertyTypeToLayoutMap.Find(PropertyTypeName);
+			}
+
+			if ( LayoutCallbacks )
+			{
+				const FPropertyTypeLayoutCallback& Callback = LayoutCallbacks->Find(PropertyHandle);
+				return Callback;
 			}
 		}
 	}
