Core: Remove unused bChangesTopology member from FPropertyChangedEvent and removed most uses, in preparation for deprecating the three-argument constructor

diff --git a/Engine/Source/Editor/DetailCustomizations/Private/ComponentTransformDetails.cpp b/Engine/Source/Editor/DetailCustomizations/Private/ComponentTransformDetails.cpp
--- a/Engine/Source/Editor/DetailCustomizations/Private/ComponentTransformDetails.cpp
+++ b/Engine/Source/Editor/DetailCustomizations/Private/ComponentTransformDetails.cpp
@@ -1173,7 +1173,7 @@ void FComponentTransformDetails::OnSetRotation( float NewValue, bool bCommitted,
 							PropagateTransformPropertyChange(Object, RelativeRotationProperty, OldRelativeRotation, RelativeRotation);
 						}
 
-						FPropertyChangedEvent PropertyChangedEvent( RelativeRotationProperty, false, !bCommitted && bEditingRotationInUI ? EPropertyChangeType::Interactive : EPropertyChangeType::ValueSet );
+						FPropertyChangedEvent PropertyChangedEvent( RelativeRotationProperty, !bCommitted && bEditingRotationInUI ? EPropertyChangeType::Interactive : EPropertyChangeType::ValueSet );
 
 						if( NotifyHook )
 						{
@@ -1401,7 +1401,7 @@ void FComponentTransformDetails::ScaleObject( float NewValue, int32 Axis, bool b
 					}
 					PropertyChain.AddHead(RelativeScale3DProperty);
 
-					FPropertyChangedEvent PropertyChangedEvent(RelativeScale3DProperty, false, EPropertyChangeType::ValueSet);
+					FPropertyChangedEvent PropertyChangedEvent(RelativeScale3DProperty, EPropertyChangeType::ValueSet);
 					FPropertyChangedChainEvent PropertyChangedChainEvent(PropertyChain, PropertyChangedEvent);
 					Object->PostEditChangeChainProperty( PropertyChangedChainEvent );
 
diff --git a/Engine/Source/Editor/DetailCustomizations/Private/MathStructProxyCustomizations.cpp b/Engine/Source/Editor/DetailCustomizations/Private/MathStructProxyCustomizations.cpp
--- a/Engine/Source/Editor/DetailCustomizations/Private/MathStructProxyCustomizations.cpp
+++ b/Engine/Source/Editor/DetailCustomizations/Private/MathStructProxyCustomizations.cpp
@@ -455,8 +455,7 @@ bool FTransformStructCustomization::FlushValues( TWeakPtr<IPropertyHandle> Prope
 
 	if(PropertyUtilities.IsValid() && !bIsUsingSlider)
 	{
-		const bool bTopologyChange = false;
-		FPropertyChangedEvent ChangeEvent(PropertyHandle->GetProperty(), bTopologyChange, EPropertyChangeType::ValueSet);
+		FPropertyChangedEvent ChangeEvent(PropertyHandle->GetProperty(), EPropertyChangeType::ValueSet);
 		PropertyUtilities->NotifyFinishedChangingProperties(ChangeEvent);
 	}
 
diff --git a/Engine/Source/Editor/PropertyEditor/Private/Presentation/PropertyEditor/PropertyEditor.cpp b/Engine/Source/Editor/PropertyEditor/Private/Presentation/PropertyEditor/PropertyEditor.cpp
--- a/Engine/Source/Editor/PropertyEditor/Private/Presentation/PropertyEditor/PropertyEditor.cpp
+++ b/Engine/Source/Editor/PropertyEditor/Private/Presentation/PropertyEditor/PropertyEditor.cpp
@@ -326,8 +326,7 @@ void FPropertyEditor::SetEditConditionState( bool bShouldEnable )
 		}
 	}
 
-	const bool bTopologyChange = true;
-	FPropertyChangedEvent ChangeEvent(PropertyNode->GetProperty(), bTopologyChange);
+	FPropertyChangedEvent ChangeEvent(PropertyNode->GetProperty());
 	PropertyNode->NotifyPostChange( ChangeEvent, PropertyUtilities->GetNotifyHook() );
 }
 
@@ -402,9 +401,7 @@ void FPropertyEditor::OnCustomResetToDefault( FSimpleDelegate OnCustomResetToDef
 		OnCustomResetToDefaultDelegate.Execute();
 
 		// Call PostEditchange on all the objects
-		// Assume reset to default, can change topology
-		const bool bTopologyChange = true;
-		FPropertyChangedEvent ChangeEvent( PropertyNode->GetProperty(), bTopologyChange );
+		FPropertyChangedEvent ChangeEvent( PropertyNode->GetProperty() );
 		PropertyNode->NotifyPostChange( ChangeEvent, PropertyUtilities->GetNotifyHook() );
 	}
 }
diff --git a/Engine/Source/Editor/PropertyEditor/Private/PropertyHandleImpl.cpp b/Engine/Source/Editor/PropertyEditor/Private/PropertyHandleImpl.cpp
--- a/Engine/Source/Editor/PropertyEditor/Private/PropertyHandleImpl.cpp
+++ b/Engine/Source/Editor/PropertyEditor/Private/PropertyHandleImpl.cpp
@@ -433,8 +433,7 @@ FPropertyAccess::Result FPropertyValueImpl::ImportText( const TArray<FObjectBase
 		if ( bNotifiedPreChange )
 		{
 			// Call PostEditChange on all objects.
-			const bool bTopologyChange = false;
-			FPropertyChangedEvent ChangeEvent(NodeProperty, bTopologyChange, bFinished ? EPropertyChangeType::ValueSet : EPropertyChangeType::Interactive );
+			FPropertyChangedEvent ChangeEvent(NodeProperty, bFinished ? EPropertyChangeType::ValueSet : EPropertyChangeType::Interactive );
 			ChangeEvent.SetArrayIndexPerObject(ArrayIndicesPerObject);
 
 			InPropertyNode->NotifyPostChange( ChangeEvent, NotifyHook );
@@ -459,8 +458,7 @@ FPropertyAccess::Result FPropertyValueImpl::ImportText( const TArray<FObjectBase
 
 		if (PropertyUtilities.IsValid() && !bInteractiveChangeInProgress)
 		{
-			const bool bTopologyChange = false;
-			FPropertyChangedEvent ChangeEvent(NodeProperty, bTopologyChange, bFinished ? EPropertyChangeType::ValueSet : EPropertyChangeType::Interactive);
+			FPropertyChangedEvent ChangeEvent(NodeProperty, bFinished ? EPropertyChangeType::ValueSet : EPropertyChangeType::Interactive);
 			InPropertyNode->FixPropertiesInEvent(ChangeEvent);
 			PropertyUtilities.Pin()->NotifyFinishedChangingProperties(ChangeEvent);
 		}
@@ -1112,18 +1110,17 @@ void FPropertyValueImpl::AddChild()
 				}
 			}
 
-			const bool bTopologyChange = true;
 			if ( bNotifiedPreChange )
 			{
 				// send the PostEditChange notification; it will be propagated to all selected objects
-				FPropertyChangedEvent ChangeEvent(NodeProperty, bTopologyChange, EPropertyChangeType::ArrayAdd);
+				FPropertyChangedEvent ChangeEvent(NodeProperty, EPropertyChangeType::ArrayAdd);
 				ChangeEvent.SetArrayIndexPerObject(ArrayIndicesPerObject);
 				PropertyNodePin->NotifyPostChange( ChangeEvent, NotifyHook );
 			}
 
 			if (PropertyUtilities.IsValid())
 			{
-				FPropertyChangedEvent ChangeEvent(NodeProperty, bTopologyChange, EPropertyChangeType::ArrayAdd);
+				FPropertyChangedEvent ChangeEvent(NodeProperty, EPropertyChangeType::ArrayAdd);
 				PropertyNodePin->FixPropertiesInEvent(ChangeEvent);
 				PropertyUtilities.Pin()->NotifyFinishedChangingProperties(ChangeEvent);
 			}
@@ -1188,16 +1185,15 @@ void FPropertyValueImpl::ClearChildren()
 
 			if ( bNotifiedPreChange )
 			{
-				const bool bTopologyChange = true;
 				// send the PostEditChange notification; it will be propagated to all selected objects
-				FPropertyChangedEvent ChangeEvent(NodeProperty, bTopologyChange, EPropertyChangeType::ValueSet);
+				FPropertyChangedEvent ChangeEvent(NodeProperty, EPropertyChangeType::ValueSet);
 				PropertyNodePin->NotifyPostChange( ChangeEvent, NotifyHook );
 			}
 		}
 
 		if (PropertyUtilities.IsValid())
 		{
-			FPropertyChangedEvent ChangeEvent(NodeProperty, true, EPropertyChangeType::ValueSet);
+			FPropertyChangedEvent ChangeEvent(NodeProperty, EPropertyChangeType::ValueSet);
 			PropertyNodePin->FixPropertiesInEvent(ChangeEvent);
 			PropertyUtilities.Pin()->NotifyFinishedChangingProperties(ChangeEvent);
 		}
@@ -1265,16 +1261,15 @@ void FPropertyValueImpl::InsertChild( TSharedPtr<FPropertyNode> ChildNodeToInser
 			FPropertyValueImpl::GenerateArrayIndexMapToObjectNode(ArrayIndicesPerObject[ObjectIndex], ChildNodePtr );
 		}
 
-		const bool bTopologyChange = true;
 		{
-			FPropertyChangedEvent ChangeEvent(ParentNode->GetProperty(), bTopologyChange, EPropertyChangeType::ArrayAdd);
+			FPropertyChangedEvent ChangeEvent(ParentNode->GetProperty(), EPropertyChangeType::ArrayAdd);
 			ChangeEvent.SetArrayIndexPerObject(ArrayIndicesPerObject);
 			ChildNodePtr->NotifyPostChange(ChangeEvent, NotifyHook);
 		}
 
 		if (PropertyUtilities.IsValid())
 		{
-			FPropertyChangedEvent ChangeEvent(ParentNode->GetProperty(), bTopologyChange, EPropertyChangeType::ArrayAdd);
+			FPropertyChangedEvent ChangeEvent(ParentNode->GetProperty(), EPropertyChangeType::ArrayAdd);
 			ChildNodePtr->FixPropertiesInEvent(ChangeEvent);
 			PropertyUtilities.Pin()->NotifyFinishedChangingProperties(ChangeEvent);
 		}
@@ -1337,15 +1332,14 @@ void FPropertyValueImpl::DeleteChild( TSharedPtr<FPropertyNode> ChildNodeToDelet
 			}
 		}
 
-		const bool bTopologyChange = true;
 		{
-			FPropertyChangedEvent ChangeEvent(ParentNode->GetProperty(), bTopologyChange);
+			FPropertyChangedEvent ChangeEvent(ParentNode->GetProperty());
 			ChildNodePtr->NotifyPostChange(ChangeEvent, NotifyHook);
 		}
 
 		if (PropertyUtilities.IsValid())
 		{
-			FPropertyChangedEvent ChangeEvent(ParentNode->GetProperty(), bTopologyChange);
+			FPropertyChangedEvent ChangeEvent(ParentNode->GetProperty());
 			ChildNodePtr->FixPropertiesInEvent(ChangeEvent);
 			PropertyUtilities.Pin()->NotifyFinishedChangingProperties(ChangeEvent);
 		}
@@ -1429,14 +1423,13 @@ void FPropertyValueImpl::DuplicateChild( TSharedPtr<FPropertyNode> ChildNodeToDu
 		// 		const bool bRecurse = false;
 		// 		ParentNode->SetExpanded(bExpand, bRecurse);
 
-		const bool bTopologyChange = true;
 		{
-			FPropertyChangedEvent ChangeEvent(ParentNode->GetProperty(), bTopologyChange, EPropertyChangeType::ValueSet);
+			FPropertyChangedEvent ChangeEvent(ParentNode->GetProperty(), EPropertyChangeType::ValueSet);
 			ChildNodePtr->NotifyPostChange(ChangeEvent, NotifyHook);
 		}
 		if (PropertyUtilities.IsValid())
 		{
-			FPropertyChangedEvent ChangeEvent(ParentNode->GetProperty(), bTopologyChange, EPropertyChangeType::ValueSet);
+			FPropertyChangedEvent ChangeEvent(ParentNode->GetProperty(), EPropertyChangeType::ValueSet);
 			ChildNodePtr->FixPropertiesInEvent(ChangeEvent);
 			PropertyUtilities.Pin()->NotifyFinishedChangingProperties(ChangeEvent);
 		}
diff --git a/Engine/Source/Editor/PropertyEditor/Private/PropertyNode.cpp b/Engine/Source/Editor/PropertyEditor/Private/PropertyNode.cpp
--- a/Engine/Source/Editor/PropertyEditor/Private/PropertyNode.cpp
+++ b/Engine/Source/Editor/PropertyEditor/Private/PropertyNode.cpp
@@ -1594,8 +1594,7 @@ void FPropertyNode::ResetToDefault( FNotifyHook* InNotifyHook )
 		{
 			// Call PostEditchange on all the objects
 			// Assume reset to default, can change topology
-			const bool bTopologyChange = true;
-			FPropertyChangedEvent ChangeEvent( TheProperty, bTopologyChange );
+			FPropertyChangedEvent ChangeEvent( TheProperty );
 			NotifyPostChange( ChangeEvent, InNotifyHook );
 		}
 
@@ -1902,7 +1901,7 @@ void FPropertyNode::NotifyPostChange( FPropertyChangedEvent& InPropertyChangedEv
 					if (CurProperty != InPropertyChangedEvent.Property)
 					{
 						//parent object node property.  Reset other internals and leave the event type as unspecified
-						ChangedEvent = FPropertyChangedEvent(CurProperty, InPropertyChangedEvent.bChangesTopology, InPropertyChangedEvent.ChangeType);
+						ChangedEvent = FPropertyChangedEvent(CurProperty, InPropertyChangedEvent.ChangeType);
 					}
 					ChangedEvent.ObjectIteratorIndex = CurrentObjectIndex;
 					Object->PostEditChangeProperty( ChangedEvent );
@@ -1913,7 +1912,7 @@ void FPropertyNode::NotifyPostChange( FPropertyChangedEvent& InPropertyChangedEv
 					if (CurProperty != InPropertyChangedEvent.Property)
 					{
 						//parent object node property.  Reset other internals and leave the event type as unspecified
-						ChangedEvent = FPropertyChangedEvent(CurProperty, InPropertyChangedEvent.bChangesTopology, InPropertyChangedEvent.ChangeType);
+						ChangedEvent = FPropertyChangedEvent(CurProperty, InPropertyChangedEvent.ChangeType);
 					}
 					FPropertyChangedChainEvent ChainEvent(*PropertyChain, ChangedEvent);
 					ChainEvent.ObjectIteratorIndex = CurrentObjectIndex;
diff --git a/Engine/Source/Editor/PropertyEditor/Private/SDetailsViewBase.cpp b/Engine/Source/Editor/PropertyEditor/Private/SDetailsViewBase.cpp
--- a/Engine/Source/Editor/PropertyEditor/Private/SDetailsViewBase.cpp
+++ b/Engine/Source/Editor/PropertyEditor/Private/SDetailsViewBase.cpp
@@ -216,7 +216,7 @@ void SDetailsViewBase::SetColorPropertyFromColorPicker(FLinearColor NewColor)
 
 			PinnedColorPropertyNode->NotifyPreChange(Property, GetNotifyHook());
 
-			FPropertyChangedEvent ChangeEvent(Property, false, EPropertyChangeType::ValueSet);
+			FPropertyChangedEvent ChangeEvent(Property, EPropertyChangeType::ValueSet);
 			PinnedColorPropertyNode->NotifyPostChange(ChangeEvent, GetNotifyHook());
 		}
 	}
@@ -230,7 +230,7 @@ void SDetailsViewBase::OnColorPickerWindowClosed(const TSharedRef<SWindow>& Wind
 		UProperty* Property = PinnedColorPropertyNode->GetProperty();
 		if (Property && PropertyUtilities.IsValid())
 		{
-			FPropertyChangedEvent ChangeEvent(Property, false, EPropertyChangeType::ArrayAdd);
+			FPropertyChangedEvent ChangeEvent(Property, EPropertyChangeType::ArrayAdd);
 			PinnedColorPropertyNode->FixPropertiesInEvent(ChangeEvent);
 			PropertyUtilities->NotifyFinishedChangingProperties(ChangeEvent);
 		}
diff --git a/Engine/Source/Editor/PropertyEditor/Private/SPropertyTreeViewImpl.cpp b/Engine/Source/Editor/PropertyEditor/Private/SPropertyTreeViewImpl.cpp
--- a/Engine/Source/Editor/PropertyEditor/Private/SPropertyTreeViewImpl.cpp
+++ b/Engine/Source/Editor/PropertyEditor/Private/SPropertyTreeViewImpl.cpp
@@ -1255,7 +1255,7 @@ void SPropertyTreeViewImpl::SetColor(FLinearColor NewColor)
 
 		ColorPropertyNode->NotifyPreChange(NodeProperty, GetNotifyHook());
 
-		FPropertyChangedEvent ChangeEvent(NodeProperty, false, EPropertyChangeType::ValueSet);
+		FPropertyChangedEvent ChangeEvent(NodeProperty, EPropertyChangeType::ValueSet);
 		ColorPropertyNode->NotifyPostChange( ChangeEvent, GetNotifyHook() );
 	}
 }
diff --git a/Engine/Source/Editor/PropertyEditor/Private/SSingleProperty.cpp b/Engine/Source/Editor/PropertyEditor/Private/SSingleProperty.cpp
--- a/Engine/Source/Editor/PropertyEditor/Private/SSingleProperty.cpp
+++ b/Engine/Source/Editor/PropertyEditor/Private/SSingleProperty.cpp
@@ -303,7 +303,7 @@ void SSingleProperty::SetColorPropertyFromColorPicker(FLinearColor NewColor)
 		//@todo if multiple objects we need to iterate
 		ValueNode->NotifyPreChange(NodeProperty, GetNotifyHook());
 
-		FPropertyChangedEvent ChangeEvent(NodeProperty, false, EPropertyChangeType::ValueSet);
+		FPropertyChangedEvent ChangeEvent(NodeProperty, EPropertyChangeType::ValueSet);
 		ValueNode->NotifyPostChange( ChangeEvent, GetNotifyHook() );
 	}
 }
diff --git a/Engine/Source/Editor/PropertyEditor/Private/UserInterface/PropertyEditor/SPropertyEditorColor.cpp b/Engine/Source/Editor/PropertyEditor/Private/UserInterface/PropertyEditor/SPropertyEditorColor.cpp
--- a/Engine/Source/Editor/PropertyEditor/Private/UserInterface/PropertyEditor/SPropertyEditorColor.cpp
+++ b/Engine/Source/Editor/PropertyEditor/Private/UserInterface/PropertyEditor/SPropertyEditorColor.cpp
@@ -234,7 +234,7 @@ void SPropertyEditorColor::SetColor(FLinearColor NewColor)
 			}
 		}
 
-		FPropertyChangedEvent ChangeEvent(Property, false, EPropertyChangeType::ValueSet);
+		FPropertyChangedEvent ChangeEvent(Property, EPropertyChangeType::ValueSet);
 		PropertyNode->NotifyPostChange( ChangeEvent, PropertyUtilities->GetNotifyHook() );
 
 		GEditor->EndTransaction();		
@@ -273,7 +273,7 @@ void SPropertyEditorColor::OnColorPickerCancelled( FLinearColor OriginalColor )
 			}
 		}
 
-		FPropertyChangedEvent ChangeEvent(Property, false, EPropertyChangeType::ValueSet);
+		FPropertyChangedEvent ChangeEvent(Property, EPropertyChangeType::ValueSet);
 		PropertyNode->NotifyPostChange( ChangeEvent, PropertyUtilities->GetNotifyHook() );
 	}
 
diff --git a/Engine/Source/Editor/UnrealEd/Private/Tests/EditorBuildPromotionTests.cpp b/Engine/Source/Editor/UnrealEd/Private/Tests/EditorBuildPromotionTests.cpp
--- a/Engine/Source/Editor/UnrealEd/Private/Tests/EditorBuildPromotionTests.cpp
+++ b/Engine/Source/Editor/UnrealEd/Private/Tests/EditorBuildPromotionTests.cpp
@@ -946,7 +946,7 @@ namespace EditorBuildPromotionTestUtils
 
 			TargetObject->PreEditChange(FoundProperty);
 			FoundProperty->ImportText(*NewValueString, FoundProperty->ContainerPtrToValuePtr<uint8>(TargetObject), 0, TargetObject);
-			FPropertyChangedEvent PropertyChangedEvent(FoundProperty, false, EPropertyChangeType::ValueSet);
+			FPropertyChangedEvent PropertyChangedEvent(FoundProperty, EPropertyChangeType::ValueSet);
 			TargetObject->PostEditChangeProperty(PropertyChangedEvent);
 		}
 	}
diff --git a/Engine/Source/Editor/WorldBrowser/Private/StreamingLevels/StreamingLevelModel.cpp b/Engine/Source/Editor/WorldBrowser/Private/StreamingLevels/StreamingLevelModel.cpp
--- a/Engine/Source/Editor/WorldBrowser/Private/StreamingLevels/StreamingLevelModel.cpp
+++ b/Engine/Source/Editor/WorldBrowser/Private/StreamingLevels/StreamingLevelModel.cpp
@@ -87,7 +87,7 @@ void FStreamingLevelModel::SetLevelColor(FLinearColor InColor)
 
 		LevelStreaming.Get()->LevelColor = InColor;
 
-		FPropertyChangedEvent PropertyChangedEvent(DrawColorProperty, false, EPropertyChangeType::ValueSet);
+		FPropertyChangedEvent PropertyChangedEvent(DrawColorProperty, EPropertyChangeType::ValueSet);
 		LevelStreaming->PostEditChangeProperty(PropertyChangedEvent);
 	}
 }
diff --git a/Engine/Source/Runtime/CoreUObject/Private/UObject/Obj.cpp b/Engine/Source/Runtime/CoreUObject/Private/UObject/Obj.cpp
--- a/Engine/Source/Runtime/CoreUObject/Private/UObject/Obj.cpp
+++ b/Engine/Source/Runtime/CoreUObject/Private/UObject/Obj.cpp
@@ -243,7 +243,7 @@ void UObject::PreEditChange( FEditPropertyChain& PropertyAboutToChange )
 
 void UObject::PostEditChangeChainProperty(FPropertyChangedChainEvent& PropertyChangedEvent)
 {
-	FPropertyChangedEvent PropertyEvent(PropertyChangedEvent.PropertyChain.GetActiveNode()->GetValue(), PropertyChangedEvent.bChangesTopology, PropertyChangedEvent.ChangeType);
+	FPropertyChangedEvent PropertyEvent(PropertyChangedEvent.PropertyChain.GetActiveNode()->GetValue(), PropertyChangedEvent.ChangeType);
 
 	if( PropertyChangedEvent.PropertyChain.GetActiveMemberNode() )
 	{
diff --git a/Engine/Source/Runtime/CoreUObject/Public/UObject/UnrealType.h b/Engine/Source/Runtime/CoreUObject/Public/UObject/UnrealType.h
--- a/Engine/Source/Runtime/CoreUObject/Public/UObject/UnrealType.h
+++ b/Engine/Source/Runtime/CoreUObject/Public/UObject/UnrealType.h
@@ -2710,17 +2710,32 @@ namespace EPropertyChangeType
 struct FPropertyChangedEvent
 {
 	//default constructor
-	FPropertyChangedEvent(UProperty* InProperty, const bool bInChangesTopology=false, EPropertyChangeType::Type InChangeType=EPropertyChangeType::Unspecified)
+	FPropertyChangedEvent(UProperty* InProperty)
+		: Property(InProperty)
+		, MemberProperty(InProperty)
+		, ChangeType(EPropertyChangeType::Unspecified)
+		, ObjectIteratorIndex(INDEX_NONE)
+		, ArrayIndicesPerObject(nullptr)
 	{
-		Property = InProperty;
-		MemberProperty = InProperty;
-		bChangesTopology = bInChangesTopology;
-		ChangeType = InChangeType;
+	}
+
+	FPropertyChangedEvent(UProperty* InProperty, EPropertyChangeType::Type InChangeType)
+		: Property(InProperty)
+		, MemberProperty(InProperty)
+		, ChangeType(InChangeType)
+		, ObjectIteratorIndex(INDEX_NONE)
+		, ArrayIndicesPerObject(nullptr)
+	{
+	}
 
-		//default to out of bounds/unused
-		ObjectIteratorIndex = -1;
-		//default to no array index knowledge
-		ArrayIndicesPerObject = NULL;
+	//@TODO: DEPRECATED(4.6, "The bInChangesTopology parameter has been removed, use the two-argument constructor of FPropertyChangedEvent instead")
+	FPropertyChangedEvent(UProperty* InProperty, const bool /*bInChangesTopology*/, EPropertyChangeType::Type InChangeType)
+		: Property(InProperty)
+		, MemberProperty(InProperty)
+		, ChangeType(InChangeType)
+		, ObjectIteratorIndex(INDEX_NONE)
+		, ArrayIndicesPerObject(nullptr)
+	{
 	}
 
 	void SetActiveMemberProperty( UProperty* InActiveMemberProperty )
@@ -2740,7 +2755,7 @@ struct FPropertyChangedEvent
 	 * Gets the Array Index of the "current object" based on a particular name
 	 * InName - Name of the property to find the array index for
 	 */
-	int32 GetArrayIndex (const FString& InName)
+	int32 GetArrayIndex(const FString& InName)
 	{
 		//default to unknown index
 		int32 Retval = -1;
@@ -2766,10 +2781,11 @@ struct FPropertyChangedEvent
 	 */
 	UProperty* MemberProperty;
 
-	bool bChangesTopology;
+	// The kind of change event that occurred
 	EPropertyChangeType::Type ChangeType;
-	//Used by the param system to say which object is receiving the event in the case of multi-select
-	int32                       ObjectIteratorIndex;
+
+	// Used by the param system to say which object is receiving the event in the case of multi-select
+	int32 ObjectIteratorIndex;
 private:
 	//In the property window, multiple objects can be selected at once.  In the case of adding/inserting to an array, each object COULD have different indices for the new entries in the array
 	const TArray< TMap<FString,int32> >* ArrayIndicesPerObject;
