[INTEGRATE] Change 2305458 by [REDACTED_PERSON][REDACTED_HANDLE] on 2014/09/22 08:23:49

	The default (6 digits after point) string representation of float is used in PropertyEditor. It's the highest exact precision, and it's used everywhere else in the code.

	ttp 346819 Cannot recompile a blueprint
	The issue was introduced in cl#2228482

diff --git a/Engine/Source/Editor/PropertyEditor/Private/PropertyHandleImpl.cpp b/Engine/Source/Editor/PropertyEditor/Private/PropertyHandleImpl.cpp
--- a/Engine/Source/Editor/PropertyEditor/Private/PropertyHandleImpl.cpp
+++ b/Engine/Source/Editor/PropertyEditor/Private/PropertyHandleImpl.cpp
@@ -2112,7 +2112,7 @@ FPropertyAccess::Result FPropertyHandleFloat::SetValue( const float& NewValue, E
 	// Clamp the value from any meta data ranges stored on the property value
 	float FinalValue = Implementation->ClampFloatValueFromMetaData( NewValue );
 
-	const FString ValueStr = FString::Printf( TEXT("%.9g"), FinalValue );
+	const FString ValueStr = FString::Printf( TEXT("%f"), FinalValue );
 	Res = Implementation->ImportText( ValueStr, Flags );
 
 	return Res;
