BP: in PropertyEditorEditInline one can set a None class, to null property.

diff --git a/Engine/Source/Editor/PropertyEditor/Private/UserInterface/PropertyEditor/SPropertyEditorEditInline.cpp b/Engine/Source/Editor/PropertyEditor/Private/UserInterface/PropertyEditor/SPropertyEditorEditInline.cpp
--- a/Engine/Source/Editor/PropertyEditor/Private/UserInterface/PropertyEditor/SPropertyEditorEditInline.cpp
+++ b/Engine/Source/Editor/PropertyEditor/Private/UserInterface/PropertyEditor/SPropertyEditorEditInline.cpp
@@ -153,6 +153,7 @@ TSharedRef<SWidget> SPropertyEditorEditInline::GenerateClassPicker()
 	FClassViewerInitializationOptions Options;
 	Options.bShowUnloadedBlueprints = true;
 	Options.bShowDisplayNames = true;
+	Options.bShowNoneOption = true;
 
 	TSharedPtr<FPropertyEditorInlineClassFilter> ClassFilter = MakeShareable( new FPropertyEditorInlineClassFilter );
 	Options.ClassFilter = ClassFilter;
@@ -182,11 +183,6 @@ TSharedRef<SWidget> SPropertyEditorEditInline::GenerateClassPicker()
 
 void SPropertyEditorEditInline::OnClassPicked(UClass* InClass)
 {
-	if( !InClass )
-	{
-		return;
-	}
-
 	TArray<FObjectBaseAddress> ObjectsToModify;
 	TArray<FString> NewValues;
 
@@ -197,16 +193,25 @@ void SPropertyEditorEditInline::OnClassPicked(UClass* InClass)
 	{
 		for ( TPropObjectIterator Itor( ObjectNode->ObjectIterator() ) ; Itor ; ++Itor )
 		{
-			UObject*		Object = Itor->Get();
-			UObject*		UseOuter			= ( InClass->IsChildOf( UClass::StaticClass() ) ? Cast<UClass>(Object)->GetDefaultObject() : Object );
-			EObjectFlags	MaskedOuterFlags	= UseOuter ? UseOuter->GetMaskedFlags(RF_PropagateToSubObjects) : RF_NoFlags;
-			if (UseOuter && UseOuter->HasAnyFlags(RF_ClassDefaultObject | RF_ArchetypeObject))
+			FString NewValue;
+			if (InClass)
 			{
-				MaskedOuterFlags |= RF_ArchetypeObject;
+				UObject*		Object = Itor->Get();
+				UObject*		UseOuter = (InClass->IsChildOf(UClass::StaticClass()) ? Cast<UClass>(Object)->GetDefaultObject() : Object);
+				EObjectFlags	MaskedOuterFlags = UseOuter ? UseOuter->GetMaskedFlags(RF_PropagateToSubObjects) : RF_NoFlags;
+				if (UseOuter && UseOuter->HasAnyFlags(RF_ClassDefaultObject | RF_ArchetypeObject))
+				{
+					MaskedOuterFlags |= RF_ArchetypeObject;
+				}
+				UObject*		NewObject = StaticConstructObject(InClass, UseOuter, NAME_None, MaskedOuterFlags, NULL);
+
+				NewValue = NewObject->GetPathName();
 			}
-			UObject*		NewObject = StaticConstructObject(InClass, UseOuter, NAME_None, MaskedOuterFlags, NULL);
-
-			NewValues.Add( NewObject->GetPathName() );
+			else
+			{
+				NewValue = FName(NAME_None).ToString();
+			}
+			NewValues.Add(NewValue);
 		}
 
 		const TSharedRef< IPropertyHandle > PropertyHandle = PropertyEditor->GetPropertyHandle();
