Added Tag to FDetailWidgetRow.
Made missing landscape mode items selectable for tutortial

diff --git a/Engine/Source/Editor/LandscapeEditor/Private/LandscapeEditorDetailCustomization_NewLandscape.cpp b/Engine/Source/Editor/LandscapeEditor/Private/LandscapeEditorDetailCustomization_NewLandscape.cpp
--- a/Engine/Source/Editor/LandscapeEditor/Private/LandscapeEditorDetailCustomization_NewLandscape.cpp
+++ b/Engine/Source/Editor/LandscapeEditor/Private/LandscapeEditorDetailCustomization_NewLandscape.cpp
@@ -352,6 +352,7 @@ void FLandscapeEditorDetailCustomization_NewLandscape::CustomizeDetails(IDetailL
 	];
 
 	NewLandscapeCategory.AddCustomRow("Overall Resolution")
+	.RowTag("LandscapeEditor.OverallResolution")
 	.NameContent()
 	[
 		SNew(SBox)
@@ -410,6 +411,7 @@ void FLandscapeEditorDetailCustomization_NewLandscape::CustomizeDetails(IDetailL
 	];
 
 	NewLandscapeCategory.AddCustomRow("Total Components")
+	.RowTag("LandscapeEditor.TotalComponents")
 	.NameContent()
 	[
 		SNew(SBox)
diff --git a/Engine/Source/Editor/PropertyEditor/Private/DetailItemNode.cpp b/Engine/Source/Editor/PropertyEditor/Private/DetailItemNode.cpp
--- a/Engine/Source/Editor/PropertyEditor/Private/DetailItemNode.cpp
+++ b/Engine/Source/Editor/PropertyEditor/Private/DetailItemNode.cpp
@@ -144,16 +144,17 @@ void FDetailItemNode::ToggleExpansion()
 
 TSharedRef< ITableRow > FDetailItemNode::GenerateNodeWidget( const TSharedRef<STableViewBase>& OwnerTable, const FDetailColumnSizeData& ColumnSizeData, const TSharedRef<IPropertyUtilities>& PropertyUtilities )
 {
-	FGraphNodeMetaData TagMeta(TEXT("DetailRowItem"));
-	if (Customization.IsValidCustomization() && Customization.GetPropertyNode().IsValid() )
-	{
-		TagMeta.FriendlyName = Customization.GetPropertyNode()->GetDisplayName();
-	}
+	FTagMetaData TagMeta(TEXT("DetailRowItem"));
 	if (ParentCategory.IsValid())
 	{
-		FString Tag = FString::Printf(TEXT("%s_%s"), *TagMeta.FriendlyName, *ParentCategory.Pin()->GetCategoryPathName());
-		TagMeta.Tag = *Tag;
-		TagMeta.TabTypeToOpen = TEXT("LevelEditorSelectionDetails");
+		if (Customization.IsValidCustomization() && Customization.GetPropertyNode().IsValid())
+		{
+			TagMeta.Tag = *FString::Printf(TEXT("DetailRowItem.%s"), *Customization.GetPropertyNode()->GetDisplayName());
+		}
+		else if (Customization.HasCustomWidget() )
+		{
+			TagMeta.Tag = Customization.GetWidgetRow().RowTagName;
+		}
 	}
 	if( Customization.HasPropertyNode() && Customization.GetPropertyNode()->AsCategoryNode() )
 	{
@@ -161,7 +162,7 @@ TSharedRef< ITableRow > FDetailItemNode::GenerateNodeWidget( const TSharedRef<ST
 			SNew(SDetailCategoryTableRow, AsShared(), OwnerTable)
 			.IsEnabled(IsParentEnabled)
 			.DisplayName(Customization.GetPropertyNode()->GetDisplayName())
-			.AddMetaData<FTutorialMetaData>(TagMeta)
+			.AddMetaData<FTagMetaData>(TagMeta)
 			.InnerCategory( true );
 	}
 	else
@@ -169,7 +170,7 @@ TSharedRef< ITableRow > FDetailItemNode::GenerateNodeWidget( const TSharedRef<ST
 		return
 			SNew(SDetailSingleItemRow, &Customization, HasMultiColumnWidget(), AsShared(), OwnerTable )
 			.IsEnabled( IsParentEnabled )
-			.AddMetaData<FTutorialMetaData>(TagMeta)
+			.AddMetaData<FTagMetaData>(TagMeta)
 			.ColumnSizeData(ColumnSizeData);
 	}
 }
diff --git a/Engine/Source/Editor/PropertyEditor/Public/DetailWidgetRow.h b/Engine/Source/Editor/PropertyEditor/Public/DetailWidgetRow.h
--- a/Engine/Source/Editor/PropertyEditor/Public/DetailWidgetRow.h
+++ b/Engine/Source/Editor/PropertyEditor/Public/DetailWidgetRow.h
@@ -70,6 +70,7 @@ public:
 		, FilterTextString()
 		, CopyMenuAction()
 		, PasteMenuAction()
+		, RowTagName()
 	{
 	}
 	
@@ -173,6 +174,15 @@ public:
 		return CopyMenuAction.ExecuteAction.IsBound() && PasteMenuAction.ExecuteAction.IsBound();
 	}
 
+	/**
+	* Sets a tag which can be used to identify this row 
+	*/
+	FDetailWidgetRow& RowTag(const FName& InRowTagName)
+	{
+		RowTagName = InRowTagName;
+		return *this;
+	}
+
 public:
 	/** Name column content */
 	FDetailWidgetDecl NameWidget;
@@ -190,5 +200,7 @@ public:
 	FUIAction CopyMenuAction;
 	/** Action for pasting data on this row */
 	FUIAction PasteMenuAction;
+	/* Tag to identify this row */
+	FName	RowTagName;
 };
 
