Fix crash when blueprints attempt to recompile an asset and refresh a details panel while the details panel is in the middle of an operation such as an array add

diff --git a/Engine/Source/Editor/PropertyEditor/Private/SDetailsViewBase.cpp b/Engine/Source/Editor/PropertyEditor/Private/SDetailsViewBase.cpp
--- a/Engine/Source/Editor/PropertyEditor/Private/SDetailsViewBase.cpp
+++ b/Engine/Source/Editor/PropertyEditor/Private/SDetailsViewBase.cpp
@@ -608,6 +608,15 @@ void SDetailsViewBase::Tick(const FGeometry& AllottedGeometry, const double InCu
 		DeferredActions.Empty();
 	}
 
+	if( RootPropertyNode == RootNodePendingKill )
+	{
+		// Reaquire the root property node.  It may have been changed by the deferred actions if something like a blueprint editor forcefully resets a details panel during a posteditchange
+		RootPropertyNode = GetRootNode();
+
+		RestoreExpandedItems();
+	}
+
+
 	bool bValidateExternalNodes = true;
 	FPropertyNode::DataValidationResult Result = RootPropertyNode->EnsureDataIsValid();
 	if (Result == FPropertyNode::PropertiesChanged || Result == FPropertyNode::EditInlineNewValueChanged)
@@ -755,6 +764,13 @@ void SDetailsViewBase::SaveExpandedItems()
 	TArray<FString> ExpandedPropertyItems;
 	GetExpandedItems(RootPropertyNode, ExpandedPropertyItems);
 
+	// Handle spaces in expanded node names by wrapping them in quotes
+	for( FString& String : ExpandedPropertyItems )
+	{
+		String.InsertAt(0, '"');
+		String.AppendChar('"');
+	}
+
 	TArray<FString> ExpandedCustomItems = ExpandedDetailNodes.Array();
 
 	// Expanded custom items may have spaces but SetSingleLineArray doesnt support spaces (treats it as another element in the array)
