In Details tabs, where no category is specified for a group of properties, use the heading "General" instead of a prettified version of the class name.
#ttp 345463 - NUX:Details: Don't draw confusing C++ class names as categories when no category is given
#branch UE4
#proj Editor.PropertyEditor
#note Some of the examples in the TTP are due to the category for the UPROPERTY being explicitly set equal to the class name.

diff --git a/Engine/Source/Editor/PropertyEditor/Private/DetailLayoutBuilderImpl.cpp b/Engine/Source/Editor/PropertyEditor/Private/DetailLayoutBuilderImpl.cpp
--- a/Engine/Source/Editor/PropertyEditor/Private/DetailLayoutBuilderImpl.cpp
+++ b/Engine/Source/Editor/PropertyEditor/Private/DetailLayoutBuilderImpl.cpp
@@ -18,11 +18,16 @@ FDetailLayoutBuilderImpl::FDetailLayoutBuilderImpl(FClassToPropertyMap& InProper
 
 IDetailCategoryBuilder& FDetailLayoutBuilderImpl::EditCategory( FName CategoryName, const FString& NewLocalizedDisplayName, ECategoryPriority::Type CategoryType )
 {
-	// Use the base class name if the user did not specify a name
+	FString LocalizedDisplayName = NewLocalizedDisplayName;
+
+	// Use a generic name if one was not specified
 	if( CategoryName == NAME_None )
 	{
-		const UStruct* BaseStruct = GetDetailsView().GetBaseStruct();
-		CategoryName = BaseStruct ? BaseStruct->GetFName() : FName("Generic");
+		static const FString GeneralString = NSLOCTEXT("DetailLayoutBuilderImpl", "General", "General").ToString();
+		static const FName GeneralName = *GeneralString;
+
+		CategoryName = GeneralName;
+		LocalizedDisplayName = GeneralString;
 	}
 
 	TSharedPtr<FDetailCategoryImpl> CategoryImpl;
@@ -53,7 +58,7 @@ IDetailCategoryBuilder& FDetailLayoutBuilderImpl::EditCategory( FName CategoryNa
 		CategoryImpl->SetSortOrder( SortOrder );
 	}
 
-	CategoryImpl->SetDisplayName( CategoryName, NewLocalizedDisplayName );
+	CategoryImpl->SetDisplayName( CategoryName, LocalizedDisplayName );
 
 	return *CategoryImpl;
 }
