Added the facility to create a new asset in the asset picker widget, by adding extra entries to the SPropertyMenuAssetPicker.
#ttp 345817 - NUX: MVP: Asset pickers should have Create New button
#branch UE4
#proj Editor.PropertyEditor, Editor.DetailCustomizations

diff --git a/Engine/Source/Editor/DetailCustomizations/Private/DialogueWaveWidgets.cpp b/Engine/Source/Editor/DetailCustomizations/Private/DialogueWaveWidgets.cpp
--- a/Engine/Source/Editor/DetailCustomizations/Private/DialogueWaveWidgets.cpp
+++ b/Engine/Source/Editor/DetailCustomizations/Private/DialogueWaveWidgets.cpp
@@ -199,13 +199,15 @@ TSharedRef<SWidget> SDialogueVoicePropertyEditor::OnGetMenuContent()
 {
 	TArray<const UClass*> AllowedClasses;
 	AllowedClasses.Add(UDialogueVoice::StaticClass());
+
 	UDialogueVoice* DialogueVoice = NULL;
 	{
 		UObject* Object = NULL;
 		DialogueVoicePropertyHandle->GetValue(Object);
 		DialogueVoice = Cast<UDialogueVoice>(Object);
 	}
-	return PropertyCustomizationHelpers::MakeAssetPickerWithMenu(DialogueVoice, false, &AllowedClasses, OnShouldFilterAsset, FOnAssetSelected::CreateSP( this, &SDialogueVoicePropertyEditor::OnAssetSelectedFromPicker ), FSimpleDelegate::CreateSP( this, &SDialogueVoicePropertyEditor::CloseMenu) );
+
+	return PropertyCustomizationHelpers::MakeAssetPickerWithMenu(DialogueVoice, false, AllowedClasses, OnShouldFilterAsset, FOnAssetSelected::CreateSP(this, &SDialogueVoicePropertyEditor::OnAssetSelectedFromPicker), FSimpleDelegate::CreateSP(this, &SDialogueVoicePropertyEditor::CloseMenu));
 }
 
 void SDialogueVoicePropertyEditor::CloseMenu()
diff --git a/Engine/Source/Editor/PropertyEditor/Private/PropertyCustomizationHelpers.cpp b/Engine/Source/Editor/PropertyEditor/Private/PropertyCustomizationHelpers.cpp
--- a/Engine/Source/Editor/PropertyEditor/Private/PropertyCustomizationHelpers.cpp
+++ b/Engine/Source/Editor/PropertyEditor/Private/PropertyCustomizationHelpers.cpp
@@ -164,7 +164,7 @@ namespace PropertyCustomizationHelpers
 			.IsFocusable( false );
 	}
 
-	TSharedRef<SWidget> MakeInsertDeleteDuplicateButton( FExecuteAction OnInsertClicked, FExecuteAction OnDeleteClicked, FExecuteAction OnDuplicateClicked )
+	TSharedRef<SWidget> MakeInsertDeleteDuplicateButton(FExecuteAction OnInsertClicked, FExecuteAction OnDeleteClicked, FExecuteAction OnDuplicateClicked)
 	{	
 		FMenuBuilder MenuContentBuilder( true, NULL );
 		{
@@ -203,13 +203,14 @@ namespace PropertyCustomizationHelpers
 			.OnAssetSelected( OnAssetSelectedFromPicker );
 	}
 
-	TSharedRef<SWidget> MakeAssetPickerWithMenu( const FAssetData& InitialObject, const bool AllowClear, const TArray<const UClass*>* const AllowedClasses, FOnShouldFilterAsset OnShouldFilterAsset, FOnAssetSelected OnSet, FSimpleDelegate OnClose )
+	TSharedRef<SWidget> MakeAssetPickerWithMenu( const FAssetData& InitialObject, const bool AllowClear, TArray<const UClass*> AllowedClasses, FOnShouldFilterAsset OnShouldFilterAsset, FOnAssetSelected OnSet, FSimpleDelegate OnClose)
 	{
-		return 
-			SNew( SPropertyMenuAssetPicker )
+		return
+			SNew(SPropertyMenuAssetPicker)
 			.InitialObject(InitialObject)
 			.AllowClear(AllowClear)
 			.AllowedClasses(AllowedClasses)
+			.NewAssetFactories(GetNewAssetFactoriesForClasses(AllowedClasses))
 			.OnShouldFilterAsset(OnShouldFilterAsset)
 			.OnSet(OnSet)
 			.OnClose(OnClose);
@@ -273,6 +274,29 @@ namespace PropertyCustomizationHelpers
 
 		return EditConditionProperty;
 	}
+
+	TArray<UFactory*> GetNewAssetFactoriesForClasses(const TArray<const UClass*>& Classes)
+	{
+		TArray<UFactory*> Factories;
+		for (TObjectIterator<UClass> It; It; ++It)
+		{
+			UClass* Class = *It;
+			if (Class->IsChildOf(UFactory::StaticClass()) && !Class->HasAnyClassFlags(CLASS_Abstract))
+			{
+				UFactory* Factory = Class->GetDefaultObject<UFactory>();
+				if (Factory->ShouldShowInNewMenu() && ensure(!Factory->GetDisplayName().IsEmpty()))
+				{
+					UClass* SupportedClass = Factory->GetSupportedClass();
+					if (SupportedClass != nullptr && Classes.ContainsByPredicate([=](const UClass* InClass) { return SupportedClass->IsChildOf(InClass); }))
+					{
+						Factories.Add(Factory);
+					}
+				}
+			}
+		}
+
+		return Factories;
+	}
 }
 
 void SObjectPropertyEntryBox::Construct( const FArguments& InArgs )
diff --git a/Engine/Source/Editor/PropertyEditor/Private/UserInterface/PropertyEditor/SPropertyEditorAsset.cpp b/Engine/Source/Editor/PropertyEditor/Private/UserInterface/PropertyEditor/SPropertyEditorAsset.cpp
--- a/Engine/Source/Editor/PropertyEditor/Private/UserInterface/PropertyEditor/SPropertyEditorAsset.cpp
+++ b/Engine/Source/Editor/PropertyEditor/Private/UserInterface/PropertyEditor/SPropertyEditorAsset.cpp
@@ -2,6 +2,7 @@
 
 #include "PropertyEditorPrivatePCH.h"
 #include "SPropertyEditorAsset.h"
+#include "SPropertyEditorNewAsset.h"
 #include "PropertyNode.h"
 #include "PropertyEditor.h"
 #include "AssetThumbnail.h"
@@ -14,6 +15,7 @@
 #include "SAssetDropTarget.h"
 #include "AssetRegistryModule.h"
 #include "Particles/ParticleSystem.h"
+
 #define LOCTEXT_NAMESPACE "PropertyEditor"
 
 DECLARE_DELEGATE( FOnCopy );
@@ -286,14 +288,24 @@ TSharedRef<SWidget> SPropertyEditorAsset::OnGetMenuContent()
 		TSharedPtr<SceneOutliner::FOutlinerFilters> ActorFilters = MakeShareable( new SceneOutliner::FOutlinerFilters );
 		ActorFilters->Add( MakeShareable( new TDelegateFilter< const AActor* const >( TDelegateFilter< const AActor* const >::FPredicate::CreateSP( this, &SPropertyEditorAsset::IsFilteredActor ) ) ) );
 
-		return PropertyCustomizationHelpers::MakeActorPickerWithMenu(Cast<AActor>(Value.Object), bAllowClear, ActorFilters, FOnActorSelected::CreateSP( this, &SPropertyEditorAsset::OnActorSelected), FSimpleDelegate::CreateSP( this, &SPropertyEditorAsset::CloseComboButton ), FSimpleDelegate::CreateSP( this, &SPropertyEditorAsset::OnUse ) );
+		return PropertyCustomizationHelpers::MakeActorPickerWithMenu(Cast<AActor>(Value.Object),
+																	 bAllowClear,
+																	 ActorFilters,
+																	 FOnActorSelected::CreateSP( this, &SPropertyEditorAsset::OnActorSelected),
+																	 FSimpleDelegate::CreateSP( this, &SPropertyEditorAsset::CloseComboButton ),
+																	 FSimpleDelegate::CreateSP( this, &SPropertyEditorAsset::OnUse ) );
 	}
 	else
 	{
 		TArray<const UClass*> AllowedClasses;
 		AllowedClasses.Add(ObjectClass);
 
-		return PropertyCustomizationHelpers::MakeAssetPickerWithMenu(Value.AssetData, bAllowClear, &AllowedClasses, OnShouldFilterAsset, FOnAssetSelected::CreateSP( this, &SPropertyEditorAsset::OnAssetSelected), FSimpleDelegate::CreateSP( this, &SPropertyEditorAsset::CloseComboButton ) );
+		return PropertyCustomizationHelpers::MakeAssetPickerWithMenu(Value.AssetData,
+																	 bAllowClear,
+																	 AllowedClasses,
+																	 OnShouldFilterAsset,
+																	 FOnAssetSelected::CreateSP(this, &SPropertyEditorAsset::OnAssetSelected),
+																	 FSimpleDelegate::CreateSP(this, &SPropertyEditorAsset::CloseComboButton));
 	}
 }
 
diff --git a/Engine/Source/Editor/PropertyEditor/Private/UserInterface/PropertyEditor/SPropertyEditorAsset.h b/Engine/Source/Editor/PropertyEditor/Private/UserInterface/PropertyEditor/SPropertyEditorAsset.h
--- a/Engine/Source/Editor/PropertyEditor/Private/UserInterface/PropertyEditor/SPropertyEditorAsset.h
+++ b/Engine/Source/Editor/PropertyEditor/Private/UserInterface/PropertyEditor/SPropertyEditorAsset.h
@@ -229,8 +229,14 @@ private:
 	 */
 	FReply OnAssetThumbnailDoubleClick( const FGeometry& InMyGeometry, const FPointerEvent& InMouseEvent );
 
+	/** 
+	 * Handle creating a new asset from the provided factory
+	 */
+	void SPropertyEditorAsset::OnCreateNewAsset( TWeakObjectPtr<UFactory> FactoryPtr );
+
 	/** @return True if the property can be edited */
 	bool CanEdit() const;
+
 private:
 
 	/** Main combobutton */
diff --git a/Engine/Source/Editor/PropertyEditor/Private/UserInterface/PropertyEditor/SPropertyEditorNewAsset.cpp b/Engine/Source/Editor/PropertyEditor/Private/UserInterface/PropertyEditor/SPropertyEditorNewAsset.cpp
new file mode 100644
--- /dev/null
+++ b/Engine/Source/Editor/PropertyEditor/Private/UserInterface/PropertyEditor/SPropertyEditorNewAsset.cpp
@@ -0,0 +1,224 @@
+// Copyright 1998-2014 Epic Games, Inc. All Rights Reserved.
+
+#include "PropertyEditorPrivatePCH.h"
+#include "SPropertyEditorNewAsset.h"
+#include "Editor/ContentBrowser/Public/ContentBrowserModule.h"
+#include "AssetRegistryModule.h"
+#include "Developer/AssetTools/Public/AssetToolsModule.h"
+
+#define LOCTEXT_NAMESPACE "SPropertyEditorNewAsset"
+
+UObject* SPropertyEditorNewAsset::Create(TWeakObjectPtr<UFactory> FactoryPtr)
+{
+	if (!FactoryPtr.IsValid())
+	{
+		return false;
+	}
+
+	FFormatNamedArguments Args;
+	Args.Add(TEXT("AssetType"), FactoryPtr->GetDisplayName());
+
+	TSharedPtr<SWindow> PropertyEditorNewAssetWindow =
+		SNew(SWindow)
+		.Title(FText::Format(LOCTEXT("SelectPath", "Select Path for {AssetType}"), Args))
+		.ToolTipText(FText::Format(LOCTEXT("SelectPathTooltip", "Select the path where the {AssetType} will be created."), Args))
+		.ClientSize(FVector2D(400, 400));
+
+	TSharedPtr<SPropertyEditorNewAsset> PropertyEditorNewAssetWidget;
+	PropertyEditorNewAssetWindow->SetContent(SAssignNew(PropertyEditorNewAssetWidget, SPropertyEditorNewAsset, PropertyEditorNewAssetWindow, FactoryPtr));
+
+	GEditor->EditorAddModalWindow(PropertyEditorNewAssetWindow.ToSharedRef());
+
+	check(PropertyEditorNewAssetWidget.IsValid());
+	return PropertyEditorNewAssetWidget->NewAsset;
+}
+
+void SPropertyEditorNewAsset::Construct(const FArguments& InArgs, TSharedPtr<SWindow> InParentWindow, TWeakObjectPtr<UFactory> InFactoryPtr)
+{
+	bIsReportingError = false;
+	ParentWindowPtr = InParentWindow;
+	FactoryPtr = InFactoryPtr;
+	NewAsset = nullptr;
+	AssetPath = FString("/Game");
+
+	check(FactoryPtr.IsValid());
+
+	FString AssetName = FactoryPtr->GetDefaultNewAssetName();
+
+	FPathPickerConfig PathPickerConfig;
+	PathPickerConfig.DefaultPath = AssetPath;
+	PathPickerConfig.OnPathSelected = FOnPathSelected::CreateRaw(this, &SPropertyEditorNewAsset::OnSelectAssetPath);
+
+	FContentBrowserModule& ContentBrowserModule = FModuleManager::LoadModuleChecked<FContentBrowserModule>("ContentBrowser");
+
+	ChildSlot
+	[
+		SNew(SVerticalBox)
+		+ SVerticalBox::Slot()
+		.FillHeight(1.0f)
+		[
+			ContentBrowserModule.Get().CreatePathPicker(PathPickerConfig)
+		]
+		+ SVerticalBox::Slot()
+		.AutoHeight()
+		.Padding(0.0f, 2.0f, 0.0f, 4.0f)
+		[
+			SNew(SSeparator)
+		]
+		+ SVerticalBox::Slot()
+		.AutoHeight()
+		.VAlign(VAlign_Center)
+		.Padding(2.0f, 0.0f, 2.0f, 0.0f)
+		[
+			SNew(SHorizontalBox)
+			+ SHorizontalBox::Slot()
+			.VAlign(VAlign_Center)
+			.Padding(0.0f, 0.0f, 16.0f, 0.0f)
+			.FillWidth(0.6f)
+			[
+				SNew(SHorizontalBox)
+				+ SHorizontalBox::Slot()
+				.VAlign(VAlign_Center)
+				.AutoWidth()
+				[
+					SNew(STextBlock)
+					.Text(LOCTEXT("TextboxLabel", "Asset Name:"))
+				]
+				+ SHorizontalBox::Slot()
+				.Padding(4.0f, 0.0f, 0.0f, 0.0f)
+				.VAlign(VAlign_Center)
+				.FillWidth(1.0f)
+				[
+					SAssignNew(FilenameWidget, SEditableTextBox)
+					.Text(FText::FromString(AssetName))
+					.OnTextChanged(this, &SPropertyEditorNewAsset::OnFilenameChanged)
+				]
+			]
+			+ SHorizontalBox::Slot()
+			.HAlign(HAlign_Right)
+			.VAlign(VAlign_Center)
+			.FillWidth(0.4f)
+			[
+				SNew(SHorizontalBox)
+				+ SHorizontalBox::Slot()
+				.HAlign(HAlign_Right)
+				.Padding(0.0f, 0.0f, 4.0f, 0.0f)
+				.AutoWidth()
+				[
+					SNew(SButton)
+					.VAlign(VAlign_Center)
+					.ContentPadding(FMargin(8, 2, 8, 2))
+					.OnClicked(this, &SPropertyEditorNewAsset::OnCreateNewAssetClicked)
+					.IsEnabled(this, &SPropertyEditorNewAsset::IsCreateNewAssetEnabled)
+					[
+						SNew(STextBlock)
+						.Text(LOCTEXT("CreateButtonText", "Create"))
+					]
+				]
+				+ SHorizontalBox::Slot()
+				.HAlign(HAlign_Right)
+				.AutoWidth()
+				[
+					SNew(SButton)
+					.VAlign(VAlign_Center)
+					.ContentPadding(FMargin(8, 2, 8, 2))
+					.OnClicked(this, &SPropertyEditorNewAsset::OnCancelCreateNewAsset)
+					[
+						SNew(STextBlock)
+						.Text(LOCTEXT("CancelButtonText", "Cancel"))
+					]
+				]
+			]
+		]
+		+ SVerticalBox::Slot()
+		.AutoHeight()
+		.VAlign(VAlign_Center)
+		.Padding(0.0f, 10.0f, 0.0f, 0.0f)
+		[
+			SNullWidget::NullWidget
+		]
+	];
+
+	OnFilenameChanged(FText::FromString(AssetName));
+}
+
+FReply SPropertyEditorNewAsset::OnCreateNewAssetClicked()
+{
+	TSharedPtr<SWindow> ParentWindow = ParentWindowPtr.Pin();
+	if (ParentWindow.IsValid())
+	{
+		ParentWindow->RequestDestroyWindow();
+	}
+
+	FString AssetName = FilenameWidget->GetText().ToString();
+	FString AssetFullName = AssetPath / AssetName;
+
+	if (FactoryPtr.IsValid())
+	{
+		UFactory* FactoryInstance = DuplicateObject<UFactory>(FactoryPtr.Get(), GetTransientPackage());
+
+		FEditorDelegates::OnConfigureNewAssetProperties.Broadcast(FactoryInstance);
+		if (FactoryInstance->ConfigureProperties())
+		{
+			FAssetToolsModule& AssetToolsModule = FAssetToolsModule::GetModule();
+			NewAsset = AssetToolsModule.Get().CreateAsset(AssetName, AssetPath, FactoryInstance->GetSupportedClass(), FactoryInstance);
+		}
+	}
+
+	return FReply::Handled();
+}
+
+FReply SPropertyEditorNewAsset::OnCancelCreateNewAsset()
+{
+	TSharedPtr<SWindow> ParentWindow = ParentWindowPtr.Pin();
+	if (ParentWindow.IsValid())
+	{
+		ParentWindow->RequestDestroyWindow();
+	}
+
+	return FReply::Handled();
+}
+
+void SPropertyEditorNewAsset::OnSelectAssetPath(const FString& Path)
+{
+	AssetPath = Path;
+	OnFilenameChanged(FilenameWidget->GetText());
+}
+
+void SPropertyEditorNewAsset::OnFilenameChanged(const FText& InNewName)
+{
+	TArray<FAssetData> AssetData;
+	FAssetRegistryModule& AssetRegistryModule = FModuleManager::LoadModuleChecked<FAssetRegistryModule>("AssetRegistry");
+	AssetRegistryModule.Get().GetAssetsByPath(FName(*AssetPath), AssetData);
+
+	FText ErrorText;
+	if (!FEditorFileUtils::IsFilenameValidForSaving(InNewName.ToString(), ErrorText) || !FName(*InNewName.ToString()).IsValidObjectName(ErrorText))
+	{
+		FilenameWidget->SetError(LOCTEXT("AssetNameInvalid", "Asset name is invalid"));
+		bIsReportingError = true;
+		return;
+	}
+	else
+	{
+		// Check to see if the name conflicts
+		for (auto Iter = AssetData.CreateConstIterator(); Iter; ++Iter)
+		{
+			if (Iter->AssetName.ToString() == InNewName.ToString())
+			{
+				FilenameWidget->SetError(LOCTEXT("AssetInUseError", "Asset name already in use!"));
+				bIsReportingError = true;
+				return;
+			}
+		}
+	}
+
+	FilenameWidget->SetError(FText::FromString(TEXT("")));
+	bIsReportingError = false;
+}
+
+bool SPropertyEditorNewAsset::IsCreateNewAssetEnabled() const
+{
+	return !bIsReportingError;
+}
+
+#undef LOCTEXT_NAMESPACE
diff --git a/Engine/Source/Editor/PropertyEditor/Private/UserInterface/PropertyEditor/SPropertyEditorNewAsset.h b/Engine/Source/Editor/PropertyEditor/Private/UserInterface/PropertyEditor/SPropertyEditorNewAsset.h
new file mode 100644
--- /dev/null
+++ b/Engine/Source/Editor/PropertyEditor/Private/UserInterface/PropertyEditor/SPropertyEditorNewAsset.h
@@ -0,0 +1,54 @@
+// Copyright 1998-2014 Epic Games, Inc. All Rights Reserved.
+
+#pragma once
+
+
+class SPropertyEditorNewAsset : public SCompoundWidget
+{
+public:
+	SLATE_BEGIN_ARGS(SPropertyEditorNewAsset)
+	{}
+
+	SLATE_END_ARGS()
+
+	static UObject* Create(TWeakObjectPtr<UFactory> FactoryPtr);
+
+	/** Constructs this widget with InArgs */
+	void Construct(const FArguments& InArgs, TSharedPtr<SWindow> InParentWindow, TWeakObjectPtr<UFactory> InFactoryPtr);
+
+private:
+
+ 	/** Callback when the "create asset" button is clicked. */
+ 	FReply OnCreateNewAssetClicked();
+ 
+ 	/** Callback when the selected asset path has changed. */
+ 	void OnSelectAssetPath(const FString& Path);
+ 
+ 	/** Destroys the window when the operation is cancelled. */
+ 	FReply OnCancelCreateNewAsset();
+
+	/** Callback when the user changes the filename for the Blueprint */
+ 	void OnFilenameChanged(const FText& InNewName);
+ 
+ 	/** Callback to see if creating an asset is enabled */
+ 	bool IsCreateNewAssetEnabled() const;
+
+private:
+	/** The window this widget is nested in */
+	TWeakPtr<SWindow> ParentWindowPtr;
+
+	/** The factory used to create a new asset */
+	TWeakObjectPtr<UFactory> FactoryPtr;
+
+ 	/** The selected path to create the asset */
+ 	FString AssetPath;
+ 
+ 	/** Filename textbox widget */
+ 	TSharedPtr<SEditableTextBox> FilenameWidget;
+
+ 	/** True if an error is currently being reported */
+ 	bool bIsReportingError;
+
+	/** Pointer to the new asset */
+	UObject* NewAsset;
+};
diff --git a/Engine/Source/Editor/PropertyEditor/Private/UserInterface/PropertyEditor/SPropertyMenuAssetPicker.cpp b/Engine/Source/Editor/PropertyEditor/Private/UserInterface/PropertyEditor/SPropertyMenuAssetPicker.cpp
--- a/Engine/Source/Editor/PropertyEditor/Private/UserInterface/PropertyEditor/SPropertyMenuAssetPicker.cpp
+++ b/Engine/Source/Editor/PropertyEditor/Private/UserInterface/PropertyEditor/SPropertyMenuAssetPicker.cpp
@@ -6,6 +6,8 @@
 #include "DelegateFilter.h"
 #include "Editor/ContentBrowser/Public/ContentBrowserModule.h"
 #include "PropertyEditorAssetConstants.h"
+#include "EditorStyleSet.h"
+#include "ClassIconFinder.h"
 
 #define LOCTEXT_NAMESPACE "PropertyEditor"
 
@@ -13,13 +15,33 @@ void SPropertyMenuAssetPicker::Construct( const FArguments& InArgs )
 {
 	CurrentObject = InArgs._InitialObject;
 	bAllowClear = InArgs._AllowClear;
-	AllowedClasses = *InArgs._AllowedClasses;
+	AllowedClasses = InArgs._AllowedClasses;
+	NewAssetFactories = InArgs._NewAssetFactories;
 	OnShouldFilterAsset = InArgs._OnShouldFilterAsset;
 	OnSet = InArgs._OnSet;
 	OnClose = InArgs._OnClose;
 
 	FMenuBuilder MenuBuilder(true, NULL);
 
+	if (NewAssetFactories.Num() > 0)
+	{
+		MenuBuilder.BeginSection(NAME_None, LOCTEXT("CreateNewAsset", "Create New Asset"));
+		{
+			for (UFactory* Factory : NewAssetFactories)
+			{
+				TWeakObjectPtr<UFactory> FactoryPtr(Factory);
+
+				MenuBuilder.AddMenuEntry(
+					Factory->GetDisplayName(),
+					Factory->GetToolTip(),
+					FSlateIcon(FEditorStyle::GetStyleSetName(), FClassIconFinder::FindIconNameForClass(Factory->GetSupportedClass())),
+					FUIAction(FExecuteAction::CreateSP(this, &SPropertyMenuAssetPicker::OnCreateNewAssetSelected, FactoryPtr))
+					);
+			}
+		}
+		MenuBuilder.EndSection();
+	}
+
 	MenuBuilder.BeginSection(NAME_None, LOCTEXT("CurrentAssetOperationsHeader", "Current Asset"));
 	{
 		if( CurrentObject.IsValid() )
@@ -214,4 +236,16 @@ void SPropertyMenuAssetPicker::SetValue( const FAssetData& AssetData )
 	OnSet.ExecuteIfBound(AssetData);
 }
 
+void SPropertyMenuAssetPicker::OnCreateNewAssetSelected(TWeakObjectPtr<UFactory> FactoryPtr)
+{
+	if (FactoryPtr.IsValid())
+	{
+		UObject* NewAsset = SPropertyEditorNewAsset::Create(FactoryPtr);
+		if (NewAsset != nullptr)
+		{
+			SetValue(NewAsset);
+		}
+	}
+}
+
 #undef LOCTEXT_NAMESPACE
\ No newline at end of file
diff --git a/Engine/Source/Editor/PropertyEditor/Private/UserInterface/PropertyEditor/SPropertyMenuAssetPicker.h b/Engine/Source/Editor/PropertyEditor/Private/UserInterface/PropertyEditor/SPropertyMenuAssetPicker.h
--- a/Engine/Source/Editor/PropertyEditor/Private/UserInterface/PropertyEditor/SPropertyMenuAssetPicker.h
+++ b/Engine/Source/Editor/PropertyEditor/Private/UserInterface/PropertyEditor/SPropertyMenuAssetPicker.h
@@ -8,11 +8,11 @@ public:
 	SLATE_BEGIN_ARGS( SPropertyMenuAssetPicker )
 		: _InitialObject(NULL)
 		, _AllowClear(true)
-		, _AllowedClasses(NULL)
 	{}
 		SLATE_ARGUMENT( FAssetData, InitialObject )
 		SLATE_ARGUMENT( bool, AllowClear )
-		SLATE_ARGUMENT( const TArray<const UClass*>*, AllowedClasses )
+		SLATE_ARGUMENT( TArray<const UClass*>, AllowedClasses )
+		SLATE_ARGUMENT( TArray<UFactory*>, NewAssetFactories )
 		SLATE_EVENT( FOnShouldFilterAsset, OnShouldFilterAsset )
 		SLATE_EVENT( FOnAssetSelected, OnSet )
 		SLATE_EVENT( FSimpleDelegate, OnClose )
@@ -57,6 +57,12 @@ private:
 	 */
 	void OnAssetSelected( const FAssetData& AssetData );
 
+	/**
+	 * Delegate for handling creating new assets from the menu.
+	 * @param	FactoryPtr		Factory which creates asset
+	 */
+	void OnCreateNewAssetSelected(TWeakObjectPtr<UFactory> Factory);
+
 	/** 
 	 * Set the value of the asset referenced by this property editor.
 	 * Will set the underlying property handle if there is one.
@@ -70,8 +76,12 @@ private:
 	/** Whether the asset can be 'None' in this case */
 	bool bAllowClear;
 
+	/** Array of classes to filter by */
 	TArray<const UClass*> AllowedClasses;
 
+	/** Array of factories which can create a new asset of suitable type */
+	TArray<UFactory*> NewAssetFactories;
+
 	/** Delegate for filtering valid assets */
 	FOnShouldFilterAsset OnShouldFilterAsset;
 
@@ -80,5 +90,4 @@ private:
 
 	/** Delegate for closing the containing menu */
 	FSimpleDelegate OnClose;
-
 };
\ No newline at end of file
diff --git a/Engine/Source/Editor/PropertyEditor/PropertyEditor.Build.cs b/Engine/Source/Editor/PropertyEditor/PropertyEditor.Build.cs
--- a/Engine/Source/Editor/PropertyEditor/PropertyEditor.Build.cs
+++ b/Engine/Source/Editor/PropertyEditor/PropertyEditor.Build.cs
@@ -58,6 +58,7 @@ public class PropertyEditor : ModuleRules
 
 		DynamicallyLoadedModuleNames.AddRange(
 			new string[] {
+                "AssetRegistry",
 				"ContentBrowser",
 				"Documentation",
 				"MainFrame",
diff --git a/Engine/Source/Editor/PropertyEditor/Public/PropertyCustomizationHelpers.h b/Engine/Source/Editor/PropertyEditor/Public/PropertyCustomizationHelpers.h
--- a/Engine/Source/Editor/PropertyEditor/Public/PropertyCustomizationHelpers.h
+++ b/Engine/Source/Editor/PropertyEditor/Public/PropertyCustomizationHelpers.h
@@ -29,13 +29,16 @@ namespace PropertyCustomizationHelpers
 	PROPERTYEDITOR_API TSharedRef<SWidget> MakeUseSelectedButton( FSimpleDelegate OnUseSelectedClicked, TAttribute<FText> OptionalToolTipText = FText(), TAttribute<bool> IsEnabled = true );
 	PROPERTYEDITOR_API TSharedRef<SWidget> MakeBrowseButton( FSimpleDelegate OnClearClicked, TAttribute<FText> OptionalToolTipText = FText(), TAttribute<bool> IsEnabled = true );
 	PROPERTYEDITOR_API TSharedRef<SWidget> MakeAssetPickerAnchorButton( FOnGetAllowedClasses OnGetAllowedClasses, FOnAssetSelected OnAssetSelectedFromPicker );
-	PROPERTYEDITOR_API TSharedRef<SWidget> MakeAssetPickerWithMenu( const FAssetData& InitialObject, const bool AllowClear, const TArray<const UClass*>* const AllowedClasses, FOnShouldFilterAsset OnShouldFilterAsset, FOnAssetSelected OnSet, FSimpleDelegate OnClose );
+	PROPERTYEDITOR_API TSharedRef<SWidget> MakeAssetPickerWithMenu( const FAssetData& InitialObject, const bool AllowClear, TArray<const UClass*> AllowedClasses, FOnShouldFilterAsset OnShouldFilterAsset, FOnAssetSelected OnSet, FSimpleDelegate OnClose );
 	PROPERTYEDITOR_API TSharedRef<SWidget> MakeActorPickerAnchorButton( FOnGetActorFilters OnGetActorFilters, FOnActorSelected OnActorSelectedFromPicker );
 	PROPERTYEDITOR_API TSharedRef<SWidget> MakeActorPickerWithMenu( AActor* const InitialActor, const bool AllowClear, const TSharedPtr< SceneOutliner::FOutlinerFilters >& ActorFilters, FOnActorSelected OnSet, FSimpleDelegate OnClose, FSimpleDelegate OnUseSelected );
 	PROPERTYEDITOR_API TSharedRef<SWidget> MakeInteractiveActorPicker( FOnGetAllowedClasses OnGetAllowedClasses, FOnShouldFilterActor OnShouldFilterActor, FOnActorSelected OnActorSelectedFromPicker );
 
 	/** @return the UBoolProperty edit condition property if one exists. */
 	PROPERTYEDITOR_API class UBoolProperty* GetEditConditionProperty(const class UProperty* InProperty, bool& bNegate);
+
+	/** Returns a list of factories which can be used to create new assets, based on the supplied class */
+	TArray<UFactory*> GetNewAssetFactoriesForClasses(const TArray<const UClass*>& Classes);
 }
 
 /** Delegate used to get a generic object */
