Fixed non-unity and Mac builds.

diff --git a/Engine/Source/Editor/PropertyEditor/Private/UserInterface/PropertyEditor/SPropertyEditorAsset.h b/Engine/Source/Editor/PropertyEditor/Private/UserInterface/PropertyEditor/SPropertyEditorAsset.h
--- a/Engine/Source/Editor/PropertyEditor/Private/UserInterface/PropertyEditor/SPropertyEditorAsset.h
+++ b/Engine/Source/Editor/PropertyEditor/Private/UserInterface/PropertyEditor/SPropertyEditorAsset.h
@@ -229,14 +229,8 @@ private:
 	 */
 	FReply OnAssetThumbnailDoubleClick( const FGeometry& InMyGeometry, const FPointerEvent& InMouseEvent );
 
-	/** 
-	 * Handle creating a new asset from the provided factory
-	 */
-	void SPropertyEditorAsset::OnCreateNewAsset( TWeakObjectPtr<UFactory> FactoryPtr );
-
 	/** @return True if the property can be edited */
 	bool CanEdit() const;
-
 private:
 
 	/** Main combobutton */
diff --git a/Engine/Source/Editor/PropertyEditor/Private/UserInterface/PropertyEditor/SPropertyEditorNewAsset.cpp b/Engine/Source/Editor/PropertyEditor/Private/UserInterface/PropertyEditor/SPropertyEditorNewAsset.cpp
--- a/Engine/Source/Editor/PropertyEditor/Private/UserInterface/PropertyEditor/SPropertyEditorNewAsset.cpp
+++ b/Engine/Source/Editor/PropertyEditor/Private/UserInterface/PropertyEditor/SPropertyEditorNewAsset.cpp
@@ -12,7 +12,7 @@ UObject* SPropertyEditorNewAsset::Create(TWeakObjectPtr<UFactory> FactoryPtr)
 {
 	if (!FactoryPtr.IsValid())
 	{
-		return false;
+		return nullptr;
 	}
 
 	FFormatNamedArguments Args;
diff --git a/Engine/Source/Editor/PropertyEditor/Private/UserInterface/PropertyEditor/SPropertyMenuAssetPicker.cpp b/Engine/Source/Editor/PropertyEditor/Private/UserInterface/PropertyEditor/SPropertyMenuAssetPicker.cpp
--- a/Engine/Source/Editor/PropertyEditor/Private/UserInterface/PropertyEditor/SPropertyMenuAssetPicker.cpp
+++ b/Engine/Source/Editor/PropertyEditor/Private/UserInterface/PropertyEditor/SPropertyMenuAssetPicker.cpp
@@ -2,6 +2,7 @@
 
 #include "PropertyEditorPrivatePCH.h"
 #include "SPropertyMenuAssetPicker.h"
+#include "SPropertyEditorNewAsset.h"
 #include "AssetRegistryModule.h"
 #include "DelegateFilter.h"
 #include "Editor/ContentBrowser/Public/ContentBrowserModule.h"
