Fixed hyperlinks  in details panels being offset.

diff --git a/Engine/Source/Editor/PropertyEditor/Private/SDetailNameArea.cpp b/Engine/Source/Editor/PropertyEditor/Private/SDetailNameArea.cpp
--- a/Engine/Source/Editor/PropertyEditor/Private/SDetailNameArea.cpp
+++ b/Engine/Source/Editor/PropertyEditor/Private/SDetailNameArea.cpp
@@ -107,13 +107,9 @@ TSharedRef< SWidget > SDetailNameArea::BuildObjectNameArea( const TArray< TWeakO
 			ObjectNameArea->AddSlot()
 			.AutoWidth()
 			.Padding(0,0,6,0)
+			.VAlign(VAlign_Center)
 			[
-				SNew(SBox)
-				.WidthOverride( 200.0f )
-				.VAlign(VAlign_Center)
-				[
-					ObjectNameBox
-				]
+				ObjectNameBox
 			];
 		}
 
