Fixed bug in which SObjectPropertyEntryBox widgets without an allowed class specified would allow all asset types to be created.  Now the option is explicitly provided to specify which factories to use to populate the menu with new asset types.
#ttp 349345 - Editor: Asset pickers with custom filtering can show all object types
#branch UE4
#proj Editor.DetailCustomizations, Editor.PropertyEditor

diff --git a/Engine/Source/Editor/DetailCustomizations/Private/DialogueWaveWidgets.cpp b/Engine/Source/Editor/DetailCustomizations/Private/DialogueWaveWidgets.cpp
--- a/Engine/Source/Editor/DetailCustomizations/Private/DialogueWaveWidgets.cpp
+++ b/Engine/Source/Editor/DetailCustomizations/Private/DialogueWaveWidgets.cpp
@@ -207,7 +207,14 @@ TSharedRef<SWidget> SDialogueVoicePropertyEditor::OnGetMenuContent()
 		DialogueVoice = Cast<UDialogueVoice>(Object);
 	}
 
-	return PropertyCustomizationHelpers::MakeAssetPickerWithMenu(DialogueVoice, false, AllowedClasses, OnShouldFilterAsset, FOnAssetSelected::CreateSP(this, &SDialogueVoicePropertyEditor::OnAssetSelectedFromPicker), FSimpleDelegate::CreateSP(this, &SDialogueVoicePropertyEditor::CloseMenu));
+	return PropertyCustomizationHelpers::MakeAssetPickerWithMenu(
+		DialogueVoice,
+		false,
+		AllowedClasses,
+		PropertyCustomizationHelpers::GetNewAssetFactoriesForClasses(AllowedClasses),
+		OnShouldFilterAsset,
+		FOnAssetSelected::CreateSP(this, &SDialogueVoicePropertyEditor::OnAssetSelectedFromPicker),
+		FSimpleDelegate::CreateSP(this, &SDialogueVoicePropertyEditor::CloseMenu));
 }
 
 void SDialogueVoicePropertyEditor::CloseMenu()
diff --git a/Engine/Source/Editor/DetailCustomizations/Private/SlateBrushCustomization.cpp b/Engine/Source/Editor/DetailCustomizations/Private/SlateBrushCustomization.cpp
--- a/Engine/Source/Editor/DetailCustomizations/Private/SlateBrushCustomization.cpp
+++ b/Engine/Source/Editor/DetailCustomizations/Private/SlateBrushCustomization.cpp
@@ -1058,6 +1058,10 @@ class SBrushResourceObjectBox : public SCompoundWidget
 		FSimpleDelegate OnBrushResourceChangedDelegate = FSimpleDelegate::CreateSP(this, &SBrushResourceObjectBox::OnBrushResourceChanged);
 		ResourceObjectProperty->SetOnPropertyValueChanged(OnBrushResourceChangedDelegate);
 
+		TArray<const UClass*> SupportedClasses;
+		SupportedClasses.Add(UTexture2D::StaticClass());
+		SupportedClasses.Add(UMaterialInterface::StaticClass());
+
 		ChildSlot
 		[
 			SNew(SVerticalBox)
@@ -1067,6 +1071,7 @@ class SBrushResourceObjectBox : public SCompoundWidget
 				SNew(SObjectPropertyEntryBox)
 				.PropertyHandle(InResourceObjectProperty)
 				.ThumbnailPool(StructCustomizationUtils->GetThumbnailPool())
+				.NewAssetFactories(PropertyCustomizationHelpers::GetNewAssetFactoriesForClasses(SupportedClasses))
 				.OnShouldFilterAsset(this, &SBrushResourceObjectBox::OnFilterAssetPicker)
 				.OnObjectChanged(this, &SBrushResourceObjectBox::OnAssetPicked)
 			]
diff --git a/Engine/Source/Editor/PropertyEditor/Private/PropertyCustomizationHelpers.cpp b/Engine/Source/Editor/PropertyEditor/Private/PropertyCustomizationHelpers.cpp
--- a/Engine/Source/Editor/PropertyEditor/Private/PropertyCustomizationHelpers.cpp
+++ b/Engine/Source/Editor/PropertyEditor/Private/PropertyCustomizationHelpers.cpp
@@ -203,14 +203,14 @@ namespace PropertyCustomizationHelpers
 			.OnAssetSelected( OnAssetSelectedFromPicker );
 	}
 
-	TSharedRef<SWidget> MakeAssetPickerWithMenu( const FAssetData& InitialObject, const bool AllowClear, TArray<const UClass*> AllowedClasses, FOnShouldFilterAsset OnShouldFilterAsset, FOnAssetSelected OnSet, FSimpleDelegate OnClose)
+	TSharedRef<SWidget> MakeAssetPickerWithMenu( const FAssetData& InitialObject, const bool AllowClear, const TArray<const UClass*>& AllowedClasses, const TArray<UFactory*>& NewAssetFactories, FOnShouldFilterAsset OnShouldFilterAsset, FOnAssetSelected OnSet, FSimpleDelegate OnClose)
 	{
 		return
 			SNew(SPropertyMenuAssetPicker)
 			.InitialObject(InitialObject)
 			.AllowClear(AllowClear)
 			.AllowedClasses(AllowedClasses)
-			.NewAssetFactories(GetNewAssetFactoriesForClasses(AllowedClasses))
+			.NewAssetFactories(NewAssetFactories)
 			.OnShouldFilterAsset(OnShouldFilterAsset)
 			.OnSet(OnSet)
 			.OnClose(OnClose);
@@ -349,6 +349,7 @@ void SObjectPropertyEntryBox::Construct( const FArguments& InArgs )
 			SAssignNew(PropertyEditorAsset, SPropertyEditorAsset)
 				.ObjectPath( this, &SObjectPropertyEntryBox::OnGetObjectPath )
 				.Class( InArgs._AllowedClass )
+				.NewAssetFactories( InArgs._NewAssetFactories )
 				.OnSetObject(this, &SObjectPropertyEntryBox::OnSetObject)
 				.ThumbnailPool(InArgs._ThumbnailPool)
 				.DisplayThumbnail(bDisplayThumbnail)
diff --git a/Engine/Source/Editor/PropertyEditor/Private/UserInterface/PropertyEditor/SPropertyEditorAsset.cpp b/Engine/Source/Editor/PropertyEditor/Private/UserInterface/PropertyEditor/SPropertyEditorAsset.cpp
--- a/Engine/Source/Editor/PropertyEditor/Private/UserInterface/PropertyEditor/SPropertyEditorAsset.cpp
+++ b/Engine/Source/Editor/PropertyEditor/Private/UserInterface/PropertyEditor/SPropertyEditorAsset.cpp
@@ -93,6 +93,17 @@ void SPropertyEditorAsset::Construct( const FArguments& InArgs, const TSharedPtr
 		bIsActor = ObjectClass->IsChildOf( AActor::StaticClass() );
 	}
 	OnShouldFilterAsset = InArgs._OnShouldFilterAsset;
+
+	if (InArgs._NewAssetFactories.IsSet())
+	{
+		NewAssetFactories = InArgs._NewAssetFactories.GetValue();
+	}
+	else if (ObjectClass->GetClass() != UObject::StaticClass())
+	{
+		TArray<const UClass*> AllowedClasses;
+		AllowedClasses.Add(ObjectClass);
+		NewAssetFactories = PropertyCustomizationHelpers::GetNewAssetFactoriesForClasses(AllowedClasses);
+	}
 	
 	TSharedPtr<SHorizontalBox> HorizontalBox = NULL;
 	ChildSlot
@@ -303,6 +314,7 @@ TSharedRef<SWidget> SPropertyEditorAsset::OnGetMenuContent()
 		return PropertyCustomizationHelpers::MakeAssetPickerWithMenu(Value.AssetData,
 																	 bAllowClear,
 																	 AllowedClasses,
+																	 NewAssetFactories,
 																	 OnShouldFilterAsset,
 																	 FOnAssetSelected::CreateSP(this, &SPropertyEditorAsset::OnAssetSelected),
 																	 FSimpleDelegate::CreateSP(this, &SPropertyEditorAsset::CloseComboButton));
diff --git a/Engine/Source/Editor/PropertyEditor/Private/UserInterface/PropertyEditor/SPropertyEditorAsset.h b/Engine/Source/Editor/PropertyEditor/Private/UserInterface/PropertyEditor/SPropertyEditorAsset.h
--- a/Engine/Source/Editor/PropertyEditor/Private/UserInterface/PropertyEditor/SPropertyEditorAsset.h
+++ b/Engine/Source/Editor/PropertyEditor/Private/UserInterface/PropertyEditor/SPropertyEditorAsset.h
@@ -33,6 +33,7 @@ public:
 		SLATE_ARGUMENT( FIntPoint, ThumbnailSize )
 		SLATE_ATTRIBUTE( FString, ObjectPath )
 		SLATE_ARGUMENT( UClass*, Class )
+		SLATE_ARGUMENT( TOptional<TArray<UFactory*>>, NewAssetFactories )
 		SLATE_EVENT( FOnSetObject, OnSetObject )
 		SLATE_EVENT(FOnShouldFilterAsset, OnShouldFilterAsset)
 		SLATE_NAMED_SLOT( FArguments, CustomContentSlot )
@@ -251,6 +252,9 @@ private:
 	/** The class of the object we are editing */
 	UClass* ObjectClass;
 
+	/** A list of the factories we can use to create new assets */
+	TArray<UFactory*> NewAssetFactories;
+
 	/** Whether the asset can be 'None' in this case */
 	bool bAllowClear;
 
diff --git a/Engine/Source/Editor/PropertyEditor/Public/PropertyCustomizationHelpers.h b/Engine/Source/Editor/PropertyEditor/Public/PropertyCustomizationHelpers.h
--- a/Engine/Source/Editor/PropertyEditor/Public/PropertyCustomizationHelpers.h
+++ b/Engine/Source/Editor/PropertyEditor/Public/PropertyCustomizationHelpers.h
@@ -30,7 +30,7 @@ namespace PropertyCustomizationHelpers
 	PROPERTYEDITOR_API TSharedRef<SWidget> MakeUseSelectedButton( FSimpleDelegate OnUseSelectedClicked, TAttribute<FText> OptionalToolTipText = FText(), TAttribute<bool> IsEnabled = true );
 	PROPERTYEDITOR_API TSharedRef<SWidget> MakeBrowseButton( FSimpleDelegate OnClearClicked, TAttribute<FText> OptionalToolTipText = FText(), TAttribute<bool> IsEnabled = true );
 	PROPERTYEDITOR_API TSharedRef<SWidget> MakeAssetPickerAnchorButton( FOnGetAllowedClasses OnGetAllowedClasses, FOnAssetSelected OnAssetSelectedFromPicker );
-	PROPERTYEDITOR_API TSharedRef<SWidget> MakeAssetPickerWithMenu( const FAssetData& InitialObject, const bool AllowClear, TArray<const UClass*> AllowedClasses, FOnShouldFilterAsset OnShouldFilterAsset, FOnAssetSelected OnSet, FSimpleDelegate OnClose );
+	PROPERTYEDITOR_API TSharedRef<SWidget> MakeAssetPickerWithMenu( const FAssetData& InitialObject, const bool AllowClear, const TArray<const UClass*>& AllowedClasses, const TArray<UFactory*>& NewAssetFactories, FOnShouldFilterAsset OnShouldFilterAsset, FOnAssetSelected OnSet, FSimpleDelegate OnClose );
 	PROPERTYEDITOR_API TSharedRef<SWidget> MakeActorPickerAnchorButton( FOnGetActorFilters OnGetActorFilters, FOnActorSelected OnActorSelectedFromPicker );
 	PROPERTYEDITOR_API TSharedRef<SWidget> MakeActorPickerWithMenu( AActor* const InitialActor, const bool AllowClear, const TSharedPtr< SceneOutliner::FOutlinerFilters >& ActorFilters, FOnActorSelected OnSet, FSimpleDelegate OnClose, FSimpleDelegate OnUseSelected );
 	PROPERTYEDITOR_API TSharedRef<SWidget> MakeInteractiveActorPicker( FOnGetAllowedClasses OnGetAllowedClasses, FOnShouldFilterActor OnShouldFilterActor, FOnActorSelected OnActorSelectedFromPicker );
@@ -39,7 +39,7 @@ namespace PropertyCustomizationHelpers
 	PROPERTYEDITOR_API class UBoolProperty* GetEditConditionProperty(const class UProperty* InProperty, bool& bNegate);
 
 	/** Returns a list of factories which can be used to create new assets, based on the supplied class */
-	TArray<UFactory*> GetNewAssetFactoriesForClasses(const TArray<const UClass*>& Classes);
+	PROPERTYEDITOR_API TArray<UFactory*> GetNewAssetFactoriesForClasses(const TArray<const UClass*>& Classes);
 }
 
 /** Delegate used to get a generic object */
@@ -67,6 +67,8 @@ public:
 		SLATE_ARGUMENT( TSharedPtr<FAssetThumbnailPool>, ThumbnailPool )
 		/** Classes that are allowed in the asset picker */
 		SLATE_ARGUMENT( UClass*, AllowedClass )
+		/** Optional list of factories which may be used to create new assets */
+		SLATE_ARGUMENT( TOptional<TArray<UFactory*>>, NewAssetFactories )
 		/** Called to check if an asset should be set */
 		SLATE_EVENT(FOnShouldSetAsset, OnShouldSetAsset)
 		/** Called when the object value changes */
