Fixed bad alignment of class picker buttons in the details panel

diff --git a/Engine/Source/Editor/PropertyEditor/Private/PropertyCustomizationHelpers.cpp b/Engine/Source/Editor/PropertyEditor/Private/PropertyCustomizationHelpers.cpp
--- a/Engine/Source/Editor/PropertyEditor/Private/PropertyCustomizationHelpers.cpp
+++ b/Engine/Source/Editor/PropertyEditor/Private/PropertyCustomizationHelpers.cpp
@@ -400,7 +400,6 @@ void SClassPropertyEntryBox::Construct(const FArguments& InArgs)
 	[	
 		SNew(SHorizontalBox)
 		+SHorizontalBox::Slot()
-		.AutoWidth()
 		.VAlign(VAlign_Center)
 		[
 			SAssignNew(PropertyEditorClass, SPropertyEditorClass)
