Editor - Mulitline text property fields now require Shift+Enter to enter a new line.  Pressing Enter stores the value like it does on other fields.

diff --git a/Engine/Source/Editor/PropertyEditor/Private/UserInterface/PropertyEditor/SPropertyEditorText.cpp b/Engine/Source/Editor/PropertyEditor/Private/UserInterface/PropertyEditor/SPropertyEditorText.cpp
--- a/Engine/Source/Editor/PropertyEditor/Private/UserInterface/PropertyEditor/SPropertyEditorText.cpp
+++ b/Engine/Source/Editor/PropertyEditor/Private/UserInterface/PropertyEditor/SPropertyEditorText.cpp
@@ -27,7 +27,7 @@ void SPropertyEditorText::Construct( const FArguments& InArgs, const TSharedRef<
 			.SelectAllTextOnCommit( false )
 			.IsReadOnly(this, &SPropertyEditorText::IsReadOnly)
 			.AutoWrapText(true)
-			//.ModiferKeyForNewLine(EModifierKey::Shift)
+			.ModiferKeyForNewLine(EModifierKey::Shift)
 		];
 	}
 	else
