Auto-collapse detail categories which only have advanced properties.

diff --git a/Engine/Source/Editor/PropertyEditor/Private/DetailCategoryBuilderImpl.cpp b/Engine/Source/Editor/PropertyEditor/Private/DetailCategoryBuilderImpl.cpp
--- a/Engine/Source/Editor/PropertyEditor/Private/DetailCategoryBuilderImpl.cpp
+++ b/Engine/Source/Editor/PropertyEditor/Private/DetailCategoryBuilderImpl.cpp
@@ -79,7 +79,7 @@ FDetailCategoryImpl::FDetailCategoryImpl( FName InCategoryName, TSharedRef<FDeta
 	, DetailLayoutBuilder( InDetailLayout )
 	, CategoryName( InCategoryName )
 	, SortOrder( 0 )
-	, bRestoreExpansionState( true )
+	, bRestoreExpansionState( !ContainsOnlyAdvanced() )
 	, bShouldBeInitiallyCollapsed( false )
 	, bUserShowAdvanced( false )
 	, bForceAdvanced( false )
@@ -505,7 +505,8 @@ bool FDetailCategoryImpl::ShouldBeExpanded() const
 	}
 	else if( bRestoreExpansionState )
 	{
-		bool bShouldBeExpanded = true;
+		// Collapse by default if there are no simple child nodes
+		bool bShouldBeExpanded = !ContainsOnlyAdvanced();
 		// Save the collapsed state of this section
 		GConfig->GetBool( TEXT("DetailCategories"), *CategoryPathName, bShouldBeExpanded, GEditorUserSettingsIni );
 		return bShouldBeExpanded;
@@ -708,7 +709,7 @@ void FDetailCategoryImpl::GetChildren( TArray< TSharedRef<IDetailTreeNode> >& Ou
 void FDetailCategoryImpl::FilterNode( const FDetailFilter& InFilter )
 {
 	bHasFilterStrings = InFilter.FilterStrings.Num() > 0;
-	bForceAdvanced = bHasFilterStrings || InFilter.bShowAllAdvanced == true;
+	bForceAdvanced = bHasFilterStrings || InFilter.bShowAllAdvanced == true || ContainsOnlyAdvanced();
 
 	bHasVisibleDetails = false;
 
