Prevent reset to default arrows from showing up in details panels for config properties.  Config properties cannot be reset to default

diff --git a/Engine/Source/Editor/PropertyEditor/Private/Presentation/PropertyEditor/PropertyEditor.cpp b/Engine/Source/Editor/PropertyEditor/Private/Presentation/PropertyEditor/PropertyEditor.cpp
--- a/Engine/Source/Editor/PropertyEditor/Private/Presentation/PropertyEditor/PropertyEditor.cpp
+++ b/Engine/Source/Editor/PropertyEditor/Private/Presentation/PropertyEditor/PropertyEditor.cpp
@@ -438,8 +438,9 @@ bool FPropertyEditor::IsResetToDefaultAvailable() const
 
 	// Should not be able to reset fixed size arrays
 	const bool bFixedSized = Property && Property->PropertyFlags & CPF_EditFixedSize;
+	const bool bCanResetToDefault = !(Property && Property->PropertyFlags & CPF_Config);
 
-	return Property && !PropertyHandle->IsEditConst() && !bFixedSized && PropertyHandle->DiffersFromDefault();
+	return Property && !PropertyHandle->IsEditConst() && bCanResetToDefault && !bFixedSized && PropertyHandle->DiffersFromDefault();
 }
 
 bool FPropertyEditor::ValueDiffersFromDefault() const
