Equality operators for FPropertyPath and FPropertyInfo

diff --git a/Engine/Source/Editor/PropertyEditor/Public/PropertyPath.h b/Engine/Source/Editor/PropertyEditor/Public/PropertyPath.h
--- a/Engine/Source/Editor/PropertyEditor/Public/PropertyPath.h
+++ b/Engine/Source/Editor/PropertyEditor/Public/PropertyPath.h
@@ -10,6 +10,11 @@ public:
 		return Property == Other.Property && ArrayIndex == Other.ArrayIndex;
 	}
 
+	bool operator!=( const FPropertyInfo& Other ) const 
+	{
+		return !(*this == Other);
+	}
+
 	TWeakObjectPtr< UProperty > Property;
 	int32 ArrayIndex;
 };
@@ -37,7 +42,6 @@ public:
 		return NewPath;
 	}
 
-
 	int32 GetNumProperties() const 
 	{ 
 		return Properties.Num(); 
@@ -177,4 +181,42 @@ public:
 private:
 
 	TArray< FPropertyInfo > Properties;
-};
\ No newline at end of file
+};
+
+FORCEINLINE bool operator==( const FPropertyPath& LHS, const FPropertyPath& RHS )
+{
+	if( LHS.GetNumProperties() != RHS.GetNumProperties() )
+	{
+		return false;
+	}
+
+	for( int32 i = 0, end = RHS.GetNumProperties(); i != end; ++i )
+	{
+		if( LHS.GetPropertyInfo(i) != RHS.GetPropertyInfo(i) )
+		{
+			return false;
+		}
+	}
+
+	return true;
+}
+
+FORCEINLINE bool operator!=(const FPropertyPath& LHS, const FPropertyPath& RHS)
+{
+	return !(LHS == RHS);
+}
+
+FORCEINLINE uint32 GetTypeHash(FPropertyPath const& Path)
+{
+	//lets concat the name and just hash the string
+	uint32 Ret = 0;
+	for( int32 i = 0, end = Path.GetNumProperties(); i != end; ++i )
+	{
+		if( Path.GetPropertyInfo(i).Property.IsValid() )
+		{
+			Ret = Ret ^ GetTypeHash(Path.GetPropertyInfo(i).Property->GetFName());
+		}
+	}
+	return Ret;
+}
+
