Editor - The PropertyHandles now apply property changes even if the values are the same, if we're coming out of an interactive change, like the color picker, you're dragging the color around, the last color you pick will already be the same value set, however it will have previously only been set with interactive flags, if the final one is non-interactive it will be applied again for anyone watching for the non-interactive change.

diff --git a/Engine/Source/Editor/PropertyEditor/Private/PropertyHandleImpl.cpp b/Engine/Source/Editor/PropertyEditor/Private/PropertyHandleImpl.cpp
--- a/Engine/Source/Editor/PropertyEditor/Private/PropertyHandleImpl.cpp
+++ b/Engine/Source/Editor/PropertyEditor/Private/PropertyHandleImpl.cpp
@@ -699,7 +699,10 @@ FPropertyAccess::Result FPropertyValueImpl::SetValueAsString( const FString& InV
 		FString PreviousValue;
 		GetValueAsString( PreviousValue );
 
-		if (ParentNode && (ParentNode->GetInstancesNum() == 1 || (Value.Len() && FCString::Strcmp(*PreviousValue, *Value) != 0)))
+		const bool bDidValueChange = Value.Len() && (FCString::Strcmp(*PreviousValue, *Value) != 0);
+		const bool bComingOutOfInteractiveChange = bInteractiveChangeInProgress && ( ( Flags & EPropertyValueSetFlags::InteractiveChange ) != EPropertyValueSetFlags::InteractiveChange );
+
+		if ( ParentNode && ( ParentNode->GetInstancesNum() == 1 || bComingOutOfInteractiveChange || bDidValueChange ) )
 		{
 			ImportText( Value, PropertyNodePin.Get(), Flags );
 		}
