Added better tooltip for advanced properties

diff --git a/Engine/Source/Editor/PropertyEditor/Private/DetailAdvancedDropdownNode.cpp b/Engine/Source/Editor/PropertyEditor/Private/DetailAdvancedDropdownNode.cpp
--- a/Engine/Source/Editor/PropertyEditor/Private/DetailAdvancedDropdownNode.cpp
+++ b/Engine/Source/Editor/PropertyEditor/Private/DetailAdvancedDropdownNode.cpp
@@ -63,15 +63,16 @@ public:
 					+ SVerticalBox::Slot()
 					.AutoHeight()
 					[
-						SAssignNew( ExpanderButton, SButton )
-						.ButtonStyle( FEditorStyle::Get(), "NoBorder" )
+						SAssignNew(ExpanderButton, SButton)
+						.ButtonStyle(FEditorStyle::Get(), "NoBorder")
 						.HAlign(HAlign_Center)
 						.ContentPadding(2)
-						.OnClicked( InArgs._OnClicked )
-						.IsEnabled( InArgs._IsButtonEnabled )
+						.OnClicked(InArgs._OnClicked)
+						.IsEnabled(InArgs._IsButtonEnabled)
+						.ToolTipText(this, &SAdvancedDropdownRow::GetAdvancedPulldownToolTipText )
 						[
-							SNew( SImage )
-							.Image( this, &SAdvancedDropdownRow::GetAdvancedPulldownImage )
+							SNew(SImage)
+							.Image(this, &SAdvancedDropdownRow::GetAdvancedPulldownImage)
 						]
 					]
 				];
@@ -155,6 +156,11 @@ private:
 		return bDisplayShowAdvancedMessage && !IsExpanded.Get() ? EVisibility::Visible : EVisibility::Collapsed;
 	}
 
+	FText GetAdvancedPulldownToolTipText() const
+	{
+		return IsExpanded.Get() ? NSLOCTEXT("DetailsView", "HideAdvanced", "Hide Advanced") : NSLOCTEXT("DetailsView", "ShowAdvanced", "Show Advanced");
+	}
+
 	const FSlateBrush* GetAdvancedPulldownImage() const
 	{
 		if( ExpanderButton->IsHovered() )
