Fix keying working with no animation

diff --git a/Engine/Source/Editor/PropertyEditor/Private/SDetailSingleItemRow.cpp b/Engine/Source/Editor/PropertyEditor/Private/SDetailSingleItemRow.cpp
--- a/Engine/Source/Editor/PropertyEditor/Private/SDetailSingleItemRow.cpp
+++ b/Engine/Source/Editor/PropertyEditor/Private/SDetailSingleItemRow.cpp
@@ -342,7 +342,7 @@ TSharedRef<SWidget> SDetailSingleItemRow::CreateKeyframeButton( FDetailLayoutCus
 {
 	IDetailsViewPrivate& DetailsView = InTreeNode->GetDetailsView();
 
-	TSharedPtr<IDetailKeyframeHandler> KeyframeHandler = DetailsView.GetKeyframeHandler();
+	KeyframeHandler = DetailsView.GetKeyframeHandler();
 
 	EVisibility SetKeyVisibility = EVisibility::Collapsed;
 
@@ -351,7 +351,8 @@ TSharedRef<SWidget> SDetailSingleItemRow::CreateKeyframeButton( FDetailLayoutCus
 		TSharedPtr<IPropertyHandle> Handle = PropertyEditorHelpers::GetPropertyHandle(InCustomization.GetPropertyNode().ToSharedRef(), nullptr, nullptr);
 
 		UClass* ObjectClass = InCustomization.GetPropertyNode()->FindObjectItemParent()->GetObjectBaseClass();
-		SetKeyVisibility = KeyframeHandler->IsPropertyKeyable(*ObjectClass, *Handle) ? EVisibility::Visible : EVisibility::Hidden;
+		SetKeyVisibility = KeyframeHandler.Pin()->IsPropertyKeyable(*ObjectClass, *Handle) ? EVisibility::Visible : EVisibility::Hidden;
+		
 	}
 
 	return 
@@ -361,21 +362,24 @@ TSharedRef<SWidget> SDetailSingleItemRow::CreateKeyframeButton( FDetailLayoutCus
 		.ContentPadding(0.0f)
 		.ButtonStyle(FEditorStyle::Get(), "Sequencer.AddKey.Details")
 		.Visibility(SetKeyVisibility)
+		.IsEnabled( this, &SDetailSingleItemRow::IsKeyframeButtonEnabled )
 		.ToolTipText( NSLOCTEXT("PropertyView", "AddKeyframeButton_ToolTip", "Adds a keyframe for this property to the current animation") )
 		.OnClicked( this, &SDetailSingleItemRow::OnAddKeyframeClicked );
 }
 
-FReply SDetailSingleItemRow::OnAddKeyframeClicked()
+bool SDetailSingleItemRow::IsKeyframeButtonEnabled() const
 {
-	IDetailsViewPrivate& DetailsView = OwnerTreeNode.Pin()->GetDetailsView();
-
-	TSharedPtr<IDetailKeyframeHandler> KeyframeHandler = DetailsView.GetKeyframeHandler();
-
-	check(KeyframeHandler.IsValid());
+	return KeyframeHandler.IsValid() ? KeyframeHandler.Pin()->IsPropertyKeyingEnabled() : false;
+}
 
-	TSharedPtr<IPropertyHandle> Handle = PropertyEditorHelpers::GetPropertyHandle(Customization->GetPropertyNode().ToSharedRef(), nullptr, nullptr);
+FReply SDetailSingleItemRow::OnAddKeyframeClicked()
+{	
+	if( KeyframeHandler.IsValid() )
+	{
+		TSharedPtr<IPropertyHandle> Handle = PropertyEditorHelpers::GetPropertyHandle(Customization->GetPropertyNode().ToSharedRef(), nullptr, nullptr);
 
-	KeyframeHandler->OnKeyPropertyClicked(*Handle);
+		KeyframeHandler.Pin()->OnKeyPropertyClicked(*Handle);
+	}
 
 	return FReply::Handled();
 }
diff --git a/Engine/Source/Editor/PropertyEditor/Private/SDetailSingleItemRow.h b/Engine/Source/Editor/PropertyEditor/Private/SDetailSingleItemRow.h
--- a/Engine/Source/Editor/PropertyEditor/Private/SDetailSingleItemRow.h
+++ b/Engine/Source/Editor/PropertyEditor/Private/SDetailSingleItemRow.h
@@ -31,9 +31,11 @@ private:
 	const FSlateBrush* GetBorderImage() const;
 	TSharedRef<SWidget> CreateExtensionWidget( TSharedRef<SWidget> ValueWidget, FDetailLayoutCustomization& InCustomization, TSharedRef<IDetailTreeNode> InTreeNode );
 	TSharedRef<SWidget> CreateKeyframeButton( FDetailLayoutCustomization& InCustomization, TSharedRef<IDetailTreeNode> InTreeNode );
+	bool IsKeyframeButtonEnabled() const;
 	FReply OnAddKeyframeClicked();
 	bool IsHighlighted() const;
 private:
+	TWeakPtr<IDetailKeyframeHandler> KeyframeHandler;
 	/** Customization for this widget */
 	FDetailLayoutCustomization* Customization;
 	FDetailColumnSizeData ColumnSizeData;
diff --git a/Engine/Source/Editor/PropertyEditor/Public/IDetailKeyframeHandler.h b/Engine/Source/Editor/PropertyEditor/Public/IDetailKeyframeHandler.h
--- a/Engine/Source/Editor/PropertyEditor/Public/IDetailKeyframeHandler.h
+++ b/Engine/Source/Editor/PropertyEditor/Public/IDetailKeyframeHandler.h
@@ -9,6 +9,8 @@ public:
 
 	virtual bool IsPropertyKeyable(const UClass& InObjectClass, const class IPropertyHandle& PropertyHandle) const = 0;
 
+	virtual bool IsPropertyKeyingEnabled() const = 0;
+
 	virtual void OnKeyPropertyClicked(const IPropertyHandle& KeyedPropertyHandle) = 0;
 
 };
diff --git a/Engine/Source/Editor/UMGEditor/Private/Animation/UMGDetailKeyframeHandler.cpp b/Engine/Source/Editor/UMGEditor/Private/Animation/UMGDetailKeyframeHandler.cpp
--- a/Engine/Source/Editor/UMGEditor/Private/Animation/UMGDetailKeyframeHandler.cpp
+++ b/Engine/Source/Editor/UMGEditor/Private/Animation/UMGDetailKeyframeHandler.cpp
@@ -15,6 +15,12 @@ bool FUMGDetailKeyframeHandler::IsPropertyKeyable(const UClass& InObjectClass, c
 	return BlueprintEditor.Pin()->GetSequencer()->CanKeyProperty( InObjectClass, InPropertyHandle );
 }
 
+bool FUMGDetailKeyframeHandler::IsPropertyKeyingEnabled() const
+{
+	UMovieScene* MovieScene = BlueprintEditor.Pin()->GetSequencer()->GetRootMovieScene();
+	return MovieScene != nullptr && MovieScene != UWidgetAnimation::GetNullAnimation()->MovieScene;
+}
+
 void FUMGDetailKeyframeHandler::OnKeyPropertyClicked(const IPropertyHandle& KeyedPropertyHandle)
 {
 	TArray<UObject*> Objects;
diff --git a/Engine/Source/Editor/UMGEditor/Private/Animation/UMGDetailKeyframeHandler.h b/Engine/Source/Editor/UMGEditor/Private/Animation/UMGDetailKeyframeHandler.h
--- a/Engine/Source/Editor/UMGEditor/Private/Animation/UMGDetailKeyframeHandler.h
+++ b/Engine/Source/Editor/UMGEditor/Private/Animation/UMGDetailKeyframeHandler.h
@@ -11,6 +11,8 @@ public:
 
 	virtual bool IsPropertyKeyable(const UClass& InObjectClass, const class IPropertyHandle& PropertyHandle) const override;
 
+	virtual bool IsPropertyKeyingEnabled() const override;
+
 	virtual void OnKeyPropertyClicked(const IPropertyHandle& KeyedPropertyHandle) override;
 
 private:
