UE-4760 FPropertyValueImpl::GetValueData: check if Property is null, refactor.

Btw. Property is never really used.

diff --git a/Engine/Source/Editor/PropertyEditor/Private/PropertyHandleImpl.cpp b/Engine/Source/Editor/PropertyEditor/Private/PropertyHandleImpl.cpp
--- a/Engine/Source/Editor/PropertyEditor/Private/PropertyHandleImpl.cpp
+++ b/Engine/Source/Editor/PropertyEditor/Private/PropertyHandleImpl.cpp
@@ -147,7 +147,7 @@ FPropertyAccess::Result FPropertyValueImpl::GetPropertyValueText( FText& OutText
 
 FPropertyAccess::Result FPropertyValueImpl::GetValueData( void*& OutAddress ) const
 {
-	FPropertyAccess::Result Res = FPropertyAccess::Success;
+	FPropertyAccess::Result Res = FPropertyAccess::Fail;
 	OutAddress = NULL;
 	TSharedPtr<FPropertyNode> PropertyNodePin = PropertyNode.Pin();
 	if( PropertyNodePin.IsValid() )
@@ -159,33 +159,19 @@ FPropertyAccess::Result FPropertyValueImpl::GetValueData( void*& OutAddress ) co
 		if( (ReadAddresses.Num() > 0 && bAllValuesTheSame) || ReadAddresses.Num() == 1 ) 
 		{
 			ValueAddress = ReadAddresses.GetAddress(0);
-
-			if( ValueAddress )
+			const UProperty* Property = PropertyNodePin->GetProperty();
+			if (ValueAddress && Property)
 			{
-				int32 Index = 0;
-				UProperty* Property = PropertyNodePin->GetProperty();
+				const int32 Index = 0;
 				OutAddress = ValueAddress + Index * Property->ElementSize;
-
-				if( OutAddress == NULL )
-				{
-					Res = FPropertyAccess::Fail;
-				}
-			}
-			else
-			{
-				Res = FPropertyAccess::Fail;
+				Res = FPropertyAccess::Success;
 			}
 		}
-		else
+		else if (ReadAddresses.Num() > 1)
 		{
-			Res = ReadAddresses.Num() > 1 ? FPropertyAccess::MultipleValues : FPropertyAccess::Fail;
+			Res = FPropertyAccess::MultipleValues;
 		}
 	}
-	else
-	{
-		// The property node is not valid.  Probably means this value was stored somewhere and selection changed causing the node to be destroyed
-		Res = FPropertyAccess::Fail;
-	}
 
 	return Res;
 }
