UE-4800 CRASH: UDS: Crash when deleting default element of array struct variable

Problem: FPropertyItemValueDataTrackerSlate keeps raw pointer to invalid memory.
Fix: FPropertyItemValueDataTrackerSlate is reconstructed.
It's a temporary fix. FPropertyItemValueDataTrackerSlate should never use raw pointer.

diff --git a/Engine/Source/Editor/PropertyEditor/Private/PropertyNode.cpp b/Engine/Source/Editor/PropertyEditor/Private/PropertyNode.cpp
--- a/Engine/Source/Editor/PropertyEditor/Private/PropertyNode.cpp
+++ b/Engine/Source/Editor/PropertyEditor/Private/PropertyNode.cpp
@@ -732,27 +732,35 @@ public:
 		uint8*		ValueAddress;
 	};
 
-	/**
-	 * Constructor
-	 *
-	 * @param	InPropItem		the property window item this struct will hold values for
-	 * @param	InOwnerObject	the object which contains the property value
-	 */
-	FPropertyItemValueDataTrackerSlate( FPropertyNode* InPropertyNode, UObject* InOwnerObject )
-		: OwnerObject( InOwnerObject )
-		, PropertyNode(InPropertyNode)
-		, bHasDefaultValue(false)
+	void Reset(FPropertyNode* InPropertyNode, UObject* InOwnerObject)
+	{
+		OwnerObject = InOwnerObject;
+		PropertyNode = InPropertyNode;
+		bHasDefaultValue = false;
+		InnerInitialize();
+	}
+
+	void InnerInitialize()
 	{
-		PropertyValueRoot.OwnerObject = InOwnerObject;
+		{
+			PropertyValueRoot.OwnerObject = NULL;
+			PropertyDefaultValueRoot.OwnerObject = NULL;
+			PropertyValueAddress = NULL;
+			PropertyValueBaseAddress = NULL;
+			PropertyDefaultBaseAddress = NULL;
+			PropertyDefaultAddress = NULL;
+		}
+
+		PropertyValueRoot.OwnerObject = OwnerObject.Get();
 		check(PropertyNode);
 		UProperty* Property = PropertyNode->GetProperty();
 		check(Property);
 		check(PropertyValueRoot.OwnerObject);
 
-		FPropertyNode* ParentNode		= PropertyNode->GetParentNode();
+		FPropertyNode* ParentNode = PropertyNode->GetParentNode();
 
 		// if the object specified is a class object, transfer to the CDO instead
-		if ( Cast<UClass>(PropertyValueRoot.OwnerObject) != NULL )
+		if (Cast<UClass>(PropertyValueRoot.OwnerObject) != NULL)
 		{
 			PropertyValueRoot.OwnerObject = Cast<UClass>(PropertyValueRoot.OwnerObject)->GetDefaultObject();
 		}
@@ -769,13 +777,13 @@ public:
 			PropertyValueAddress = PropertyNode->GetValueAddress(PropertyValueRoot.ValueAddress);
 		}
 
-		if( IsValidTracker() )
+		if (IsValidTracker())
 		{
-			 bHasDefaultValue = Private_HasDefaultValue();
+			bHasDefaultValue = Private_HasDefaultValue();
 			// calculate the values for the default object
-			if ( bHasDefaultValue )
+			if (bHasDefaultValue)
 			{
-				PropertyDefaultValueRoot.OwnerObject = PropertyValueRoot.OwnerObject->GetArchetype();
+				PropertyDefaultValueRoot.OwnerObject = PropertyValueRoot.OwnerObject ? PropertyValueRoot.OwnerObject->GetArchetype() : NULL;
 				PropertyDefaultBaseAddress = OuterArrayProp == NULL
 					? PropertyNode->GetValueBaseAddress(PropertyDefaultValueRoot.ValueAddress)
 					: ParentNode->GetValueBaseAddress(PropertyDefaultValueRoot.ValueAddress);
@@ -784,7 +792,7 @@ public:
 				//////////////////////////
 				// If this is an array property, we must take special measures; PropertyDefaultBaseAddress points to an FScriptArray*, while
 				// PropertyDefaultAddress points to the FScriptArray's Data pointer.
-				if ( ArrayProp != NULL )
+				if (ArrayProp != NULL)
 				{
 					PropertyValueAddress = PropertyValueBaseAddress;
 					PropertyDefaultAddress = PropertyDefaultBaseAddress;
@@ -793,6 +801,20 @@ public:
 		}
 	}
 
+	/**
+	 * Constructor
+	 *
+	 * @param	InPropItem		the property window item this struct will hold values for
+	 * @param	InOwnerObject	the object which contains the property value
+	 */
+	FPropertyItemValueDataTrackerSlate( FPropertyNode* InPropertyNode, UObject* InOwnerObject )
+		: OwnerObject( InOwnerObject )
+		, PropertyNode(InPropertyNode)
+		, bHasDefaultValue(false)
+	{
+		InnerInitialize();
+	}
+
 	/**
 	 * @return Whether or not this tracker has a valid address to a property and object
 	 */
@@ -2021,6 +2043,10 @@ TSharedPtr< FPropertyItemValueDataTrackerSlate > FPropertyNode::GetValueTracker(
 		{
 			ValueTracker = MakeShareable( new FPropertyItemValueDataTrackerSlate( this, Object ) );
 		}
+		else
+		{
+			ValueTracker->Reset(this, Object);
+		}
 		RetVal = ValueTracker;
 
 	}
