Fill UDS in arrays with default values [UE-4800]

diff --git a/Engine/Source/Editor/PropertyEditor/Private/PropertyHandleImpl.cpp b/Engine/Source/Editor/PropertyEditor/Private/PropertyHandleImpl.cpp
--- a/Engine/Source/Editor/PropertyEditor/Private/PropertyHandleImpl.cpp
+++ b/Engine/Source/Editor/PropertyEditor/Private/PropertyHandleImpl.cpp
@@ -1086,14 +1086,7 @@ void FPropertyValueImpl::AddChild()
 
 					FScriptArrayHelper	ArrayHelper(Array,Addr);
 					const int32 ArrayIndex = ArrayHelper.AddValue();
-
-					if (const UStructProperty* StructProperty = Cast<const UStructProperty>(Array->Inner))
-					{
-						if(!FStructureEditorUtils::Fill_MakeStructureDefaultValue(Cast<const UUserDefinedStruct>(StructProperty->Struct), ArrayHelper.GetRawPtr(ArrayIndex)))
-						{
-							UE_LOG(LogPropertyNode, Warning, TEXT("MakeStructureDefaultValue parsing error. Property: %s "), *StructProperty->GetName());
-						}
-					}
+					FPropertyNode::AdditionalInitializationUDS(Array->Inner, ArrayHelper.GetRawPtr(ArrayIndex));
 					
 					ArrayIndicesPerObject[i].Add(NodeProperty->GetName(), ArrayIndex);
 				}
@@ -1239,7 +1232,9 @@ void FPropertyValueImpl::InsertChild( TSharedPtr<FPropertyNode> ChildNodeToInser
 			}
 		}
 
-		ArrayHelper.InsertValues(Index, 1 );	
+		ArrayHelper.InsertValues(Index, 1 );
+
+		FPropertyNode::AdditionalInitializationUDS(ArrayProperty->Inner, ArrayHelper.GetRawPtr(Index));
 
 		//set up indices for the coming events
 		TArray< TMap<FString,int32> > ArrayIndicesPerObject;
diff --git a/Engine/Source/Editor/PropertyEditor/Private/PropertyNode.cpp b/Engine/Source/Editor/PropertyEditor/Private/PropertyNode.cpp
--- a/Engine/Source/Editor/PropertyEditor/Private/PropertyNode.cpp
+++ b/Engine/Source/Editor/PropertyEditor/Private/PropertyNode.cpp
@@ -742,7 +742,7 @@ public:
 
 	void InnerInitialize()
 	{
-		{
+	{
 			PropertyValueRoot.OwnerObject = NULL;
 			PropertyDefaultValueRoot.OwnerObject = NULL;
 			PropertyValueAddress = NULL;
@@ -757,10 +757,10 @@ public:
 		check(Property);
 		check(PropertyValueRoot.OwnerObject);
 
-		FPropertyNode* ParentNode = PropertyNode->GetParentNode();
+		FPropertyNode* ParentNode		= PropertyNode->GetParentNode();
 
 		// if the object specified is a class object, transfer to the CDO instead
-		if (Cast<UClass>(PropertyValueRoot.OwnerObject) != NULL)
+		if ( Cast<UClass>(PropertyValueRoot.OwnerObject) != NULL )
 		{
 			PropertyValueRoot.OwnerObject = Cast<UClass>(PropertyValueRoot.OwnerObject)->GetDefaultObject();
 		}
@@ -777,11 +777,11 @@ public:
 			PropertyValueAddress = PropertyNode->GetValueAddress(PropertyValueRoot.ValueAddress);
 		}
 
-		if (IsValidTracker())
+		if( IsValidTracker() )
 		{
-			bHasDefaultValue = Private_HasDefaultValue();
+			 bHasDefaultValue = Private_HasDefaultValue();
 			// calculate the values for the default object
-			if (bHasDefaultValue)
+			if ( bHasDefaultValue )
 			{
 				PropertyDefaultValueRoot.OwnerObject = PropertyValueRoot.OwnerObject ? PropertyValueRoot.OwnerObject->GetArchetype() : NULL;
 				PropertyDefaultBaseAddress = OuterArrayProp == NULL
@@ -792,7 +792,7 @@ public:
 				//////////////////////////
 				// If this is an array property, we must take special measures; PropertyDefaultBaseAddress points to an FScriptArray*, while
 				// PropertyDefaultAddress points to the FScriptArray's Data pointer.
-				if (ArrayProp != NULL)
+				if ( ArrayProp != NULL )
 				{
 					PropertyValueAddress = PropertyValueBaseAddress;
 					PropertyDefaultAddress = PropertyDefaultBaseAddress;
@@ -2143,6 +2143,17 @@ bool FPropertyNode::IsFilterAcceptable(const TArray<FString>& InAcceptableNames,
 	return bCompleteMatchFound;
 }
 
+void FPropertyNode::AdditionalInitializationUDS(UProperty* Property, uint8* RawPtr)
+{
+	if (const UStructProperty* StructProperty = Cast<const UStructProperty>(Property))
+	{
+		if (!FStructureEditorUtils::Fill_MakeStructureDefaultValue(Cast<const UUserDefinedStruct>(StructProperty->Struct), RawPtr))
+		{
+			UE_LOG(LogPropertyNode, Warning, TEXT("MakeStructureDefaultValue parsing error. Property: %s "), *StructProperty->GetPathName());
+		}
+	}
+}
+
 void FPropertyNode::PropagateArrayPropertyChange( UObject* ModifiedObject, const FString& OriginalArrayContent, EPropertyArrayChangeType::Type ChangeType, int32 Index )
 {
 	UProperty* NodeProperty = GetProperty();
@@ -2226,16 +2237,18 @@ void FPropertyNode::PropagateArrayPropertyChange( UObject* ModifiedObject, const
 			// Check if the original value was the default value and change it only then
 			if (bIsDefault)
 			{
+				int32 ElementToInitialize = -1;
 				switch (ChangeType)
 				{
 					case EPropertyArrayChangeType::Add:
-						ArrayHelper.AddValue();
+						ElementToInitialize = ArrayHelper.AddValue();
 						break;
 					case EPropertyArrayChangeType::Clear:
 						ArrayHelper.EmptyValues();
 						break;
 					case EPropertyArrayChangeType::Insert:
 						ArrayHelper.InsertValues(ArrayIndex, 1);
+						ElementToInitialize = ArrayIndex;
 						break;
 					case EPropertyArrayChangeType::Delete:
 						ArrayHelper.RemoveValues(ArrayIndex, 1);
@@ -2247,6 +2260,10 @@ void FPropertyNode::PropagateArrayPropertyChange( UObject* ModifiedObject, const
 						Object->InstanceSubobjectTemplates();
 						break;
 				}
+				if (ElementToInitialize >= 0)
+				{
+					AdditionalInitializationUDS(ArrayProperty->Inner, ArrayHelper.GetRawPtr(ElementToInitialize));
+				}
 			}
 		}
 
diff --git a/Engine/Source/Editor/PropertyEditor/Private/PropertyNode.h b/Engine/Source/Editor/PropertyEditor/Private/PropertyNode.h
--- a/Engine/Source/Editor/PropertyEditor/Private/PropertyNode.h
+++ b/Engine/Source/Editor/PropertyEditor/Private/PropertyNode.h
@@ -466,6 +466,8 @@ public:
 	void PropagateArrayPropertyChange( UObject* ModifiedObject, const FString& OriginalArrayContent,
 									   EPropertyArrayChangeType::Type ChangeType, int32 Index);
 
+	static void AdditionalInitializationUDS(UProperty* Property, uint8* RawPtr);
+
 	/** Broadcasts when a property value changes */
 	DECLARE_EVENT( FPropertyNode, FPropertyValueChangedEvent );
 	FPropertyValueChangedEvent& OnPropertyValueChanged() { return PropertyValueChangedEvent; }
