Cis fix, missing hader

diff --git a/Engine/Source/Editor/PropertyEditor/Private/PropertyHandleImpl.cpp b/Engine/Source/Editor/PropertyEditor/Private/PropertyHandleImpl.cpp
--- a/Engine/Source/Editor/PropertyEditor/Private/PropertyHandleImpl.cpp
+++ b/Engine/Source/Editor/PropertyEditor/Private/PropertyHandleImpl.cpp
@@ -10,7 +10,6 @@
 #include "Editor/UnrealEd/Public/Kismet2/KismetEditorUtilities.h"
 #include "IPropertyUtilities.h"
 #include "PropertyEditor.h"
-#include "Editor/UnrealEd/Public/Kismet2/StructureEditorUtils.h"
 
 DEFINE_LOG_CATEGORY_STATIC(LogPropertyView, Log, All);
 
diff --git a/Engine/Source/Editor/PropertyEditor/Private/PropertyNode.cpp b/Engine/Source/Editor/PropertyEditor/Private/PropertyNode.cpp
--- a/Engine/Source/Editor/PropertyEditor/Private/PropertyNode.cpp
+++ b/Engine/Source/Editor/PropertyEditor/Private/PropertyNode.cpp
@@ -6,6 +6,7 @@
 #include "CategoryPropertyNode.h"
 #include "ScopedTransaction.h"
 #include "PropertyRestriction.h"
+#include "Editor/UnrealEd/Public/Kismet2/StructureEditorUtils.h"
 
 FPropertySettings& FPropertySettings::Get()
 {
