Editor: Improve warning message for properties with bogus clamp ranges to include the property name

diff --git a/Engine/Source/Editor/PropertyEditor/Private/UserInterface/PropertyEditor/SPropertyEditorNumeric.h b/Engine/Source/Editor/PropertyEditor/Private/UserInterface/PropertyEditor/SPropertyEditorNumeric.h
--- a/Engine/Source/Editor/PropertyEditor/Private/UserInterface/PropertyEditor/SPropertyEditorNumeric.h
+++ b/Engine/Source/Editor/PropertyEditor/Private/UserInterface/PropertyEditor/SPropertyEditorNumeric.h
@@ -69,7 +69,7 @@ public:
 
 		if (ClampMin >= ClampMax && (ClampMinString.Len() || ClampMaxString.Len()))
 		{
-			UE_LOG(LogPropertyNode, Warning, TEXT("Clamp Min (%s) >= Clamp Max (%s) for Ranged Numeric"), *ClampMinString, *ClampMaxString);
+			UE_LOG(LogPropertyNode, Warning, TEXT("Clamp Min (%s) >= Clamp Max (%s) for Ranged Numeric property %s"), *ClampMinString, *ClampMaxString, *Property->GetPathName());
 		}
 
 		const NumericType ActualUIMin = FMath::Max(UIMin, ClampMin);
@@ -82,7 +82,7 @@ public:
 
 		if (ActualUIMin >= ActualUIMax && (MetaUIMinString.Len() || MetaUIMaxString.Len()))
 		{
-			UE_LOG(LogPropertyNode, Warning, TEXT("UI Min (%s) >= UI Max (%s) for Ranged Numeric"), *UIMinString, *UIMaxString);
+			UE_LOG(LogPropertyNode, Warning, TEXT("UI Min (%s) >= UI Max (%s) for Ranged Numeric property %s"), *UIMinString, *UIMaxString, *Property->GetPathName());
 		}
 
 		FObjectPropertyNode* ObjectPropertyNode = PropertyNode->FindObjectItemParent();
@@ -236,7 +236,7 @@ private:
 
 
 	/**
-	 * Called when the slider stops moving.  We end the previously created transation
+	 * Called when the slider stops moving.  We end the previously created transaction
 	 */
 	void OnEndSliderMovement( NumericType NewValue )
 	{
