Editor: Prevent ranged numeric warnings if UIMin is used without UIMax (it is a valid case that constrains the slider widget on one end only)

diff --git a/Engine/Source/Editor/PropertyEditor/Private/UserInterface/PropertyEditor/SPropertyEditorNumeric.h b/Engine/Source/Editor/PropertyEditor/Private/UserInterface/PropertyEditor/SPropertyEditorNumeric.h
--- a/Engine/Source/Editor/PropertyEditor/Private/UserInterface/PropertyEditor/SPropertyEditorNumeric.h
+++ b/Engine/Source/Editor/PropertyEditor/Private/UserInterface/PropertyEditor/SPropertyEditorNumeric.h
@@ -80,7 +80,7 @@ public:
 		TOptional<NumericType> SliderMinValue = (UIMinString.Len()) ? ActualUIMin : TOptional<NumericType>();
 		TOptional<NumericType> SliderMaxValue = (UIMaxString.Len()) ? ActualUIMax : TOptional<NumericType>();
 
-		if (ActualUIMin >= ActualUIMax && (MetaUIMinString.Len() || MetaUIMaxString.Len()))
+		if ((ActualUIMin >= ActualUIMax) && (SliderMinValue.IsSet() && SliderMaxValue.IsSet()))
 		{
 			UE_LOG(LogPropertyNode, Warning, TEXT("UI Min (%s) >= UI Max (%s) for Ranged Numeric property %s"), *UIMinString, *UIMaxString, *Property->GetPathName());
 		}
