Fixed the details view scrollbar being visible even when nothing was selected

diff --git a/Engine/Source/Editor/PropertyEditor/Private/SDetailsView.cpp b/Engine/Source/Editor/PropertyEditor/Private/SDetailsView.cpp
--- a/Engine/Source/Editor/PropertyEditor/Private/SDetailsView.cpp
+++ b/Engine/Source/Editor/PropertyEditor/Private/SDetailsView.cpp
@@ -47,71 +47,74 @@ void SDetailsView::Construct(const FArguments& InArgs)
 	ColumnSizeData.RightColumnWidth = TAttribute<float>( this, &SDetailsView::OnGetRightColumnWidth );
 	ColumnSizeData.OnWidthChanged = SSplitter::FOnSlotResized::CreateSP( this, &SDetailsView::OnSetColumnWidth );
 
-	TSharedRef<SScrollBar> ExternalScrollbar = 
-		SNew(SScrollBar)
-		.AlwaysShowScrollbar( true );
-
-		FMenuBuilder DetailViewOptions( true, NULL );
-
-		FUIAction ShowOnlyModifiedAction( 
-			FExecuteAction::CreateSP( this, &SDetailsView::OnShowOnlyModifiedClicked ),
-			FCanExecuteAction(),
-			FIsActionChecked::CreateSP( this, &SDetailsView::IsShowOnlyModifiedChecked )
-		);
-
-		if (DetailsViewArgs.bShowModifiedPropertiesOption)
-		{
-			DetailViewOptions.AddMenuEntry( 
-				LOCTEXT("ShowOnlyModified", "Show Only Modified Properties"),
-				LOCTEXT("ShowOnlyModified_ToolTip", "Displays only properties which have been changed from their default"),
-				FSlateIcon(),
-				ShowOnlyModifiedAction,
-				NAME_None,
-				EUserInterfaceActionType::ToggleButton 
-			);
-		}
-
-		if( DetailsViewArgs.bShowDifferingPropertiesOption )
-		{
-			DetailViewOptions.AddMenuEntry(
-				LOCTEXT("ShowOnlyDiffering", "Show Only Differing Properties"),
-				LOCTEXT("ShowOnlyDiffering_ToolTip", "Displays only properties in this instance which have been changed or added from the instance being compared"),
-				FSlateIcon(),
-				FUIAction(
-					FExecuteAction::CreateSP(this, &SDetailsView::OnShowOnlyDifferingClicked),
-					FCanExecuteAction(),
-					FIsActionChecked::CreateSP(this, &SDetailsView::IsShowOnlyDifferingChecked)
-				),
-				NAME_None,
-				EUserInterfaceActionType::ToggleButton
-			);
-		}
-
-		FUIAction ShowAllAdvancedAction( 
-			FExecuteAction::CreateSP( this, &SDetailsView::OnShowAllAdvancedClicked ),
-			FCanExecuteAction(),
-			FIsActionChecked::CreateSP( this, &SDetailsView::IsShowAllAdvancedChecked )
-		);
-
-		DetailViewOptions.AddMenuEntry(
-			LOCTEXT("ShowAllAdvanced", "Show All Advanced Details"),
-			LOCTEXT("ShowAllAdvanced_ToolTip", "Shows all advanced detail sections in each category"),
+	// We want the scrollbar to always be visible when objects are selected, but not when there is no selection - however:
+	//  - We can't use AlwaysShowScrollbar for this, as this will also show the scrollbar when nothing is selected
+	//  - We can't use the Visibility construction parameter, as it gets translated into user visibility and can hide the scrollbar even when objects are selected
+	// We instead have to explicitly set the visibility after the scrollbar has been constructed to get the exact behavior we want
+	TSharedRef<SScrollBar> ExternalScrollbar = SNew(SScrollBar);
+	ExternalScrollbar->SetVisibility( TAttribute<EVisibility>( this, &SDetailsView::GetScrollBarVisibility ) );
+
+	FMenuBuilder DetailViewOptions( true, NULL );
+
+	FUIAction ShowOnlyModifiedAction( 
+		FExecuteAction::CreateSP( this, &SDetailsView::OnShowOnlyModifiedClicked ),
+		FCanExecuteAction(),
+		FIsActionChecked::CreateSP( this, &SDetailsView::IsShowOnlyModifiedChecked )
+	);
+
+	if (DetailsViewArgs.bShowModifiedPropertiesOption)
+	{
+		DetailViewOptions.AddMenuEntry( 
+			LOCTEXT("ShowOnlyModified", "Show Only Modified Properties"),
+			LOCTEXT("ShowOnlyModified_ToolTip", "Displays only properties which have been changed from their default"),
 			FSlateIcon(),
-			ShowAllAdvancedAction,
+			ShowOnlyModifiedAction,
 			NAME_None,
 			EUserInterfaceActionType::ToggleButton 
-			);
+		);
+	}
 
+	if( DetailsViewArgs.bShowDifferingPropertiesOption )
+	{
 		DetailViewOptions.AddMenuEntry(
-			LOCTEXT("CollapseAll", "Collapse All Categories"),
-			LOCTEXT("CollapseAll_ToolTip", "Collapses all root level categories"),
-			FSlateIcon(),
-			FUIAction(FExecuteAction::CreateSP(this, &SDetailsView::SetRootExpansionStates, /*bExpanded=*/false, /*bRecurse=*/false )));
-		DetailViewOptions.AddMenuEntry(
-			LOCTEXT("ExpandAll", "Expand All Categories"),
-			LOCTEXT("ExpandAll_ToolTip", "Expands all root level categories"),
+			LOCTEXT("ShowOnlyDiffering", "Show Only Differing Properties"),
+			LOCTEXT("ShowOnlyDiffering_ToolTip", "Displays only properties in this instance which have been changed or added from the instance being compared"),
 			FSlateIcon(),
-			FUIAction(FExecuteAction::CreateSP(this, &SDetailsView::SetRootExpansionStates, /*bExpanded=*/true, /*bRecurse=*/false )));
+			FUIAction(
+				FExecuteAction::CreateSP(this, &SDetailsView::OnShowOnlyDifferingClicked),
+				FCanExecuteAction(),
+				FIsActionChecked::CreateSP(this, &SDetailsView::IsShowOnlyDifferingChecked)
+			),
+			NAME_None,
+			EUserInterfaceActionType::ToggleButton
+		);
+	}
+
+	FUIAction ShowAllAdvancedAction( 
+		FExecuteAction::CreateSP( this, &SDetailsView::OnShowAllAdvancedClicked ),
+		FCanExecuteAction(),
+		FIsActionChecked::CreateSP( this, &SDetailsView::IsShowAllAdvancedChecked )
+	);
+
+	DetailViewOptions.AddMenuEntry(
+		LOCTEXT("ShowAllAdvanced", "Show All Advanced Details"),
+		LOCTEXT("ShowAllAdvanced_ToolTip", "Shows all advanced detail sections in each category"),
+		FSlateIcon(),
+		ShowAllAdvancedAction,
+		NAME_None,
+		EUserInterfaceActionType::ToggleButton 
+		);
+
+	DetailViewOptions.AddMenuEntry(
+		LOCTEXT("CollapseAll", "Collapse All Categories"),
+		LOCTEXT("CollapseAll_ToolTip", "Collapses all root level categories"),
+		FSlateIcon(),
+		FUIAction(FExecuteAction::CreateSP(this, &SDetailsView::SetRootExpansionStates, /*bExpanded=*/false, /*bRecurse=*/false )));
+	DetailViewOptions.AddMenuEntry(
+		LOCTEXT("ExpandAll", "Expand All Categories"),
+		LOCTEXT("ExpandAll_ToolTip", "Expands all root level categories"),
+		FSlateIcon(),
+		FUIAction(FExecuteAction::CreateSP(this, &SDetailsView::SetRootExpansionStates, /*bExpanded=*/true, /*bRecurse=*/false )));
 
 	TSharedRef<SHorizontalBox> FilterRow = SNew( SHorizontalBox )
 		.Visibility( this, &SDetailsView::GetFilterBoxVisibility )
@@ -239,6 +242,12 @@ EVisibility SDetailsView::GetActorNameAreaVisibility() const
 	return bVisible ? EVisibility::Visible : EVisibility::Collapsed; 
 }
 
+EVisibility SDetailsView::GetScrollBarVisibility() const
+{
+	const bool bHasObjects = RootPropertyNode.IsValid() && RootPropertyNode->GetObjectBaseClass() && RootPropertyNode->GetNumObjects();
+	return bHasObjects ? EVisibility::Visible : EVisibility::Collapsed; 
+}
+
 void SDetailsView::ForceRefresh()
 {
 	TArray< TWeakObjectPtr< UObject > > NewObjectList;
diff --git a/Engine/Source/Editor/PropertyEditor/Private/SDetailsView.h b/Engine/Source/Editor/PropertyEditor/Private/SDetailsView.h
--- a/Engine/Source/Editor/PropertyEditor/Private/SDetailsView.h
+++ b/Engine/Source/Editor/PropertyEditor/Private/SDetailsView.h
@@ -134,6 +134,9 @@ private:
 	/** Called to get the visibility of the actor name area */
 	EVisibility GetActorNameAreaVisibility() const;
 
+	/** Called to get the visibility of the scrollbar */
+	EVisibility GetScrollBarVisibility() const;
+
 	/** Returns the name of the image used for the icon on the locked button */
 	const FSlateBrush* OnGetLockButtonImageResource() const;
 	/**
