Fixing || and && in the same scope.

- Changed the behavior in one case since the intent seems to have been different.

diff --git a/Engine/Source/Editor/PropertyEditor/Private/PropertyHandleImpl.cpp b/Engine/Source/Editor/PropertyEditor/Private/PropertyHandleImpl.cpp
--- a/Engine/Source/Editor/PropertyEditor/Private/PropertyHandleImpl.cpp
+++ b/Engine/Source/Editor/PropertyEditor/Private/PropertyHandleImpl.cpp
@@ -52,7 +52,7 @@ FPropertyAccess::Result FPropertyValueImpl::GetPropertyValueString( FString& Out
 	FReadAddressList ReadAddresses;
 	bool bAllValuesTheSame = InPropertyNode->GetReadAddress( !!InPropertyNode->HasNodeFlags(EPropertyNodeFlags::SingleSelectOnly), ReadAddresses, false, true );
 
-	if( ReadAddresses.Num() > 0 && bAllValuesTheSame || ReadAddresses.Num() == 1 ) 
+	if( (ReadAddresses.Num() > 0 && bAllValuesTheSame) || ReadAddresses.Num() == 1 ) 
 	{
 		ValueAddress = ReadAddresses.GetAddress(0);
 
@@ -231,7 +231,7 @@ bool FPropertyValueImpl::SendTextToObjectProperty( const FString& Text, EPropert
 
 		// If more than one object is selected, an empty field indicates their values for this property differ.
 		// Don't send it to the objects value in this case (if we did, they would all get set to None which isn't good).
-		if (!ParentNode || ParentNode->GetInstancesNum() > 1 && !Text.Len())
+		if ((!ParentNode || ParentNode->GetInstancesNum() > 1) && !Text.Len())
 		{
 			return false;
 		}
