Back out change to make the details panel object name field smaller

diff --git a/Engine/Source/Editor/PropertyEditor/Private/SDetailNameArea.cpp b/Engine/Source/Editor/PropertyEditor/Private/SDetailNameArea.cpp
--- a/Engine/Source/Editor/PropertyEditor/Private/SDetailNameArea.cpp
+++ b/Engine/Source/Editor/PropertyEditor/Private/SDetailNameArea.cpp
@@ -105,12 +105,16 @@ TSharedRef< SWidget > SDetailNameArea::BuildObjectNameArea( const TArray< TWeakO
 			TSharedRef<IObjectNameEditableTextBox> ObjectNameBox = EdWidgetsModule.CreateObjectNameEditableTextBox(SelectedObjects);
 
 			ObjectNameArea->AddSlot()
-			.AutoWidth()
-			.Padding(0,0,6,0)
-			.VAlign(VAlign_Center)
-			[
-				ObjectNameBox
-			];
+				.AutoWidth()
+				.Padding(0, 0, 6, 0)
+				[
+					SNew(SBox)
+					.WidthOverride(200.0f)
+					.VAlign(VAlign_Center)
+					[
+						ObjectNameBox
+					]
+				];
 		}
 
 		const TWeakObjectPtr< UObject > ObjectWeakPtr = SelectedObjects.Num() == 1 ? SelectedObjects[0] : NULL;
