Details panel can now display those properties whose group matches the search criteria

#branch UE4

#add Added a new details filter button. This is enabled by default, if the user want to turn off category filtering they need to disable it from the eye menu

diff --git a/Engine/Source/Editor/PropertyEditor/Private/DetailItemNode.cpp b/Engine/Source/Editor/PropertyEditor/Private/DetailItemNode.cpp
--- a/Engine/Source/Editor/PropertyEditor/Private/DetailItemNode.cpp
+++ b/Engine/Source/Editor/PropertyEditor/Private/DetailItemNode.cpp
@@ -277,14 +277,37 @@ ENodeVisibility::Type FDetailItemNode::GetVisibility() const
 	return Visibility;
 }
 
-static bool PassesAllFilters( const FDetailLayoutCustomization& InCustomization, const FDetailFilter& InFilter )
-{
+static bool PassesAllFilters( const FDetailLayoutCustomization& InCustomization, const FDetailFilter& InFilter, const FString& InCategoryName )
+{	
+	struct Local
+	{
+		static bool StringPassesFilter(const FDetailFilter& InFilter, const FString& InString)
+		{
+			// Make sure the passed string matches all filter strings
+			if( InString.Len() > 0 )
+			{
+				for (int32 TestNameIndex = 0; TestNameIndex < InFilter.FilterStrings.Num(); ++TestNameIndex)
+				{
+					const FString& TestName = InFilter.FilterStrings[TestNameIndex];
+					if ( !InString.Contains(TestName) ) 
+					{
+						return false;
+					}
+				}
+				return true;
+			}
+			return false;
+		}
+	};
+
 	bool bPassesAllFilters = true;
-	
+
 	if( InFilter.FilterStrings.Num() > 0 || InFilter.bShowOnlyModifiedProperties == true || InFilter.bShowOnlyDiffering == true )
 	{
 		TSharedPtr<FPropertyNode> PropertyNodePin = InCustomization.GetPropertyNode();
 
+		const bool bPassesCategoryFilter = InFilter.bShowAllChildrenIfCategoryMatches ? Local::StringPassesFilter(InFilter, InCategoryName) : false;
+
 		bPassesAllFilters = false;
 		if( PropertyNodePin.IsValid() && !PropertyNodePin->AsCategoryNode() )
 		{
@@ -295,36 +318,16 @@ static bool PassesAllFilters( const FDetailLayoutCustomization& InCustomization,
 
 			const bool bPassesSearchFilter = bSearchFilterIsEmpty || ( bIsNotBeingFiltered || bIsSeenDueToFiltering || bIsParentSeenDueToFiltering );
 			const bool bPassesModifiedFilter = bPassesSearchFilter && ( InFilter.bShowOnlyModifiedProperties == false || PropertyNodePin->GetDiffersFromDefault() == true );
-			bool bPassesDifferingFilter = true;
-			if( InFilter.bShowOnlyDiffering )
-			{
-				 bPassesDifferingFilter = InFilter.WhitelistedProperties.Find(*FPropertyNode::CreatePropertyPath(PropertyNodePin.ToSharedRef())) != NULL;
-			}
+			const bool bPassesDifferingFilter = InFilter.bShowOnlyDiffering ? InFilter.WhitelistedProperties.Find(*FPropertyNode::CreatePropertyPath(PropertyNodePin.ToSharedRef())) != NULL : true;
 
 			// The property node is visible (note categories are never visible unless they have a child that is visible )
-			bPassesAllFilters = bPassesSearchFilter && bPassesModifiedFilter && bPassesDifferingFilter;
+			bPassesAllFilters = (bPassesSearchFilter && bPassesModifiedFilter && bPassesDifferingFilter) || bPassesCategoryFilter;
 		}
 		else if( InCustomization.HasCustomWidget() )
 		{
-			if( InFilter.FilterStrings.Num() > 0 && InCustomization.WidgetDecl->FilterTextString.Len() > 0  )
-			{
-				// We default to acceptable
-				bPassesAllFilters = true;
-
-				const FString& FilterMatch = InCustomization.WidgetDecl->FilterTextString;
+			const bool bPassesTextFilter = Local::StringPassesFilter(InFilter, InCustomization.WidgetDecl->FilterTextString);
 
-				// Make sure the filter match matches all filter strings
-				for (int32 TestNameIndex = 0; TestNameIndex < InFilter.FilterStrings.Num(); ++TestNameIndex)
-				{
-					const FString& TestName =  InFilter.FilterStrings[TestNameIndex];
-
-					if ( !FilterMatch.Contains( TestName) ) 
-					{
-						bPassesAllFilters = false;
-						break;
-					}
-				}
-			}
+			bPassesAllFilters = bPassesTextFilter || bPassesCategoryFilter;
 		}
 	}
 
@@ -399,7 +402,7 @@ FPropertyPath FDetailItemNode::GetPropertyPath() const
 
 void FDetailItemNode::FilterNode( const FDetailFilter& InFilter )
 {
-	bShouldBeVisibleDueToFiltering = PassesAllFilters( Customization, InFilter );
+	bShouldBeVisibleDueToFiltering = PassesAllFilters( Customization, InFilter, ParentCategory.Pin()->GetDisplayName() );
 
 	bShouldBeVisibleDueToChildFiltering = false;
 
diff --git a/Engine/Source/Editor/PropertyEditor/Private/SDetailsView.cpp b/Engine/Source/Editor/PropertyEditor/Private/SDetailsView.cpp
--- a/Engine/Source/Editor/PropertyEditor/Private/SDetailsView.cpp
+++ b/Engine/Source/Editor/PropertyEditor/Private/SDetailsView.cpp
@@ -54,67 +54,77 @@ void SDetailsView::Construct(const FArguments& InArgs)
 	TSharedRef<SScrollBar> ExternalScrollbar = SNew(SScrollBar);
 	ExternalScrollbar->SetVisibility( TAttribute<EVisibility>( this, &SDetailsView::GetScrollBarVisibility ) );
 
-	FMenuBuilder DetailViewOptions( true, NULL );
+		FMenuBuilder DetailViewOptions( true, NULL );
 
-	FUIAction ShowOnlyModifiedAction( 
-		FExecuteAction::CreateSP( this, &SDetailsView::OnShowOnlyModifiedClicked ),
-		FCanExecuteAction(),
-		FIsActionChecked::CreateSP( this, &SDetailsView::IsShowOnlyModifiedChecked )
-	);
+		if (DetailsViewArgs.bShowModifiedPropertiesOption)
+		{
+			DetailViewOptions.AddMenuEntry( 
+				LOCTEXT("ShowOnlyModified", "Show Only Modified Properties"),
+				LOCTEXT("ShowOnlyModified_ToolTip", "Displays only properties which have been changed from their default"),
+				FSlateIcon(),
+				FUIAction( 
+					FExecuteAction::CreateSP( this, &SDetailsView::OnShowOnlyModifiedClicked ),
+					FCanExecuteAction(),
+					FIsActionChecked::CreateSP( this, &SDetailsView::IsShowOnlyModifiedChecked )
+				),
+				NAME_None,
+				EUserInterfaceActionType::ToggleButton 
+			);
+		}
 
-	if (DetailsViewArgs.bShowModifiedPropertiesOption)
-	{
-		DetailViewOptions.AddMenuEntry( 
-			LOCTEXT("ShowOnlyModified", "Show Only Modified Properties"),
-			LOCTEXT("ShowOnlyModified_ToolTip", "Displays only properties which have been changed from their default"),
+		if( DetailsViewArgs.bShowDifferingPropertiesOption )
+		{
+			DetailViewOptions.AddMenuEntry(
+				LOCTEXT("ShowOnlyDiffering", "Show Only Differing Properties"),
+				LOCTEXT("ShowOnlyDiffering_ToolTip", "Displays only properties in this instance which have been changed or added from the instance being compared"),
+				FSlateIcon(),
+				FUIAction(
+					FExecuteAction::CreateSP(this, &SDetailsView::OnShowOnlyDifferingClicked),
+					FCanExecuteAction(),
+					FIsActionChecked::CreateSP(this, &SDetailsView::IsShowOnlyDifferingChecked)
+				),
+				NAME_None,
+				EUserInterfaceActionType::ToggleButton
+			);
+		}
+
+		DetailViewOptions.AddMenuEntry(
+			LOCTEXT("ShowAllAdvanced", "Show All Advanced Details"),
+			LOCTEXT("ShowAllAdvanced_ToolTip", "Shows all advanced detail sections in each category"),
 			FSlateIcon(),
-			ShowOnlyModifiedAction,
+			FUIAction( 
+			FExecuteAction::CreateSP( this, &SDetailsView::OnShowAllAdvancedClicked ),
+			FCanExecuteAction(),
+			FIsActionChecked::CreateSP( this, &SDetailsView::IsShowAllAdvancedChecked )
+				),
 			NAME_None,
 			EUserInterfaceActionType::ToggleButton 
 		);
-	}
 
-	if( DetailsViewArgs.bShowDifferingPropertiesOption )
-	{
 		DetailViewOptions.AddMenuEntry(
-			LOCTEXT("ShowOnlyDiffering", "Show Only Differing Properties"),
-			LOCTEXT("ShowOnlyDiffering_ToolTip", "Displays only properties in this instance which have been changed or added from the instance being compared"),
+			LOCTEXT("ShowAllChildrenIfCategoryMatches", "Show Child On Category Match"),
+			LOCTEXT("ShowAllChildrenIfCategoryMatches_ToolTip", "Shows children if their category matches the search criteria"),
 			FSlateIcon(),
-			FUIAction(
-				FExecuteAction::CreateSP(this, &SDetailsView::OnShowOnlyDifferingClicked),
+			FUIAction( 
+				FExecuteAction::CreateSP( this, &SDetailsView::OnShowAllChildrenIfCategoryMatchesClicked ),
 				FCanExecuteAction(),
-				FIsActionChecked::CreateSP(this, &SDetailsView::IsShowOnlyDifferingChecked)
+				FIsActionChecked::CreateSP( this, &SDetailsView::IsShowAllChildrenIfCategoryMatchesChecked )
 			),
 			NAME_None,
-			EUserInterfaceActionType::ToggleButton
-		);
-	}
+			EUserInterfaceActionType::ToggleButton 
+			);
 
-	FUIAction ShowAllAdvancedAction( 
-		FExecuteAction::CreateSP( this, &SDetailsView::OnShowAllAdvancedClicked ),
-		FCanExecuteAction(),
-		FIsActionChecked::CreateSP( this, &SDetailsView::IsShowAllAdvancedChecked )
-	);
-
-	DetailViewOptions.AddMenuEntry(
-		LOCTEXT("ShowAllAdvanced", "Show All Advanced Details"),
-		LOCTEXT("ShowAllAdvanced_ToolTip", "Shows all advanced detail sections in each category"),
-		FSlateIcon(),
-		ShowAllAdvancedAction,
-		NAME_None,
-		EUserInterfaceActionType::ToggleButton 
-		);
+		DetailViewOptions.AddMenuEntry(
+			LOCTEXT("CollapseAll", "Collapse All Categories"),
+			LOCTEXT("CollapseAll_ToolTip", "Collapses all root level categories"),
+			FSlateIcon(),
+			FUIAction(FExecuteAction::CreateSP(this, &SDetailsView::SetRootExpansionStates, /*bExpanded=*/false, /*bRecurse=*/false )));
 
-	DetailViewOptions.AddMenuEntry(
-		LOCTEXT("CollapseAll", "Collapse All Categories"),
-		LOCTEXT("CollapseAll_ToolTip", "Collapses all root level categories"),
-		FSlateIcon(),
-		FUIAction(FExecuteAction::CreateSP(this, &SDetailsView::SetRootExpansionStates, /*bExpanded=*/false, /*bRecurse=*/false )));
-	DetailViewOptions.AddMenuEntry(
-		LOCTEXT("ExpandAll", "Expand All Categories"),
-		LOCTEXT("ExpandAll_ToolTip", "Expands all root level categories"),
-		FSlateIcon(),
-		FUIAction(FExecuteAction::CreateSP(this, &SDetailsView::SetRootExpansionStates, /*bExpanded=*/true, /*bRecurse=*/false )));
+		DetailViewOptions.AddMenuEntry(
+			LOCTEXT("ExpandAll", "Expand All Categories"),
+			LOCTEXT("ExpandAll_ToolTip", "Expands all root level categories"),
+			FSlateIcon(),
+			FUIAction(FExecuteAction::CreateSP(this, &SDetailsView::SetRootExpansionStates, /*bExpanded=*/true, /*bRecurse=*/false )));
 
 	TSharedRef<SHorizontalBox> FilterRow = SNew( SHorizontalBox )
 		.Visibility( this, &SDetailsView::GetFilterBoxVisibility )
diff --git a/Engine/Source/Editor/PropertyEditor/Private/SDetailsViewBase.cpp b/Engine/Source/Editor/PropertyEditor/Private/SDetailsViewBase.cpp
--- a/Engine/Source/Editor/PropertyEditor/Private/SDetailsViewBase.cpp
+++ b/Engine/Source/Editor/PropertyEditor/Private/SDetailsViewBase.cpp
@@ -402,6 +402,13 @@ void SDetailsViewBase::OnShowOnlyModifiedClicked()
 	UpdateFilteredDetails();
 }
 
+void SDetailsViewBase::OnShowAllAdvancedClicked()
+{
+	CurrentFilter.bShowAllAdvanced = !CurrentFilter.bShowAllAdvanced;
+
+	UpdateFilteredDetails();
+}
+
 void SDetailsViewBase::OnShowOnlyDifferingClicked()
 {
 	CurrentFilter.bShowOnlyDiffering = !CurrentFilter.bShowOnlyDiffering;
@@ -409,9 +416,9 @@ void SDetailsViewBase::OnShowOnlyDifferingClicked()
 	UpdateFilteredDetails();
 }
 
-void SDetailsViewBase::OnShowAllAdvancedClicked()
+void SDetailsViewBase::OnShowAllChildrenIfCategoryMatchesClicked()
 {
-	CurrentFilter.bShowAllAdvanced = !CurrentFilter.bShowAllAdvanced;
+	CurrentFilter.bShowAllChildrenIfCategoryMatches = !CurrentFilter.bShowAllChildrenIfCategoryMatches;
 
 	UpdateFilteredDetails();
 }
diff --git a/Engine/Source/Editor/PropertyEditor/Private/SDetailsViewBase.h b/Engine/Source/Editor/PropertyEditor/Private/SDetailsViewBase.h
--- a/Engine/Source/Editor/PropertyEditor/Private/SDetailsViewBase.h
+++ b/Engine/Source/Editor/PropertyEditor/Private/SDetailsViewBase.h
@@ -54,9 +54,10 @@ struct FDetailFilter
 		: bShowOnlyModifiedProperties(false)
 		, bShowAllAdvanced(false)
 		, bShowOnlyDiffering(false)
+		, bShowAllChildrenIfCategoryMatches(true)
 	{}
 
-	bool IsEmptyFilter() const { return FilterStrings.Num() == 0 && bShowOnlyModifiedProperties == false && bShowAllAdvanced == false && bShowOnlyDiffering == false; }
+	bool IsEmptyFilter() const { return FilterStrings.Num() == 0 && bShowOnlyModifiedProperties == false && bShowAllAdvanced == false && bShowOnlyDiffering == false && bShowAllChildrenIfCategoryMatches == false; }
 
 	/** Any user search terms that items must match */
 	TArray<FString> FilterStrings;
@@ -66,6 +67,8 @@ struct FDetailFilter
 	bool bShowAllAdvanced;
 	/** If we should only show differing properties */
 	bool bShowOnlyDiffering;
+	/** If we should show all the children if their category name matches the search */
+	bool bShowAllChildrenIfCategoryMatches;
 	TSet<FPropertyPath> WhitelistedProperties;
 };
 
@@ -350,14 +353,20 @@ protected:
 	/** @return true if show only differing is checked */
 	bool IsShowOnlyDifferingChecked() const { return CurrentFilter.bShowOnlyDiffering; }
 
+	/** @return true if show all advanced is checked */
+	bool IsShowAllChildrenIfCategoryMatchesChecked() const { return CurrentFilter.bShowAllChildrenIfCategoryMatches; }
+
 	/** Called when show only modified is clicked */
 	void OnShowOnlyModifiedClicked();
 
+	/** Called when show all advanced is clicked */
+	void OnShowAllAdvancedClicked();
+
 	/** Called when show only differing is clicked */
 	void OnShowOnlyDifferingClicked();
 
-	/** Called when show all advanced is clicked */
-	void OnShowAllAdvancedClicked();
+	/** Called when show all children if category matches is clicked */
+	void OnShowAllChildrenIfCategoryMatchesClicked();
 
 	/**
 	* Updates the details with the passed in filter
