Fixed show only modified properties in the details panel not working

diff --git a/Engine/Source/Editor/PropertyEditor/Private/DetailItemNode.cpp b/Engine/Source/Editor/PropertyEditor/Private/DetailItemNode.cpp
--- a/Engine/Source/Editor/PropertyEditor/Private/DetailItemNode.cpp
+++ b/Engine/Source/Editor/PropertyEditor/Private/DetailItemNode.cpp
@@ -304,14 +304,16 @@ static bool PassesAllFilters( const FDetailLayoutCustomization& InCustomization,
 
 	if( InFilter.FilterStrings.Num() > 0 || InFilter.bShowOnlyModifiedProperties == true || InFilter.bShowOnlyDiffering == true )
 	{
+		const bool bSearchFilterIsEmpty = InFilter.FilterStrings.Num() == 0;
+
 		TSharedPtr<FPropertyNode> PropertyNodePin = InCustomization.GetPropertyNode();
 
-		const bool bPassesCategoryFilter = InFilter.bShowAllChildrenIfCategoryMatches ? Local::StringPassesFilter(InFilter, InCategoryName) : false;
+		const bool bPassesCategoryFilter = !bSearchFilterIsEmpty && InFilter.bShowAllChildrenIfCategoryMatches ? Local::StringPassesFilter(InFilter, InCategoryName) : false;
 
 		bPassesAllFilters = false;
 		if( PropertyNodePin.IsValid() && !PropertyNodePin->AsCategoryNode() )
 		{
-			const bool bSearchFilterIsEmpty = InFilter.FilterStrings.Num() == 0;
+			
 			const bool bIsNotBeingFiltered = PropertyNodePin->HasNodeFlags(EPropertyNodeFlags::IsBeingFiltered) == 0;
 			const bool bIsSeenDueToFiltering = PropertyNodePin->HasNodeFlags(EPropertyNodeFlags::IsSeenDueToFiltering) != 0;
 			const bool bIsParentSeenDueToFiltering = PropertyNodePin->HasNodeFlags(EPropertyNodeFlags::IsParentSeenDueToFiltering) != 0;
