Various fixes for include patterns that rely on delayed function instantiation to compile, which prevents including them in isolation. Tidy up some other include patterns.

diff --git a/Engine/Source/Editor/Kismet/Private/BlueprintEditorTabFactories.cpp b/Engine/Source/Editor/Kismet/Private/BlueprintEditorTabFactories.cpp
--- a/Engine/Source/Editor/Kismet/Private/BlueprintEditorTabFactories.cpp
+++ b/Engine/Source/Editor/Kismet/Private/BlueprintEditorTabFactories.cpp
@@ -57,6 +57,12 @@ FGraphEditorSummoner::FGraphEditorSummoner(TSharedPtr<class FBlueprintEditor> In
 
 }
 
+TSharedRef<SWidget> FGraphEditorSummoner::CreateTabBodyForObject(const FWorkflowTabSpawnInfo& Info, UEdGraph* DocumentID) const
+{
+	check(Info.TabInfo.IsValid());
+	return OnCreateGraphEditorWidget.Execute(Info.TabInfo.ToSharedRef(), DocumentID);
+}
+
 const FSlateBrush* FGraphEditorSummoner::GetTabIconForObject(const FWorkflowTabSpawnInfo& Info, UEdGraph* DocumentID) const
 {
 	return FBlueprintEditor::GetGlyphForGraph(DocumentID, false);
diff --git a/Engine/Source/Editor/Kismet/Private/BlueprintEditorTabFactories.h b/Engine/Source/Editor/Kismet/Private/BlueprintEditorTabFactories.h
--- a/Engine/Source/Editor/Kismet/Private/BlueprintEditorTabFactories.h
+++ b/Engine/Source/Editor/Kismet/Private/BlueprintEditorTabFactories.h
@@ -133,11 +133,7 @@ protected:
 		return TAttribute<FText>::Create(TAttribute<FText>::FGetter::CreateStatic<UEdGraph*>(&FLocalKismetCallbacks::GetGraphDisplayName, DocumentID));
 	}
 
-	virtual TSharedRef<SWidget> CreateTabBodyForObject(const FWorkflowTabSpawnInfo& Info, UEdGraph* DocumentID) const override
-	{
-		check(Info.TabInfo.IsValid());
-		return OnCreateGraphEditorWidget.Execute(Info.TabInfo.ToSharedRef(), DocumentID);
-	}
+	virtual TSharedRef<SWidget> CreateTabBodyForObject(const FWorkflowTabSpawnInfo& Info, UEdGraph* DocumentID) const override;
 
 	virtual const FSlateBrush* GetTabIconForObject(const FWorkflowTabSpawnInfo& Info, UEdGraph* DocumentID) const override;
 
diff --git a/Engine/Source/Editor/PropertyEditor/Private/Presentation/PropertyTable/PropertyTableObjectNameCell.cpp b/Engine/Source/Editor/PropertyEditor/Private/Presentation/PropertyTable/PropertyTableObjectNameCell.cpp
--- a/Engine/Source/Editor/PropertyEditor/Private/Presentation/PropertyTable/PropertyTableObjectNameCell.cpp
+++ b/Engine/Source/Editor/PropertyEditor/Private/Presentation/PropertyTable/PropertyTableObjectNameCell.cpp
@@ -70,6 +70,16 @@ TWeakObjectPtr< UObject > FPropertyTableObjectNameCell::GetObject() const
 	return ObjectNode->GetUObject( 0 );
 }
 
+TSharedRef< class IPropertyTableColumn > FPropertyTableObjectNameCell::GetColumn() const
+{ 
+	return Column.Pin().ToSharedRef(); 
+}
+
+TSharedRef< class IPropertyTableRow > FPropertyTableObjectNameCell::GetRow() const
+{
+	return Row.Pin().ToSharedRef(); 
+}
+
 TSharedRef< class IPropertyTable > FPropertyTableObjectNameCell::GetTable() const
 {
 	return Column.Pin()->GetTable();
diff --git a/Engine/Source/Editor/PropertyEditor/Private/Presentation/PropertyTable/PropertyTableObjectNameCell.h b/Engine/Source/Editor/PropertyEditor/Private/Presentation/PropertyTable/PropertyTableObjectNameCell.h
--- a/Engine/Source/Editor/PropertyEditor/Private/Presentation/PropertyTable/PropertyTableObjectNameCell.h
+++ b/Engine/Source/Editor/PropertyEditor/Private/Presentation/PropertyTable/PropertyTableObjectNameCell.h
@@ -26,8 +26,8 @@ public:
 	virtual TWeakObjectPtr< UObject > GetObject() const override;
 
 	virtual TSharedPtr< class FPropertyNode > GetNode() const override { return NULL; }
-	virtual TSharedRef< class IPropertyTableColumn > GetColumn() const override { return Column.Pin().ToSharedRef(); }
-	virtual TSharedRef< class IPropertyTableRow > GetRow() const override { return Row.Pin().ToSharedRef(); }
+	virtual TSharedRef< class IPropertyTableColumn > GetColumn() const override;
+	virtual TSharedRef< class IPropertyTableRow > GetRow() const override;
 	virtual TSharedRef< class IPropertyTable > GetTable() const override;
 	virtual TSharedPtr< class IPropertyHandle > GetPropertyHandle() const override { return NULL; }
 
diff --git a/Engine/Source/Editor/PropertyEditor/Private/Presentation/PropertyTable/PropertyTablePropertyNameCell.cpp b/Engine/Source/Editor/PropertyEditor/Private/Presentation/PropertyTable/PropertyTablePropertyNameCell.cpp
--- a/Engine/Source/Editor/PropertyEditor/Private/Presentation/PropertyTable/PropertyTablePropertyNameCell.cpp
+++ b/Engine/Source/Editor/PropertyEditor/Private/Presentation/PropertyTable/PropertyTablePropertyNameCell.cpp
@@ -53,6 +53,11 @@ void FPropertyTablePropertyNameCell::EnterEditMode()
 {
 }
 
+TSharedRef< class IPropertyTableColumn > FPropertyTablePropertyNameCell::GetColumn() const
+{ 
+	return Column.Pin().ToSharedRef(); 
+}
+
 TWeakObjectPtr< UObject > FPropertyTablePropertyNameCell::GetObject() const
 {
 	if ( !ObjectNode.IsValid() )
@@ -63,6 +68,11 @@ TWeakObjectPtr< UObject > FPropertyTablePropertyNameCell::GetObject() const
 	return ObjectNode->GetUObject( 0 );
 }
 
+TSharedRef< class IPropertyTableRow > FPropertyTablePropertyNameCell::GetRow() const
+{ 
+	return Row.Pin().ToSharedRef(); 
+}
+
 void FPropertyTablePropertyNameCell::ExitEditMode()
 {
 }
\ No newline at end of file
diff --git a/Engine/Source/Editor/PropertyEditor/Private/Presentation/PropertyTable/PropertyTablePropertyNameCell.h b/Engine/Source/Editor/PropertyEditor/Private/Presentation/PropertyTable/PropertyTablePropertyNameCell.h
--- a/Engine/Source/Editor/PropertyEditor/Private/Presentation/PropertyTable/PropertyTablePropertyNameCell.h
+++ b/Engine/Source/Editor/PropertyEditor/Private/Presentation/PropertyTable/PropertyTablePropertyNameCell.h
@@ -17,13 +17,13 @@ public:
 
 	virtual void ExitEditMode() override;
 
-	virtual TSharedRef< class IPropertyTableColumn > GetColumn() const override { return Column.Pin().ToSharedRef(); }
+	virtual TSharedRef< class IPropertyTableColumn > GetColumn() const override;
 
 	virtual TSharedPtr< class FPropertyNode > GetNode() const override { return NULL; }
 
 	virtual TWeakObjectPtr< UObject > GetObject() const override;
 
-	virtual TSharedRef< class IPropertyTableRow > GetRow() const override { return Row.Pin().ToSharedRef(); }
+	virtual TSharedRef< class IPropertyTableRow > GetRow() const override;
 
 	virtual TSharedRef< class IPropertyTable > GetTable() const override;
 
diff --git a/Engine/Source/Editor/Sequencer/Private/SectionDragOperations.cpp b/Engine/Source/Editor/Sequencer/Private/SectionDragOperations.cpp
--- a/Engine/Source/Editor/Sequencer/Private/SectionDragOperations.cpp
+++ b/Engine/Source/Editor/Sequencer/Private/SectionDragOperations.cpp
@@ -133,6 +133,11 @@ TOptional<float> FSequencerDragOperation::SnapToTimes(float InitialTime, const T
 }
 
 
+FResizeSection::FResizeSection( UMovieSceneSection& InSection, bool bInDraggingByEnd )
+	: Section( &InSection )
+	, bDraggingByEnd( bInDraggingByEnd )
+{
+}
 
 void FResizeSection::OnBeginDrag(const FVector2D& LocalMousePos, TSharedPtr<FTrackNode> SequencerNode)
 {
@@ -209,6 +214,12 @@ void FResizeSection::OnDrag( const FPointerEvent& MouseEvent, const FVector2D& L
 	}
 }
 
+FMoveSection::FMoveSection( UMovieSceneSection& InSection )
+	: Section( &InSection )
+	, DragOffset(ForceInit)
+{
+}
+
 void FMoveSection::OnBeginDrag(const FVector2D& LocalMousePos, TSharedPtr<FTrackNode> SequencerNode)
 {
 	if( Section.IsValid() )
diff --git a/Engine/Source/Editor/Sequencer/Private/SectionDragOperations.h b/Engine/Source/Editor/Sequencer/Private/SectionDragOperations.h
--- a/Engine/Source/Editor/Sequencer/Private/SectionDragOperations.h
+++ b/Engine/Source/Editor/Sequencer/Private/SectionDragOperations.h
@@ -79,10 +79,7 @@ protected:
 class FResizeSection : public FSequencerDragOperation
 {
 public:
-	FResizeSection( UMovieSceneSection& InSection, bool bInDraggingByEnd )
-		: Section( &InSection )
-		, bDraggingByEnd( bInDraggingByEnd )
-	{}
+	FResizeSection( UMovieSceneSection& InSection, bool bInDraggingByEnd );
 
 	/** FSequencerDragOperation interface */
 	virtual void OnBeginDrag(const FVector2D& LocalMousePos, TSharedPtr<FTrackNode> SequencerNode) override;
@@ -102,10 +99,7 @@ private:
 class FMoveSection : public FSequencerDragOperation
 {
 public:
-	FMoveSection( UMovieSceneSection& InSection )
-		: Section( &InSection )
-		, DragOffset(ForceInit)
-	{}
+	FMoveSection( UMovieSceneSection& InSection );
 
 	/** FSequencerDragOperation interface */
 	virtual void OnBeginDrag(const FVector2D& LocalMousePos, TSharedPtr<FTrackNode> SequencerNode) override;
diff --git a/Engine/Source/Editor/Sequencer/Private/SequencerAssetEditor.cpp b/Engine/Source/Editor/Sequencer/Private/SequencerAssetEditor.cpp
--- a/Engine/Source/Editor/Sequencer/Private/SequencerAssetEditor.cpp
+++ b/Engine/Source/Editor/Sequencer/Private/SequencerAssetEditor.cpp
@@ -98,6 +98,10 @@ void FSequencerAssetEditor::InitSequencerAssetEditor( const EToolkitMode::Type M
 	}
 }
 
+TSharedRef<ISequencer> FSequencerAssetEditor::GetSequencerInterface() const
+{ 
+	return Sequencer.ToSharedRef(); 
+}
 
 FSequencerAssetEditor::FSequencerAssetEditor()
 {
diff --git a/Engine/Source/Editor/Sequencer/Private/SequencerAssetEditor.h b/Engine/Source/Editor/Sequencer/Private/SequencerAssetEditor.h
--- a/Engine/Source/Editor/Sequencer/Private/SequencerAssetEditor.h
+++ b/Engine/Source/Editor/Sequencer/Private/SequencerAssetEditor.h
@@ -25,7 +25,7 @@ public:
 	 */
 	void InitSequencerAssetEditor( const EToolkitMode::Type Mode, const FSequencerViewParams& InViewParams, const TSharedPtr< class IToolkitHost >& InitToolkitHost, UMovieScene* InRootMovieScene, const TArray<FOnCreateTrackEditor>& TrackEditorDelegates, bool bEditWithinLevelEditor );
 
-	TSharedRef<ISequencer> GetSequencerInterface() const override { return Sequencer.ToSharedRef(); }
+	TSharedRef<ISequencer> GetSequencerInterface() const override;
 
 	/** IToolkit interface */
 	virtual FName GetToolkitFName() const override;
diff --git a/Engine/Source/Runtime/PakFile/Public/IPlatformFilePak.h b/Engine/Source/Runtime/PakFile/Public/IPlatformFilePak.h
--- a/Engine/Source/Runtime/PakFile/Public/IPlatformFilePak.h
+++ b/Engine/Source/Runtime/PakFile/Public/IPlatformFilePak.h
@@ -248,7 +248,7 @@ typedef TMap<FString, FPakEntry*> FPakDirectory;
 /**
  * Pak file.
  */
-class PAKFILE_API FPakFile
+class PAKFILE_API FPakFile : FNoncopyable
 {
 	/** Pak filename. */
 	FString PakFilename;
