UE-6710 - Details panel now allows optional floats for min/max sizes.  It always did, if you put 0 as the min/max now it's explicit with Optional.  Making the SLateChildSizeCustomization use an optional max size so that it doesn't get cut off.

diff --git a/Engine/Source/Editor/PropertyEditor/Private/DetailPropertyRow.cpp b/Engine/Source/Editor/PropertyEditor/Private/DetailPropertyRow.cpp
--- a/Engine/Source/Editor/PropertyEditor/Private/DetailPropertyRow.cpp
+++ b/Engine/Source/Editor/PropertyEditor/Private/DetailPropertyRow.cpp
@@ -318,16 +318,10 @@ void FDetailPropertyRow::MakeNameWidget( FDetailWidgetRow& Row, const TSharedPtr
 	EVerticalAlignment VerticalAlignment = VAlign_Center;
 	EHorizontalAlignment HorizontalAlignment = HAlign_Fill;
 
-	float MinWidth = 0.0f;
-	float MaxWidth = 0.0f;
-
 	if( InCustomRow.IsValid() )
 	{
 		VerticalAlignment = InCustomRow->NameWidget.VerticalAlignment;
 		HorizontalAlignment = InCustomRow->NameWidget.HorizontalAlignment;
-
-		MinWidth = InCustomRow->NameWidget.MinWidth;
-		MaxWidth = InCustomRow->NameWidget.MaxWidth;
 	}
 
 	TAttribute<bool> IsEnabledAttrib = CustomIsEnabledAttrib;
@@ -385,8 +379,8 @@ void FDetailPropertyRow::MakeValueWidget( FDetailWidgetRow& Row, const TSharedPt
 	EVerticalAlignment VerticalAlignment = VAlign_Center;
 	EHorizontalAlignment HorizontalAlignment = HAlign_Left;
 
-	float MinWidth = 0.0f;
-	float MaxWidth = 0.0f;
+	TOptional<float> MinWidth;
+	TOptional<float> MaxWidth;
 
 	if( InCustomRow.IsValid() )
 	{
diff --git a/Engine/Source/Editor/PropertyEditor/Private/SDetailSingleItemRow.cpp b/Engine/Source/Editor/PropertyEditor/Private/SDetailSingleItemRow.cpp
--- a/Engine/Source/Editor/PropertyEditor/Private/SDetailSingleItemRow.cpp
+++ b/Engine/Source/Editor/PropertyEditor/Private/SDetailSingleItemRow.cpp
@@ -17,12 +17,12 @@ class SConstrainedBox : public SCompoundWidget
 {
 public:
 	SLATE_BEGIN_ARGS( SConstrainedBox )
-		: _MinWidth(0)
-		, _MaxWidth(0)
+		: _MinWidth()
+		, _MaxWidth()
 	{}
 		SLATE_DEFAULT_SLOT( FArguments, Content )
-		SLATE_ATTRIBUTE( float, MinWidth )
-		SLATE_ATTRIBUTE( float, MaxWidth )
+		SLATE_ATTRIBUTE( TOptional<float>, MinWidth )
+		SLATE_ATTRIBUTE( TOptional<float>, MaxWidth )
 	SLATE_END_ARGS()
 
 	void Construct( const FArguments& InArgs )
@@ -38,10 +38,10 @@ public:
 
 	virtual FVector2D ComputeDesiredSize() const override
 	{
-		const float MinWidthVal = MinWidth.Get();
-		const float MaxWidthVal = MaxWidth.Get();
+		const float MinWidthVal = MinWidth.Get().Get(0.0f);
+		const float MaxWidthVal = MaxWidth.Get().Get(0.0f);
 
-		if( MinWidthVal == 0.0f && MaxWidthVal == 0.0f )
+		if ( MinWidthVal == 0.0f && MaxWidthVal == 0.0f )
 		{
 			return SCompoundWidget::ComputeDesiredSize();
 		}
@@ -59,8 +59,8 @@ public:
 		}
 	}
 private:
-	TAttribute<float> MinWidth;
-	TAttribute<float> MaxWidth;
+	TAttribute< TOptional<float> > MinWidth;
+	TAttribute< TOptional<float> > MaxWidth;
 };
 
 
diff --git a/Engine/Source/Editor/PropertyEditor/Public/DetailWidgetRow.h b/Engine/Source/Editor/PropertyEditor/Public/DetailWidgetRow.h
--- a/Engine/Source/Editor/PropertyEditor/Public/DetailWidgetRow.h
+++ b/Engine/Source/Editor/PropertyEditor/Public/DetailWidgetRow.h
@@ -34,13 +34,13 @@ public:
 		return *this;
 	}
 
-	FDetailWidgetDecl& MinDesiredWidth( float InMinWidth )
+	FDetailWidgetDecl& MinDesiredWidth( TOptional<float> InMinWidth )
 	{
 		MinWidth = InMinWidth;
 		return *this;
 	}
 
-	FDetailWidgetDecl& MaxDesiredWidth( float InMaxWidth )
+	FDetailWidgetDecl& MaxDesiredWidth( TOptional<float> InMaxWidth )
 	{
 		MaxWidth = InMaxWidth;
 		return *this;
@@ -49,8 +49,8 @@ public:
 	TSharedRef<SWidget> Widget;
 	EHorizontalAlignment HorizontalAlignment;
 	EVerticalAlignment VerticalAlignment;
-	float MinWidth;
-	float MaxWidth;
+	TOptional<float> MinWidth;
+	TOptional<float> MaxWidth;
 private:
 	class FDetailWidgetRow& ParentDecl;
 };
diff --git a/Engine/Source/Editor/UMGEditor/Private/Customizations/SlateChildSizeCustomization.cpp b/Engine/Source/Editor/UMGEditor/Private/Customizations/SlateChildSizeCustomization.cpp
--- a/Engine/Source/Editor/UMGEditor/Private/Customizations/SlateChildSizeCustomization.cpp
+++ b/Engine/Source/Editor/UMGEditor/Private/Customizations/SlateChildSizeCustomization.cpp
@@ -34,6 +34,7 @@ void FSlateChildSizeCustomization::CustomizeHeader(TSharedRef<IPropertyHandle> P
 		PropertyHandle->CreatePropertyNameWidget()
 	]
 	.ValueContent()
+	.MaxDesiredWidth(TOptional<float>())
 	[
 		SNew(SHorizontalBox)
 
@@ -85,11 +86,6 @@ void FSlateChildSizeCustomization::CustomizeHeader(TSharedRef<IPropertyHandle> P
 				SNew( SNumericEntryBox<float> )
 				.LabelVAlign(VAlign_Center)
 				.Visibility(this, &FSlateChildSizeCustomization::GetValueVisiblity, RuleHandle)
-				//.Label()
-				//[
-				//	SNew(STextBlock)
-				//	.Text(LOCTEXT("Fill", "Fill"))
-				//]
 				.Value(this, &FSlateChildSizeCustomization::GetValue, ValueHandle)
 				.OnValueCommitted(this, &FSlateChildSizeCustomization::HandleValueComitted, ValueHandle)
 				.UndeterminedString( LOCTEXT("MultipleValues", "Multiple Values") )
