Fixed tooltips in property combo box widgets.

diff --git a/Engine/Source/Editor/PropertyEditor/Private/UserInterface/PropertyEditor/SPropertyComboBox.cpp b/Engine/Source/Editor/PropertyEditor/Private/UserInterface/PropertyEditor/SPropertyComboBox.cpp
--- a/Engine/Source/Editor/PropertyEditor/Private/UserInterface/PropertyEditor/SPropertyComboBox.cpp
+++ b/Engine/Source/Editor/PropertyEditor/Private/UserInterface/PropertyEditor/SPropertyComboBox.cpp
@@ -93,6 +93,7 @@ void SPropertyComboBox::OnSelectionChangedInternal( TSharedPtr<FString> InSelect
 	if( bEnabled )
 	{
 		OnSelectionChanged.ExecuteIfBound( InSelectedItem, SelectInfo );
+		SetSelectedItem(*InSelectedItem);
 	}
 }
 
