Editor - Setting the base indent level on the property grid to 1, to save some horizontal space.

diff --git a/Engine/Source/Editor/PropertyEditor/Private/SDetailSingleItemRow.cpp b/Engine/Source/Editor/PropertyEditor/Private/SDetailSingleItemRow.cpp
--- a/Engine/Source/Editor/PropertyEditor/Private/SDetailSingleItemRow.cpp
+++ b/Engine/Source/Editor/PropertyEditor/Private/SDetailSingleItemRow.cpp
@@ -125,6 +125,7 @@ void SDetailSingleItemRow::Construct( const FArguments& InArgs, FDetailLayoutCus
 					.AutoWidth()
 					[
 						SNew( SExpanderArrow, SharedThis(this) )
+						.BaseIndentLevel(1)
 					]
 					+ SHorizontalBox::Slot()
 					.HAlign( Row.NameWidget.HorizontalAlignment )
@@ -168,6 +169,7 @@ void SDetailSingleItemRow::Construct( const FArguments& InArgs, FDetailLayoutCus
 				.AutoWidth()
 				[
 					SNew( SExpanderArrow, SharedThis(this) )
+					.BaseIndentLevel(1)
 				]
 				+ SHorizontalBox::Slot()
 				.HAlign( Row.WholeRowWidget.HorizontalAlignment )
@@ -325,8 +327,8 @@ TSharedRef<SWidget> SDetailSingleItemRow::CreateExtensionWidget(TSharedRef<SWidg
 					[
 						ValueWidget
 					]
-
-				+ SHorizontalBox::Slot()
+					
+					+ SHorizontalBox::Slot()
 					.AutoWidth()
 					[
 						ExtensionHandler->GenerateExtensionWidget(ObjectClass, Handle)
