Display the components tree under the details panel name and search bar

diff --git a/Engine/Source/Editor/LevelEditor/Private/SLevelEditor.cpp b/Engine/Source/Editor/LevelEditor/Private/SLevelEditor.cpp
--- a/Engine/Source/Editor/LevelEditor/Private/SLevelEditor.cpp
+++ b/Engine/Source/Editor/LevelEditor/Private/SLevelEditor.cpp
@@ -524,7 +524,15 @@ public:
 		GEditor->OnPostTransformComponents().AddRaw(this, &SActorDetails::OnEditorPostTransformComponents);
 
 		FPropertyEditorModule& PropPlugin = FModuleManager::LoadModuleChecked<FPropertyEditorModule>("PropertyEditor");
-		const FDetailsViewArgs DetailsViewArgs( true, true, true, false, false, GUnrealEd, false, TabIdentifier );
+		FDetailsViewArgs DetailsViewArgs;
+		DetailsViewArgs.bUpdatesFromSelection = true;
+		DetailsViewArgs.bLockable = true;
+		DetailsViewArgs.bObjectsUseNameArea = false;
+		DetailsViewArgs.NotifyHook = GUnrealEd;
+		DetailsViewArgs.ViewIdentifier = TabIdentifier;
+		DetailsViewArgs.bCustomNameAreaLocation = true;
+		DetailsViewArgs.bCustomFilterAreaLocation = true;
+
 		DetailsView = PropPlugin.CreateDetailView(DetailsViewArgs);
 
 		DetailsView->SetIsPropertyVisibleDelegate( FIsPropertyVisible::CreateStatic( &Local::IsPropertyVisible ) );
@@ -532,16 +540,32 @@ public:
 		// Set up a delegate to call to add generic details to the view
 		DetailsView->SetGenericLayoutDetailsDelegate( FOnGetDetailCustomizationInstance::CreateStatic( &FLevelEditorGenericDetails::MakeInstance ) );
 
-		SAssignNew(ComponentsBox, SScrollBox)
-			.Visibility(EVisibility::Collapsed);
+		SAssignNew(ComponentsBox, SBox)
+		.Visibility(EVisibility::Collapsed);
 
 		ChildSlot
 		[
-			SAssignNew(DetailsSplitter, SSplitter)
-			.Orientation(Orient_Vertical)
-			+ SSplitter::Slot()
+			SNew( SVerticalBox )
+			+SVerticalBox::Slot()
+			.Padding(0.0f, 0.0f, 0.0f, 2.0f)
+			.AutoHeight()
+			[
+				DetailsView->GetNameAreaWidget().ToSharedRef()
+			]
+			+ SVerticalBox::Slot()
+			.Padding(0.0f, 0.0f, 0.0f, 2.0f)
+			.AutoHeight()
 			[
-				DetailsView.ToSharedRef()
+				DetailsView->GetFilterAreaWidget().ToSharedRef()
+			]
+			+SVerticalBox::Slot()
+			[
+				SAssignNew(DetailsSplitter, SSplitter)
+				.Orientation(Orient_Vertical)
+				+ SSplitter::Slot()
+				[
+					DetailsView.ToSharedRef()
+				]
 			]
 		];
 
@@ -575,9 +599,8 @@ public:
 				{
 					bShowingComponents = true;
 
-					ComponentsBox->ClearChildren();
-					ComponentsBox->AddSlot()
-					[
+					ComponentsBox->SetContent
+					(
 						SAssignNew(SCSEditor, SSCSEditor)
 						.InEditingMode(true)
 						.ActorContext(this, &SActorDetails::GetSelectedActor)												// Get the instance of the actor in the world
@@ -586,7 +609,7 @@ public:
 						//.OnHighlightPropertyInDetailsView(this, &SLevelEditor::OnSCSEditorHighlightPropertyInDetailsView)	// Also unsure and don't think it's needed
 						//.OnAddNewComponent()
 						//.OnAddExistingComponent()
-					];
+					);
 				}
 			}
 
@@ -899,7 +922,7 @@ private:
 
 	TSharedPtr<SSplitter> DetailsSplitter;
 	TSharedPtr<IDetailsView> DetailsView;
-	TSharedPtr<SScrollBox> ComponentsBox;
+	TSharedPtr<SBox> ComponentsBox;
 	
 	/** SCS editor */
 	TSharedPtr<SSCSEditor> SCSEditor;
diff --git a/Engine/Source/Editor/PropertyEditor/Private/SDetailNameArea.cpp b/Engine/Source/Editor/PropertyEditor/Private/SDetailNameArea.cpp
--- a/Engine/Source/Editor/PropertyEditor/Private/SDetailNameArea.cpp
+++ b/Engine/Source/Editor/PropertyEditor/Private/SDetailNameArea.cpp
@@ -106,7 +106,7 @@ TSharedRef< SWidget > SDetailNameArea::BuildObjectNameArea( const TArray< TWeakO
 
 			ObjectNameArea->AddSlot()
 				.AutoWidth()
-				.Padding(0, 0, 6, 0)
+				.Padding(0, 0, 3, 0)
 				[
 					SNew(SBox)
 					.WidthOverride(200.0f)
diff --git a/Engine/Source/Editor/PropertyEditor/Private/SDetailsView.cpp b/Engine/Source/Editor/PropertyEditor/Private/SDetailsView.cpp
--- a/Engine/Source/Editor/PropertyEditor/Private/SDetailsView.cpp
+++ b/Engine/Source/Editor/PropertyEditor/Private/SDetailsView.cpp
@@ -126,7 +126,7 @@ void SDetailsView::Construct(const FArguments& InArgs)
 			FSlateIcon(),
 			FUIAction(FExecuteAction::CreateSP(this, &SDetailsView::SetRootExpansionStates, /*bExpanded=*/true, /*bRecurse=*/false )));
 
-	TSharedRef<SHorizontalBox> FilterRow = SNew( SHorizontalBox )
+	FilterRow = SNew( SHorizontalBox )
 		.Visibility( this, &SDetailsView::GetFilterBoxVisibility )
 		+SHorizontalBox::Slot()
 		.FillWidth( 1 )
@@ -175,51 +175,60 @@ void SDetailsView::Construct(const FArguments& InArgs)
 			];
 	}
 
+	// Create the name area which does not change when selection changes
+	SAssignNew(NameArea, SDetailNameArea, &SelectedObjects)
+		// the name area is only for actors
+		.Visibility(this, &SDetailsView::GetActorNameAreaVisibility)
+		.ShowActorLabel(DetailsViewArgs.bShowActorLabel)
+		// only show the selection tip if we're not selecting objects
+		.SelectionTip(!DetailsViewArgs.bHideSelectionTip);
 
-	ChildSlot
-	[
-		SNew( SVerticalBox )
-		+ SVerticalBox::Slot()
+	TSharedRef<SVerticalBox> VerticalBox = SNew(SVerticalBox);
+
+	if( !DetailsViewArgs.bCustomNameAreaLocation )
+	{
+		VerticalBox->AddSlot()
 		.AutoHeight()
-		.Padding( 0.0f, 0.0f, 0.0f, 4.0f )
+		.Padding(0.0f, 0.0f, 0.0f, 4.0f)
 		[
-			// Create the name area which does not change when selection changes
-			SAssignNew( NameArea, SDetailNameArea, &SelectedObjects )
-			// the name area is only for actors
-			.Visibility( this, &SDetailsView::GetActorNameAreaVisibility  )
-			.OnLockButtonClicked( this, &SDetailsView::OnLockButtonClicked )
-			.IsLocked( this, &SDetailsView::IsLocked )
-			.ShowLockButton( DetailsViewArgs.bLockable )
-			.ShowActorLabel( DetailsViewArgs.bShowActorLabel )
-			// only show the selection tip if we're not selecting objects
-			.SelectionTip( !DetailsViewArgs.bHideSelectionTip )
-		]
-		+ SVerticalBox::Slot()
+			NameArea.ToSharedRef()
+		];
+	}
+
+	if( !DetailsViewArgs.bCustomFilterAreaLocation )
+	{
+		VerticalBox->AddSlot()
 		.AutoHeight()
-		.Padding( 0.0f, 0.0f, 0.0f, 2.0f )
+		.Padding(0.0f, 0.0f, 0.0f, 2.0f)
+		[
+			FilterRow.ToSharedRef()
+		];
+	}
+
+	VerticalBox->AddSlot()
+	.FillHeight(1)
+	.Padding(0)
+	[
+		SNew(SHorizontalBox)
+		+ SHorizontalBox::Slot()
 		[
-			FilterRow
+			ConstructTreeView(ExternalScrollbar)
 		]
-		+ SVerticalBox::Slot()
-		.FillHeight(1)
-		.Padding(0)
+		+ SHorizontalBox::Slot()
+		.AutoWidth()
 		[
-			SNew( SHorizontalBox )
-			+ SHorizontalBox::Slot()
-			[
-				ConstructTreeView( ExternalScrollbar )
-			]
-			+ SHorizontalBox::Slot()
-			.AutoWidth()
+			SNew(SBox)
+			.WidthOverride(16.0f)
 			[
-				SNew( SBox )
-				.WidthOverride( 16.0f )
-				[
-					ExternalScrollbar
-				]
+				ExternalScrollbar
 			]
 		]
 	];
+
+	ChildSlot
+	[
+		VerticalBox
+	];
 }
 
 TSharedRef<SDetailTree> SDetailsView::ConstructTreeView( TSharedRef<SScrollBar>& ScrollBar )
@@ -642,4 +651,16 @@ bool SDetailsView::IsConnected() const
 	return RootPropertyNode.IsValid() && (RootPropertyNode->GetNumObjects() > 0);
 }
 
+const FSlateBrush* SDetailsView::OnGetLockButtonImageResource() const
+{
+	if (bIsLocked)
+	{
+		return FEditorStyle::GetBrush(TEXT("PropertyWindow.Locked"));
+	}
+	else
+	{
+		return FEditorStyle::GetBrush(TEXT("PropertyWindow.Unlocked"));
+	}
+}
+
 #undef LOCTEXT_NAMESPACE
diff --git a/Engine/Source/Editor/PropertyEditor/Private/SDetailsViewBase.cpp b/Engine/Source/Editor/PropertyEditor/Private/SDetailsViewBase.cpp
--- a/Engine/Source/Editor/PropertyEditor/Private/SDetailsViewBase.cpp
+++ b/Engine/Source/Editor/PropertyEditor/Private/SDetailsViewBase.cpp
@@ -437,6 +437,16 @@ void SDetailsViewBase::OnFilterTextChanged(const FText& InFilterText)
 
 }
 
+TSharedPtr<SWidget> SDetailsViewBase::GetNameAreaWidget()
+{
+	return DetailsViewArgs.bCustomNameAreaLocation ? NameArea : nullptr;
+}
+
+TSharedPtr<SWidget> SDetailsViewBase::GetFilterAreaWidget()
+{
+	return DetailsViewArgs.bCustomFilterAreaLocation ? FilterRow : nullptr;
+}
+
 /** 
  * Hides or shows properties based on the passed in filter text
  * 
diff --git a/Engine/Source/Editor/PropertyEditor/Private/SDetailsViewBase.h b/Engine/Source/Editor/PropertyEditor/Private/SDetailsViewBase.h
--- a/Engine/Source/Editor/PropertyEditor/Private/SDetailsViewBase.h
+++ b/Engine/Source/Editor/PropertyEditor/Private/SDetailsViewBase.h
@@ -379,6 +379,9 @@ protected:
 	/** Called when the list of currently differing properties changes */
 	virtual void UpdatePropertiesWhitelist(const TSet<FPropertyPath> InWhitelistedProperties) override { CurrentFilter.WhitelistedProperties = InWhitelistedProperties; }
 
+	virtual TSharedPtr<SWidget> GetNameAreaWidget() override;
+	virtual TSharedPtr<SWidget> GetFilterAreaWidget() override;
+
 	/** 
 	 * Hides or shows properties based on the passed in filter text
 	 * 
@@ -399,6 +402,8 @@ protected:
 	FCustomDetailLayoutMap InstancedClassToDetailLayoutMap;
 	/** The current detail layout based on selection */
 	TSharedPtr<class FDetailLayoutBuilderImpl> DetailLayout;
+	/** Row for searching and view options */
+	TSharedPtr<SHorizontalBox>  FilterRow;
 	/** Search box */
 	TSharedPtr<SWidget> SearchBox;
 	/** Customization class instances currently active in this view */
diff --git a/Engine/Source/Editor/PropertyEditor/Public/IDetailsView.h b/Engine/Source/Editor/PropertyEditor/Public/IDetailsView.h
--- a/Engine/Source/Editor/PropertyEditor/Public/IDetailsView.h
+++ b/Engine/Source/Editor/PropertyEditor/Public/IDetailsView.h
@@ -35,7 +35,10 @@ struct FDetailsViewArgs
 	uint32 bShowActorLabel : 1;
 	/** Bind this delegate to hide differing properties */
 	uint32 bShowDifferingPropertiesOption : 1;
-
+	/** If true the name area will be created but will not be displayed so it can be placed in a custom location.  */
+	uint32 bCustomNameAreaLocation : 1;
+	/** If true the filter area will be created but will not be displayed so it can be placed in a custom location.  */
+	uint32 bCustomFilterAreaLocation : 1;
 	/** Default constructor */
 	FDetailsViewArgs( const bool InUpdateFromSelection = false
 					, const bool InLockable = false
@@ -58,6 +61,8 @@ struct FDetailsViewArgs
 		, bShowModifiedPropertiesOption(true)
 		, bShowActorLabel(true)
 		, bShowDifferingPropertiesOption(false)
+		, bCustomNameAreaLocation(false)
+		, bCustomFilterAreaLocation(false)
 	{
 	}
 };
@@ -203,4 +208,10 @@ public:
 	 * Sets the set of properties that are considered differing, used when filtering out identical properties
 	 */
 	virtual void UpdatePropertiesWhitelist(const TSet<FPropertyPath> InWhitelistedProperties) = 0;
+
+	/** Returns the name area widget used to display object naming functionality so it can be placed in a custom location.  Note FDetailsViewArgs.bCustomNameAreaLocation must be true */
+	virtual TSharedPtr<SWidget> GetNameAreaWidget() = 0;
+
+	/** Returns the search area widget used to display search and view options so it can be placed in a custom location.  Note FDetailsViewArgs.bCustomFilterAreaLocation must be true */
+	virtual TSharedPtr<SWidget> GetFilterAreaWidget() = 0;
 };
