Fixed issue in SPropertyEditorAsset where Actor properties were not having their name fetched from the actor label if no PropertyEditor were specified.

diff --git a/Engine/Source/Editor/PropertyEditor/Private/UserInterface/PropertyEditor/SPropertyEditorAsset.cpp b/Engine/Source/Editor/PropertyEditor/Private/UserInterface/PropertyEditor/SPropertyEditorAsset.cpp
--- a/Engine/Source/Editor/PropertyEditor/Private/UserInterface/PropertyEditor/SPropertyEditorAsset.cpp
+++ b/Engine/Source/Editor/PropertyEditor/Private/UserInterface/PropertyEditor/SPropertyEditorAsset.cpp
@@ -483,40 +483,53 @@ FPropertyAccess::Result SPropertyEditorAsset::GetValue( FObjectOrAssetData& OutV
 	}
 	else
 	{
-		const FString CurrentObjectPath = ObjectPath.Get();
-		Result = FPropertyAccess::Success;
-
-		if( CurrentObjectPath != TEXT("None") && ( !CachedAssetData.IsValid() || CachedAssetData.ObjectPath.ToString() != CurrentObjectPath ) )
+		UObject* Object = NULL;
+		if (PropertyHandle.IsValid())
 		{
-			static FName AssetRegistryName("AssetRegistry");
+			Result = PropertyHandle->GetValue(Object);
+		}
 
-			FAssetRegistryModule& AssetRegistryModule = FModuleManager::Get().LoadModuleChecked<FAssetRegistryModule>(AssetRegistryName);
-			CachedAssetData = AssetRegistryModule.Get().GetAssetByObjectPath( *CurrentObjectPath );
+		if (Object != NULL)
+		{
+			OutValue = FObjectOrAssetData(Object);
+		}
+		else
+		{
+			const FString CurrentObjectPath = ObjectPath.Get();
+			Result = FPropertyAccess::Success;
 
-			if( PropertyHandle.IsValid() )
+			if (CurrentObjectPath != TEXT("None") && (!CachedAssetData.IsValid() || CachedAssetData.ObjectPath.ToString() != CurrentObjectPath))
 			{
-				// No property editor was specified so check if multiple property values are associated with the property handle
-				TArray<FString> ObjectValues;
-				PropertyHandle->GetPerObjectValues( ObjectValues );
+				static FName AssetRegistryName("AssetRegistry");
+
+				FAssetRegistryModule& AssetRegistryModule = FModuleManager::Get().LoadModuleChecked<FAssetRegistryModule>(AssetRegistryName);
+				CachedAssetData = AssetRegistryModule.Get().GetAssetByObjectPath(*CurrentObjectPath);
 
-				if( ObjectValues.Num() > 1 )
+				if (PropertyHandle.IsValid())
 				{
-					for( int32 ObjectIndex = 1; ObjectIndex < ObjectValues.Num() && Result == FPropertyAccess::Success; ++ObjectIndex )
+					// No property editor was specified so check if multiple property values are associated with the property handle
+					TArray<FString> ObjectValues;
+					PropertyHandle->GetPerObjectValues(ObjectValues);
+
+					if (ObjectValues.Num() > 1)
 					{
-						if( ObjectValues[ ObjectIndex ] != ObjectValues[ 0 ] )
+						for (int32 ObjectIndex = 1; ObjectIndex < ObjectValues.Num() && Result == FPropertyAccess::Success; ++ObjectIndex)
 						{
-							Result = FPropertyAccess::MultipleValues;
+							if (ObjectValues[ObjectIndex] != ObjectValues[0])
+							{
+								Result = FPropertyAccess::MultipleValues;
+							}
 						}
 					}
 				}
 			}
-		}
-		else if( CurrentObjectPath == TEXT("None") )
-		{
-			CachedAssetData = FAssetData();
-		}
+			else if (CurrentObjectPath == TEXT("None"))
+			{
+				CachedAssetData = FAssetData();
+			}
 
-		OutValue = FObjectOrAssetData( CachedAssetData );
+			OutValue = FObjectOrAssetData(CachedAssetData);
+		}
 	}
 
 	return Result;
