Fix for UE-7561: Even if a detail item is not editable, it can still be expanded (expander arrow was disabled because the containing widget was disabled

diff --git a/Engine/Source/Editor/PropertyEditor/Private/DetailItemNode.cpp b/Engine/Source/Editor/PropertyEditor/Private/DetailItemNode.cpp
--- a/Engine/Source/Editor/PropertyEditor/Private/DetailItemNode.cpp
+++ b/Engine/Source/Editor/PropertyEditor/Private/DetailItemNode.cpp
@@ -159,7 +159,6 @@ TSharedRef< ITableRow > FDetailItemNode::GenerateNodeWidget( const TSharedRef<ST
 	{
 		return
 			SNew(SDetailCategoryTableRow, AsShared(), OwnerTable)
-			.IsEnabled(IsParentEnabled)
 			.DisplayName(Customization.GetPropertyNode()->GetDisplayName())
 			.AddMetaData<FTagMetaData>(TagMeta)
 			.InnerCategory( true );
@@ -168,7 +167,6 @@ TSharedRef< ITableRow > FDetailItemNode::GenerateNodeWidget( const TSharedRef<ST
 	{
 		return
 			SNew(SDetailSingleItemRow, &Customization, HasMultiColumnWidget(), AsShared(), OwnerTable )
-			.IsEnabled( IsParentEnabled )
 			.AddMetaData<FTagMetaData>(TagMeta)
 			.ColumnSizeData(ColumnSizeData);
 	}
diff --git a/Engine/Source/Editor/PropertyEditor/Private/DetailItemNode.h b/Engine/Source/Editor/PropertyEditor/Private/DetailItemNode.h
--- a/Engine/Source/Editor/PropertyEditor/Private/DetailItemNode.h
+++ b/Engine/Source/Editor/PropertyEditor/Private/DetailItemNode.h
@@ -46,6 +46,7 @@ public:
 	virtual void SetIsHighlighted(bool bInIsHighlighted) { bIsHighlighted = bInIsHighlighted; }
 	virtual bool IsHighlighted() const { return bIsHighlighted; }
 	virtual bool IsLeaf() { return true; }
+	virtual TAttribute<bool> IsPropertyEditingEnabled() const override { return IsParentEnabled; }
 
 private:
 
diff --git a/Engine/Source/Editor/PropertyEditor/Private/IDetailTreeNode.h b/Engine/Source/Editor/PropertyEditor/Private/IDetailTreeNode.h
--- a/Engine/Source/Editor/PropertyEditor/Private/IDetailTreeNode.h
+++ b/Engine/Source/Editor/PropertyEditor/Private/IDetailTreeNode.h
@@ -101,5 +101,10 @@ public:
 	 * @return true if this is a leaf node:
 	 */
 	virtual bool IsLeaf() { return false; }
+
+	/**
+	* @return TAttribute indicating whether editing is enabled or whether the property is readonly:
+	*/
+	virtual TAttribute<bool> IsPropertyEditingEnabled() const { return false; }
 };
 
diff --git a/Engine/Source/Editor/PropertyEditor/Private/SDetailSingleItemRow.cpp b/Engine/Source/Editor/PropertyEditor/Private/SDetailSingleItemRow.cpp
--- a/Engine/Source/Editor/PropertyEditor/Private/SDetailSingleItemRow.cpp
+++ b/Engine/Source/Editor/PropertyEditor/Private/SDetailSingleItemRow.cpp
@@ -105,9 +105,12 @@ void SDetailSingleItemRow::Construct( const FArguments& InArgs, FDetailLayoutCus
 		}
 
 		TSharedRef<SWidget> KeyFrameButton = CreateKeyframeButton( *Customization, InOwnerTreeNode );
+		auto IsPropertyEditingEnabled = InOwnerTreeNode->IsPropertyEditingEnabled();
 
 		if( bHasMultipleColumns )
 		{
+			NameWidget->SetEnabled(IsPropertyEditingEnabled);
+			KeyFrameButton->SetEnabled(IsPropertyEditingEnabled);
 			Widget = 
 				SNew( SSplitter )
 				.Style( FEditorStyle::Get(), "DetailsView.Splitter" )
@@ -148,7 +151,7 @@ void SDetailSingleItemRow::Construct( const FArguments& InArgs, FDetailLayoutCus
 				.OnSlotResized( ColumnSizeData.OnWidthChanged )
 				[
 					SNew( SHorizontalBox )
-
+					.IsEnabled(IsPropertyEditingEnabled)
 					+ SHorizontalBox::Slot()
 					.Padding( DetailWidgetConstants::RightRowPadding )
 					.HAlign( Row.ValueWidget.HorizontalAlignment )
@@ -160,6 +163,8 @@ void SDetailSingleItemRow::Construct( const FArguments& InArgs, FDetailLayoutCus
 		}
 		else
 		{
+			Row.WholeRowWidget.Widget->SetEnabled(IsPropertyEditingEnabled);
+			KeyFrameButton->SetEnabled(IsPropertyEditingEnabled);
 			Widget =
 				SNew( SHorizontalBox )
 				+ SHorizontalBox::Slot()
