[INTEGRATE] Change 2411119 by [REDACTED_PERSON][REDACTED_HANDLE] on 2015/01/19 13:39:29

	Fixed up the details panel name area to be aware of in world component editing.
	Fixed up usage of the details panel name area to use new enum based usage flag for name areas

diff --git a/Engine/Plugins/Media/MediaPlayerEditor/Source/MediaPlayerEditor/Private/Widgets/SMediaPlayerEditorDetails.cpp b/Engine/Plugins/Media/MediaPlayerEditor/Source/MediaPlayerEditor/Private/Widgets/SMediaPlayerEditorDetails.cpp
--- a/Engine/Plugins/Media/MediaPlayerEditor/Source/MediaPlayerEditor/Private/Widgets/SMediaPlayerEditorDetails.cpp
+++ b/Engine/Plugins/Media/MediaPlayerEditor/Source/MediaPlayerEditor/Private/Widgets/SMediaPlayerEditorDetails.cpp
@@ -19,7 +19,6 @@ void SMediaPlayerEditorDetails::Construct( const FArguments& InArgs, UMediaPlaye
 		DetailsViewArgs.bAllowSearch = true;
 		DetailsViewArgs.bHideSelectionTip = true;
 		DetailsViewArgs.bLockable = false;
-		DetailsViewArgs.bObjectsUseNameArea = false;
 		DetailsViewArgs.bSearchInitialKeyFocus = true;
 		DetailsViewArgs.bUpdatesFromSelection = false;
 		DetailsViewArgs.bShowOptions = true;
diff --git a/Engine/Source/Developer/SettingsEditor/Private/Widgets/SSettingsEditor.cpp b/Engine/Source/Developer/SettingsEditor/Private/Widgets/SSettingsEditor.cpp
--- a/Engine/Source/Developer/SettingsEditor/Private/Widgets/SSettingsEditor.cpp
+++ b/Engine/Source/Developer/SettingsEditor/Private/Widgets/SSettingsEditor.cpp
@@ -40,7 +40,6 @@ void SSettingsEditor::Construct( const FArguments& InArgs, const ISettingsEditor
 		DetailsViewArgs.bAllowSearch = true;
 		DetailsViewArgs.bHideSelectionTip = true;
 		DetailsViewArgs.bLockable = false;
-		DetailsViewArgs.bObjectsUseNameArea = false;
 		DetailsViewArgs.bSearchInitialKeyFocus = true;
 		DetailsViewArgs.bUpdatesFromSelection = false;
 		DetailsViewArgs.NotifyHook = this;
diff --git a/Engine/Source/Editor/BehaviorTreeEditor/Private/BehaviorTreeEditor.cpp b/Engine/Source/Editor/BehaviorTreeEditor/Private/BehaviorTreeEditor.cpp
--- a/Engine/Source/Editor/BehaviorTreeEditor/Private/BehaviorTreeEditor.cpp
+++ b/Engine/Source/Editor/BehaviorTreeEditor/Private/BehaviorTreeEditor.cpp
@@ -709,7 +709,7 @@ TSharedRef<SWidget> FBehaviorTreeEditor::SpawnBlackboardDetails()
 	const bool bObjectsUseNameArea = false;
 	const bool bHideSelectionTip = true;
 	FPropertyEditorModule& PropertyEditorModule = FModuleManager::GetModuleChecked<FPropertyEditorModule>( "PropertyEditor" );
-	FDetailsViewArgs DetailsViewArgs( bIsUpdatable, bIsLockable, bAllowSearch, bObjectsUseNameArea, bHideSelectionTip );
+	FDetailsViewArgs DetailsViewArgs( bIsUpdatable, bIsLockable, bAllowSearch, FDetailsViewArgs::HideNameArea, bHideSelectionTip );
 	DetailsViewArgs.NotifyHook = this;
 	BlackboardDetailsView = PropertyEditorModule.CreateDetailView( DetailsViewArgs );
 
@@ -952,7 +952,7 @@ void FBehaviorTreeEditor::GetAbortModePreview(const class UBTDecorator* Decorato
 void FBehaviorTreeEditor::CreateInternalWidgets()
 {
 	FPropertyEditorModule& PropertyEditorModule = FModuleManager::GetModuleChecked<FPropertyEditorModule>( "PropertyEditor" );
-	FDetailsViewArgs DetailsViewArgs( false, false, true, true, false );
+	FDetailsViewArgs DetailsViewArgs( false, false, true, FDetailsViewArgs::HideNameArea, false );
 	DetailsViewArgs.NotifyHook = this;
 	DetailsView = PropertyEditorModule.CreateDetailView( DetailsViewArgs );
 	DetailsView->SetObject( NULL );
diff --git a/Engine/Source/Editor/BehaviorTreeEditor/Private/SBehaviorTreeDiff.cpp b/Engine/Source/Editor/BehaviorTreeEditor/Private/SBehaviorTreeDiff.cpp
--- a/Engine/Source/Editor/BehaviorTreeEditor/Private/SBehaviorTreeDiff.cpp
+++ b/Engine/Source/Editor/BehaviorTreeEditor/Private/SBehaviorTreeDiff.cpp
@@ -394,7 +394,7 @@ void SBehaviorTreeDiff::FBehaviorTreeDiffPanel::GeneratePanel(UEdGraph* Graph, U
 		];
 
 	FPropertyEditorModule& PropertyEditorModule = FModuleManager::GetModuleChecked<FPropertyEditorModule>( "PropertyEditor" );
-	const FDetailsViewArgs DetailsViewArgs( false, false, true, true, false );
+	const FDetailsViewArgs DetailsViewArgs( false, false, true, FDetailsViewArgs::ObjectsUseNameArea, false );
 	DetailsView = PropertyEditorModule.CreateDetailView( DetailsViewArgs );
 	DetailsView->SetObject( NULL );
 	DetailsView->SetIsPropertyVisibleDelegate(FIsPropertyVisible::CreateRaw(this, &SBehaviorTreeDiff::FBehaviorTreeDiffPanel::IsPropertyVisible));
diff --git a/Engine/Source/Editor/Blutility/Private/GlobalBlutilityDialog.cpp b/Engine/Source/Editor/Blutility/Private/GlobalBlutilityDialog.cpp
--- a/Engine/Source/Editor/Blutility/Private/GlobalBlutilityDialog.cpp
+++ b/Engine/Source/Editor/Blutility/Private/GlobalBlutilityDialog.cpp
@@ -145,7 +145,7 @@ void FGlobalBlutilityDialog::CreateInternalWidgets()
 
 	// Create a details view
 	FPropertyEditorModule& PropertyEditorModule = FModuleManager::GetModuleChecked<FPropertyEditorModule>("PropertyEditor");
-	const FDetailsViewArgs DetailsViewArgs(/*bUpdateFromSelection=*/ false, /*bLockable=*/ false, /*bAllowSearch=*/ false, /*bObjectsUseNameArea=*/ false, /*bHideSelectionTip=*/ true);
+	const FDetailsViewArgs DetailsViewArgs(/*bUpdateFromSelection=*/ false, /*bLockable=*/ false, /*bAllowSearch=*/ false, FDetailsViewArgs::HideNameArea, /*bHideSelectionTip=*/ true);
 	DetailsView = PropertyEditorModule.CreateDetailView(DetailsViewArgs);
 
 	DetailsView->SetIsPropertyVisibleDelegate(FIsPropertyVisible::CreateStatic(&Local::IsPropertyVisible));
diff --git a/Engine/Source/Editor/DeviceProfileEditor/Private/DetailsPanel/SDeviceProfileDetailsPanel.cpp b/Engine/Source/Editor/DeviceProfileEditor/Private/DetailsPanel/SDeviceProfileDetailsPanel.cpp
--- a/Engine/Source/Editor/DeviceProfileEditor/Private/DetailsPanel/SDeviceProfileDetailsPanel.cpp
+++ b/Engine/Source/Editor/DeviceProfileEditor/Private/DetailsPanel/SDeviceProfileDetailsPanel.cpp
@@ -73,7 +73,7 @@ void SDeviceProfileDetailsPanel::RefreshUI()
 		DetailsViewArgs.bAllowSearch = true;
 		DetailsViewArgs.bHideSelectionTip = true;
 		DetailsViewArgs.bLockable = false;
-		DetailsViewArgs.bObjectsUseNameArea = false;
+		DetailsViewArgs.NameAreaSettings = FDetailsViewArgs::HideNameArea;
 		DetailsViewArgs.bSearchInitialKeyFocus = true;
 		DetailsViewArgs.bUpdatesFromSelection = false;
 		DetailsViewArgs.bShowOptions = false;
diff --git a/Engine/Source/Editor/DeviceProfileEditor/Private/SDeviceProfileEditorSingleProfileView.cpp b/Engine/Source/Editor/DeviceProfileEditor/Private/SDeviceProfileEditorSingleProfileView.cpp
--- a/Engine/Source/Editor/DeviceProfileEditor/Private/SDeviceProfileEditorSingleProfileView.cpp
+++ b/Engine/Source/Editor/DeviceProfileEditor/Private/SDeviceProfileEditorSingleProfileView.cpp
@@ -19,7 +19,7 @@ void  SDeviceProfileEditorSingleProfileView::Construct(const FArguments& InArgs,
 		DetailsViewArgs.bSearchInitialKeyFocus = false;
 		DetailsViewArgs.bHideSelectionTip = true;
 		DetailsViewArgs.bLockable = false;
-		DetailsViewArgs.bObjectsUseNameArea = false;
+		DetailsViewArgs.NameAreaSettings = FDetailsViewArgs::HideNameArea;
 		DetailsViewArgs.bUpdatesFromSelection = false;
 		DetailsViewArgs.bShowOptions = false;
 	}
diff --git a/Engine/Source/Editor/EnvironmentQueryEditor/Private/EnvironmentQueryEditor.cpp b/Engine/Source/Editor/EnvironmentQueryEditor/Private/EnvironmentQueryEditor.cpp
--- a/Engine/Source/Editor/EnvironmentQueryEditor/Private/EnvironmentQueryEditor.cpp
+++ b/Engine/Source/Editor/EnvironmentQueryEditor/Private/EnvironmentQueryEditor.cpp
@@ -304,7 +304,7 @@ void FEnvironmentQueryEditor::OnSelectedNodesChanged(const TSet<class UObject*>&
 void FEnvironmentQueryEditor::CreateInternalWidgets()
 {
 	FPropertyEditorModule& PropertyEditorModule = FModuleManager::GetModuleChecked<FPropertyEditorModule>( "PropertyEditor" );
-	const FDetailsViewArgs DetailsViewArgs( false, false, true, true, false );
+	const FDetailsViewArgs DetailsViewArgs( false, false, true, FDetailsViewArgs::ObjectsUseNameArea, false );
 	DetailsView = PropertyEditorModule.CreateDetailView( DetailsViewArgs );
 	DetailsView->SetObject( NULL );
 	DetailsView->OnFinishedChangingProperties().AddSP(this, &FEnvironmentQueryEditor::OnFinishedChangingProperties);
diff --git a/Engine/Source/Editor/FoliageEdit/Private/SFoliageEditMeshDisplayItem.cpp b/Engine/Source/Editor/FoliageEdit/Private/SFoliageEditMeshDisplayItem.cpp
--- a/Engine/Source/Editor/FoliageEdit/Private/SFoliageEditMeshDisplayItem.cpp
+++ b/Engine/Source/Editor/FoliageEdit/Private/SFoliageEditMeshDisplayItem.cpp
@@ -62,7 +62,7 @@ void SFoliageEditMeshDisplayItem::Construct(const FArguments& InArgs)
 
 	// Create the details panels for the clustering tab.
 	FPropertyEditorModule& PropertyEditorModule = FModuleManager::GetModuleChecked<FPropertyEditorModule>("PropertyEditor");
-	const FDetailsViewArgs DetailsViewArgs(false, false, true, false, true, this);
+	const FDetailsViewArgs DetailsViewArgs(false, false, true, FDetailsViewArgs::HideNameArea, true, this);
 	ClusterSettingsDetails = PropertyEditorModule.CreateDetailView(DetailsViewArgs);
 	ClusterSettingsDetails->RegisterInstancedCustomPropertyLayout(UFoliageType::StaticClass(), FOnGetDetailCustomizationInstance::CreateStatic(&EmptyCustomization::MakeInstance));
 
diff --git a/Engine/Source/Editor/Kismet/Private/SKismetInspector.cpp b/Engine/Source/Editor/Kismet/Private/SKismetInspector.cpp
--- a/Engine/Source/Editor/Kismet/Private/SKismetInspector.cpp
+++ b/Engine/Source/Editor/Kismet/Private/SKismetInspector.cpp
@@ -202,8 +202,9 @@ void SKismetInspector::Construct(const FArguments& InArgs)
 	{
 		NotifyHook = Kismet2.Get();
 	}
-	FDetailsViewArgs DetailsViewArgs( /*bUpdateFromSelection=*/ false, /*bLockable=*/ false, /*bAllowSearch=*/ true, /*bObjectsUseNameArea=*/ true, /*bHideSelectionTip=*/ true, /*InNotifyHook=*/ NotifyHook, /*InSearchInitialKeyFocus=*/ false, /*InViewIdentifier=*/ InArgs._ViewIdentifier );
-	DetailsViewArgs.bHideActorNameArea = InArgs._HideNameArea;
+
+	FDetailsViewArgs::ENameAreaSettings NameAreaSettings = InArgs._HideNameArea ? FDetailsViewArgs::HideNameArea : FDetailsViewArgs::ObjectsUseNameArea;
+	FDetailsViewArgs DetailsViewArgs( /*bUpdateFromSelection=*/ false, /*bLockable=*/ false, /*bAllowSearch=*/ true, NameAreaSettings, /*bHideSelectionTip=*/ true, /*InNotifyHook=*/ NotifyHook, /*InSearchInitialKeyFocus=*/ false, /*InViewIdentifier=*/ InArgs._ViewIdentifier );
 
 	PropertyView = EditModule.CreateDetailView( DetailsViewArgs );
 		
diff --git a/Engine/Source/Editor/Kismet/Private/UserDefinedEnumEditor.cpp b/Engine/Source/Editor/Kismet/Private/UserDefinedEnumEditor.cpp
--- a/Engine/Source/Editor/Kismet/Private/UserDefinedEnumEditor.cpp
+++ b/Engine/Source/Editor/Kismet/Private/UserDefinedEnumEditor.cpp
@@ -85,8 +85,7 @@ TSharedRef<SDockTab> FUserDefinedEnumEditor::SpawnEnumeratorsTab(const FSpawnTab
 	// Create a property view
 	FPropertyEditorModule& EditModule = FModuleManager::Get().GetModuleChecked<FPropertyEditorModule>("PropertyEditor");
 
-	FDetailsViewArgs DetailsViewArgs( /*bUpdateFromSelection=*/ false, /*bLockable=*/ false, /*bAllowSearch=*/ false, /*bObjectsUseNameArea=*/ true, /*bHideSelectionTip=*/ true );
-	DetailsViewArgs.bHideActorNameArea = true;
+	FDetailsViewArgs DetailsViewArgs( /*bUpdateFromSelection=*/ false, /*bLockable=*/ false, /*bAllowSearch=*/ false, FDetailsViewArgs::HideNameArea, /*bHideSelectionTip=*/ true );
 	DetailsViewArgs.bShowOptions = false;
 
 	PropertyView = EditModule.CreateDetailView( DetailsViewArgs );
diff --git a/Engine/Source/Editor/Kismet/Private/UserDefinedStructureEditor.cpp b/Engine/Source/Editor/Kismet/Private/UserDefinedStructureEditor.cpp
--- a/Engine/Source/Editor/Kismet/Private/UserDefinedStructureEditor.cpp
+++ b/Engine/Source/Editor/Kismet/Private/UserDefinedStructureEditor.cpp
@@ -261,8 +261,7 @@ TSharedRef<SDockTab> FUserDefinedStructureEditor::SpawnStructureTab(const FSpawn
 	{
 		// Create a property view
 		FPropertyEditorModule& EditModule = FModuleManager::Get().GetModuleChecked<FPropertyEditorModule>("PropertyEditor");
-		FDetailsViewArgs DetailsViewArgs( /*bUpdateFromSelection=*/ false, /*bLockable=*/ false, /*bAllowSearch=*/ false, /*bObjectsUseNameArea=*/ true, /*bHideSelectionTip=*/ true);
-		DetailsViewArgs.bHideActorNameArea = true;
+		FDetailsViewArgs DetailsViewArgs( /*bUpdateFromSelection=*/ false, /*bLockable=*/ false, /*bAllowSearch=*/ false, FDetailsViewArgs::HideNameArea, /*bHideSelectionTip=*/ true);
 		DetailsViewArgs.bShowOptions = false;
 		PropertyView = EditModule.CreateDetailView(DetailsViewArgs);
 		FOnGetDetailCustomizationInstance LayoutStructDetails = FOnGetDetailCustomizationInstance::CreateStatic(&FUserDefinedStructureDetails::MakeInstance);
diff --git a/Engine/Source/Editor/KismetWidgets/Private/SSingleObjectDetailsPanel.cpp b/Engine/Source/Editor/KismetWidgets/Private/SSingleObjectDetailsPanel.cpp
--- a/Engine/Source/Editor/KismetWidgets/Private/SSingleObjectDetailsPanel.cpp
+++ b/Engine/Source/Editor/KismetWidgets/Private/SSingleObjectDetailsPanel.cpp
@@ -9,13 +9,12 @@
 /////////////////////////////////////////////////////
 // SSingleObjectDetailsPanel
 
-void SSingleObjectDetailsPanel::Construct(const FArguments& InArgs, bool bAutomaticallyObserveViaGetObjectToObserve)
+void SSingleObjectDetailsPanel::Construct(const FArguments& InArgs, bool bAutomaticallyObserveViaGetObjectToObserve, bool bAllowSearch)
 {
 	// Create a property view
 	FPropertyEditorModule& EditModule = FModuleManager::Get().GetModuleChecked<FPropertyEditorModule>("PropertyEditor");
 
-	FDetailsViewArgs DetailsViewArgs(/*bUpdateFromSelection=*/ false, /*bLockable=*/ false, /*bAllowSearch=*/ false, /*bObjectsUseNameArea=*/ true, /*bHideSelectionTip=*/ true);
-	DetailsViewArgs.bHideActorNameArea = true;
+	FDetailsViewArgs DetailsViewArgs(/*bUpdateFromSelection=*/ false, /*bLockable=*/ false, bAllowSearch, FDetailsViewArgs::HideNameArea, /*bHideSelectionTip=*/ true);
 
 	PropertyView = EditModule.CreateDetailView(DetailsViewArgs);
 	
diff --git a/Engine/Source/Editor/KismetWidgets/Public/SSingleObjectDetailsPanel.h b/Engine/Source/Editor/KismetWidgets/Public/SSingleObjectDetailsPanel.h
--- a/Engine/Source/Editor/KismetWidgets/Public/SSingleObjectDetailsPanel.h
+++ b/Engine/Source/Editor/KismetWidgets/Public/SSingleObjectDetailsPanel.h
@@ -11,7 +11,7 @@ public:
 	SLATE_BEGIN_ARGS(SSingleObjectDetailsPanel) {}
 	SLATE_END_ARGS()
 
-	void Construct(const FArguments& InArgs, bool bAutomaticallyObserveViaGetObjectToObserve = true);
+	void Construct(const FArguments& InArgs, bool bAutomaticallyObserveViaGetObjectToObserve = true, bool bAllowSearch = false);
 
 	// SWidget interface
 	virtual void Tick( const FGeometry& AllottedGeometry, const double InCurrentTime, const float InDeltaTime ) override;
diff --git a/Engine/Source/Editor/LandscapeEditor/Private/SLandscapeEditor.cpp b/Engine/Source/Editor/LandscapeEditor/Private/SLandscapeEditor.cpp
--- a/Engine/Source/Editor/LandscapeEditor/Private/SLandscapeEditor.cpp
+++ b/Engine/Source/Editor/LandscapeEditor/Private/SLandscapeEditor.cpp
@@ -129,8 +129,8 @@ void SLandscapeEditor::Construct(const FArguments& InArgs, TSharedRef<FLandscape
 	}
 
 	FPropertyEditorModule& PropertyEditorModule = FModuleManager::LoadModuleChecked<FPropertyEditorModule>("PropertyEditor");
-	FDetailsViewArgs DetailsViewArgs(false, false, false);
-	DetailsViewArgs.bHideActorNameArea = true;
+	FDetailsViewArgs DetailsViewArgs(false, false, false,FDetailsViewArgs::HideNameArea);
+
 	DetailsPanel = PropertyEditorModule.CreateDetailView(DetailsViewArgs);
 	DetailsPanel->SetIsPropertyVisibleDelegate(FIsPropertyVisible::CreateSP(this, &SLandscapeEditor::GetIsPropertyVisible));
 
diff --git a/Engine/Source/Editor/LevelEditor/Private/ComponentsTree.cpp b/Engine/Source/Editor/LevelEditor/Private/ComponentsTree.cpp
--- a/Engine/Source/Editor/LevelEditor/Private/ComponentsTree.cpp
+++ b/Engine/Source/Editor/LevelEditor/Private/ComponentsTree.cpp
@@ -95,8 +95,7 @@ static bool IsPropertyVisible( const FPropertyAndParent& PropertyAndParent )
 void SComponentsTree::Construct(const FArguments& InArgs, TSharedPtr<class IDetailsView> InPropertyView )
 {
 	FPropertyEditorModule& EditModule = FModuleManager::Get().GetModuleChecked<FPropertyEditorModule>("PropertyEditor");
-	FDetailsViewArgs DetailsViewArgs( /*bUpdateFromSelection=*/ false, /*bLockable=*/ false, /*bAllowSearch=*/ true, /*bObjectsUseNameArea=*/ true, /*bHideSelectionTip=*/ true );
-	DetailsViewArgs.bHideActorNameArea = true;
+	FDetailsViewArgs DetailsViewArgs( /*bUpdateFromSelection=*/ false, /*bLockable=*/ false, /*bAllowSearch=*/ true, FDetailsViewArgs::HideNameArea, /*bHideSelectionTip=*/ true );
 
 	bSelectionGuard = false;
 	USelection::SelectionChangedEvent.AddRaw(this, &SComponentsTree::OnEditorSelectionChanged);
diff --git a/Engine/Source/Editor/LevelEditor/Private/SLevelEditor.cpp b/Engine/Source/Editor/LevelEditor/Private/SLevelEditor.cpp
--- a/Engine/Source/Editor/LevelEditor/Private/SLevelEditor.cpp
+++ b/Engine/Source/Editor/LevelEditor/Private/SLevelEditor.cpp
@@ -534,7 +534,7 @@ public:
 		FDetailsViewArgs DetailsViewArgs;
 		DetailsViewArgs.bUpdatesFromSelection = true;
 		DetailsViewArgs.bLockable = true;
-		DetailsViewArgs.bObjectsUseNameArea = false;
+		DetailsViewArgs.NameAreaSettings = FDetailsViewArgs::ComponentsAndActorsUseNameArea;
 		DetailsViewArgs.NotifyHook = GUnrealEd;
 		DetailsViewArgs.ViewIdentifier = TabIdentifier;
 		DetailsViewArgs.bCustomNameAreaLocation = true;
@@ -968,7 +968,7 @@ TSharedRef<SDockTab> SLevelEditor::SpawnLevelEditorTab( const FSpawnTabArgs& Arg
 	else if ( TabIdentifier == "WorldSettingsTab" )
 	{
 		FPropertyEditorModule& PropPlugin = FModuleManager::LoadModuleChecked<FPropertyEditorModule>("PropertyEditor");
-		FDetailsViewArgs DetailsViewArgs( false, false, true, false, false, GUnrealEd );
+		FDetailsViewArgs DetailsViewArgs( false, false, true, FDetailsViewArgs::HideNameArea, false, GUnrealEd );
 		DetailsViewArgs.bShowActorLabel = false;
 
 		WorldSettingsView = PropPlugin.CreateDetailView( DetailsViewArgs );
diff --git a/Engine/Source/Editor/MaterialEditor/Private/MaterialEditor.cpp b/Engine/Source/Editor/MaterialEditor/Private/MaterialEditor.cpp
--- a/Engine/Source/Editor/MaterialEditor/Private/MaterialEditor.cpp
+++ b/Engine/Source/Editor/MaterialEditor/Private/MaterialEditor.cpp
@@ -599,7 +599,7 @@ void FMaterialEditor::CreateInternalWidgets()
 	// Manually set zoom level to avoid deferred zooming
 	GraphEditor->SetViewLocation(FVector2D::ZeroVector, 1);
 
-	const FDetailsViewArgs DetailsViewArgs( false, false, true, false, true, this );
+	const FDetailsViewArgs DetailsViewArgs( false, false, true, FDetailsViewArgs::HideNameArea, true, this );
 	MaterialDetailsView = PropertyEditorModule.CreateDetailView( DetailsViewArgs );
 
 	FOnGetDetailCustomizationInstance LayoutExpressionParameterDetails = FOnGetDetailCustomizationInstance::CreateStatic(
diff --git a/Engine/Source/Editor/MaterialEditor/Private/MaterialInstanceEditor.cpp b/Engine/Source/Editor/MaterialEditor/Private/MaterialInstanceEditor.cpp
--- a/Engine/Source/Editor/MaterialEditor/Private/MaterialInstanceEditor.cpp
+++ b/Engine/Source/Editor/MaterialEditor/Private/MaterialInstanceEditor.cpp
@@ -450,7 +450,7 @@ void FMaterialInstanceEditor::CreateInternalWidgets()
 		.MaterialEditor(SharedThis(this));
 
 	FPropertyEditorModule& PropertyEditorModule = FModuleManager::GetModuleChecked<FPropertyEditorModule>( "PropertyEditor" );
-	const FDetailsViewArgs DetailsViewArgs( false, false, true, false, true, this );
+	const FDetailsViewArgs DetailsViewArgs( false, false, true, FDetailsViewArgs::HideNameArea, true, this );
 	MaterialInstanceDetails = PropertyEditorModule.CreateDetailView( DetailsViewArgs );
 	FOnGetDetailCustomizationInstance LayoutMICDetails = FOnGetDetailCustomizationInstance::CreateStatic( 
 		&FMaterialInstanceParameterDetails::MakeInstance, MaterialEditorInstance, FGetShowHiddenParameters::CreateSP(this, &FMaterialInstanceEditor::GetShowHiddenParameters) );
diff --git a/Engine/Source/Editor/NiagaraEditor/Private/NiagaraEditor.cpp b/Engine/Source/Editor/NiagaraEditor/Private/NiagaraEditor.cpp
--- a/Engine/Source/Editor/NiagaraEditor/Private/NiagaraEditor.cpp
+++ b/Engine/Source/Editor/NiagaraEditor/Private/NiagaraEditor.cpp
@@ -85,7 +85,7 @@ void FNiagaraEditor::InitNiagaraEditor( const EToolkitMode::Type Mode, const TSh
 	);
 
 	FPropertyEditorModule& PropertyEditorModule = FModuleManager::GetModuleChecked<FPropertyEditorModule>("PropertyEditor");
-	const FDetailsViewArgs DetailsViewArgs(false, false, true, false, true, this);
+	const FDetailsViewArgs DetailsViewArgs(false, false, true, FDetailsViewArgs::HideNameArea, true, this);
 	NiagaraDetailsView = PropertyEditorModule.CreateDetailView(DetailsViewArgs);
 
 	const bool bCreateDefaultStandaloneMenu = true;
diff --git a/Engine/Source/Editor/Persona/Private/MeshMode/MeshMode.cpp b/Engine/Source/Editor/Persona/Private/MeshMode/MeshMode.cpp
--- a/Engine/Source/Editor/Persona/Private/MeshMode/MeshMode.cpp
+++ b/Engine/Source/Editor/Persona/Private/MeshMode/MeshMode.cpp
@@ -31,7 +31,7 @@ public:
 	{
 		PersonaPtr = InPersona;
 
-		SSingleObjectDetailsPanel::Construct(SSingleObjectDetailsPanel::FArguments());
+		SSingleObjectDetailsPanel::Construct(SSingleObjectDetailsPanel::FArguments(), true, true);
 
 		PropertyView->SetGenericLayoutDetailsDelegate( FOnGetDetailCustomizationInstance::CreateStatic( &FPersonaMeshDetails::MakeInstance, InPersona ) );
 	}
diff --git a/Engine/Source/Editor/PropertyEditor/Private/SDetailNameArea.cpp b/Engine/Source/Editor/PropertyEditor/Private/SDetailNameArea.cpp
--- a/Engine/Source/Editor/PropertyEditor/Private/SDetailNameArea.cpp
+++ b/Engine/Source/Editor/PropertyEditor/Private/SDetailNameArea.cpp
@@ -16,7 +16,6 @@ void SDetailNameArea::Construct( const FArguments& InArgs, const TArray< TWeakOb
 	SelectionTip = InArgs._SelectionTip;
 	bShowLockButton = InArgs._ShowLockButton;
 	bShowActorLabel = InArgs._ShowActorLabel;
-	Refresh( *SelectedObjects );
 }
 
 void SDetailNameArea::Refresh( const TArray< TWeakObjectPtr<UObject> >& SelectedObjects )
@@ -27,17 +26,41 @@ void SDetailNameArea::Refresh( const TArray< TWeakObjectPtr<UObject> >& Selected
 	];
 }
 
-void SDetailNameArea::Refresh( const TArray< TWeakObjectPtr<AActor> >& SelectedActors )
+void SDetailNameArea::Refresh( const TArray< TWeakObjectPtr<AActor> >& SelectedActors, const TArray< TWeakObjectPtr<UObject> >& SelectedObjects, FDetailsViewArgs::ENameAreaSettings NameAreaSettings  )
 {
 	// Convert the actor array to base object type
-	TArray< TWeakObjectPtr<UObject> > SelectedObjects;
-	for( int32 ActorIndex = 0 ; ActorIndex < SelectedActors.Num() ; ++ActorIndex )
+
+	TArray< TWeakObjectPtr<UObject> > FinalSelectedObjects;
+	if(NameAreaSettings == FDetailsViewArgs::ActorsUseNameArea)
 	{
-		const TWeakObjectPtr<UObject> ObjectWeakPtr = SelectedActors[ActorIndex].Get();
-		SelectedObjects.Add( ObjectWeakPtr );
+		for(auto Actor : SelectedActors)
+		{
+			const TWeakObjectPtr<UObject> ObjectWeakPtr = Actor.Get();
+			FinalSelectedObjects.Add(ObjectWeakPtr);
+		}
 	}
+	else if( NameAreaSettings == FDetailsViewArgs::ComponentsAndActorsUseNameArea )
+	{
+		for(auto Actor : SelectedActors)
+		{
+			const TWeakObjectPtr<UObject> ObjectWeakPtr = Actor.Get();
+			FinalSelectedObjects.Add(ObjectWeakPtr);
+		}
 
-	Refresh( SelectedObjects );
+		// Note: assumes that actors and components are not selected together.
+		if( FinalSelectedObjects.Num() == 0 )
+		{
+			for(auto Object : SelectedObjects)
+			{
+				UActorComponent* ActorComp = Cast<UActorComponent>(Object.Get());
+				if(ActorComp && ActorComp->GetOwner())
+				{
+					FinalSelectedObjects.AddUnique(ActorComp->GetOwner());
+				}
+			}
+		}
+	}
+	Refresh( FinalSelectedObjects );
 }
 
 const FSlateBrush* SDetailNameArea::OnGetLockButtonImageResource() const
diff --git a/Engine/Source/Editor/PropertyEditor/Private/SDetailNameArea.h b/Engine/Source/Editor/PropertyEditor/Private/SDetailNameArea.h
--- a/Engine/Source/Editor/PropertyEditor/Private/SDetailNameArea.h
+++ b/Engine/Source/Editor/PropertyEditor/Private/SDetailNameArea.h
@@ -30,7 +30,7 @@ public:
 	 *
 	 * @param SelectedActors	the new list of selected actors
 	 */
-	void Refresh( const TArray< TWeakObjectPtr<AActor> >& SelectedActors );
+	void Refresh( const TArray< TWeakObjectPtr<AActor> >& SelectedActors, const TArray< TWeakObjectPtr<UObject> >& SelectedObjects, FDetailsViewArgs::ENameAreaSettings NameAreaSettings  );
 
 private:
 	/** @return the Slate brush to use for the lock image */
diff --git a/Engine/Source/Editor/PropertyEditor/Private/SDetailsView.cpp b/Engine/Source/Editor/PropertyEditor/Private/SDetailsView.cpp
--- a/Engine/Source/Editor/PropertyEditor/Private/SDetailsView.cpp
+++ b/Engine/Source/Editor/PropertyEditor/Private/SDetailsView.cpp
@@ -257,7 +257,7 @@ FReply SDetailsView::OnOpenRawPropertyEditorClicked()
 
 EVisibility SDetailsView::GetActorNameAreaVisibility() const
 {
-	const bool bVisible = !DetailsViewArgs.bHideActorNameArea && !bViewingClassDefaultObject;
+	const bool bVisible = !DetailsViewArgs.NameAreaSettings == FDetailsViewArgs::HideNameArea && !bViewingClassDefaultObject;
 	return bVisible ? EVisibility::Visible : EVisibility::Collapsed; 
 }
 
@@ -420,13 +420,13 @@ void SDetailsView::SetObjectArrayPrivate( const TArray< TWeakObjectPtr< UObject
 	}
 
 	// Selection changed, refresh the detail area
-	if ( DetailsViewArgs.bObjectsUseNameArea )
+	if ( DetailsViewArgs.NameAreaSettings != FDetailsViewArgs::ActorsUseNameArea && DetailsViewArgs.NameAreaSettings != FDetailsViewArgs::ComponentsAndActorsUseNameArea )
 	{
 		NameArea->Refresh( SelectedObjects );
 	}
 	else
 	{
-		NameArea->Refresh( SelectedActors );
+		NameArea->Refresh( SelectedActors, SelectedObjects, DetailsViewArgs.NameAreaSettings );
 	}
 	
 	// When selection changes rebuild information about the selection
diff --git a/Engine/Source/Editor/PropertyEditor/Public/IDetailsView.h b/Engine/Source/Editor/PropertyEditor/Public/IDetailsView.h
--- a/Engine/Source/Editor/PropertyEditor/Public/IDetailsView.h
+++ b/Engine/Source/Editor/PropertyEditor/Public/IDetailsView.h
@@ -9,22 +9,31 @@
  */
 struct FDetailsViewArgs
 {
+	enum ENameAreaSettings
+	{
+		/** The name area should never be displayed */
+		HideNameArea,
+		/** All object types use name area */
+		ObjectsUseNameArea,
+		/** Only Actors use name area */
+		ActorsUseNameArea,
+		/** Components and actors use the name area. Components will display their actor owner as the name */
+		ComponentsAndActorsUseNameArea,
+	};
 	/** Identifier for this details view; NAME_None if this view is anonymous */
 	FName ViewIdentifier;
 	/** Notify hook to call when properties are changed */
 	FNotifyHook* NotifyHook;
+	/** Settings for displaying the name area */
+	ENameAreaSettings NameAreaSettings;
 	/** True if the viewed objects updates from editor selection */
 	uint32 bUpdatesFromSelection : 1;
 	/** True if this property view can be locked */
 	uint32 bLockable : 1;
 	/** True if we allow searching */
 	uint32 bAllowSearch : 1;
-	/** True if object selection wants to use the name area */
-	uint32 bObjectsUseNameArea : 1;
 	/** True if you want to not show the tip when no objects are selected (should only be used if viewing actors properties or bObjectsUseNameArea is true ) */
 	uint32 bHideSelectionTip : 1;
-	/** True if you want to hide the object/objects selected info area */
-	uint32 bHideActorNameArea : 1;
 	/** True if you want the search box to have initial keyboard focus */
 	uint32 bSearchInitialKeyFocus : 1;
 	/** Allow options to be changed */
@@ -43,19 +52,18 @@ struct FDetailsViewArgs
 	FDetailsViewArgs( const bool InUpdateFromSelection = false
 					, const bool InLockable = false
 					, const bool InAllowSearch = true
-					, const bool InObjectsUseNameArea = false
+					, const ENameAreaSettings InNameAreaSettings = ActorsUseNameArea
 					, const bool InHideSelectionTip = false
 					, FNotifyHook* InNotifyHook = NULL
 					, const bool InSearchInitialKeyFocus = false
 					, FName InViewIdentifier = NAME_None )
 		: ViewIdentifier( InViewIdentifier )
-		, NotifyHook( InNotifyHook ) 
+		, NotifyHook( InNotifyHook )
+		, NameAreaSettings( InNameAreaSettings )
 		, bUpdatesFromSelection( InUpdateFromSelection )
 		, bLockable(InLockable)
 		, bAllowSearch( InAllowSearch )
-		, bObjectsUseNameArea( InObjectsUseNameArea )
 		, bHideSelectionTip( InHideSelectionTip )
-		, bHideActorNameArea( false )
 		, bSearchInitialKeyFocus( InSearchInitialKeyFocus )
 		, bShowOptions( true )
 		, bShowModifiedPropertiesOption(true)
diff --git a/Engine/Source/Editor/SoundClassEditor/Private/SoundClassEditor.cpp b/Engine/Source/Editor/SoundClassEditor/Private/SoundClassEditor.cpp
--- a/Engine/Source/Editor/SoundClassEditor/Private/SoundClassEditor.cpp
+++ b/Engine/Source/Editor/SoundClassEditor/Private/SoundClassEditor.cpp
@@ -202,7 +202,7 @@ void FSoundClassEditor::CreateInternalWidgets()
 	GraphEditor = CreateGraphEditorWidget();
 
 	FPropertyEditorModule& PropertyEditorModule = FModuleManager::GetModuleChecked<FPropertyEditorModule>( "PropertyEditor" );
-	const FDetailsViewArgs DetailsViewArgs( false, false, true, true, false );
+	const FDetailsViewArgs DetailsViewArgs( false, false, true, FDetailsViewArgs::ObjectsUseNameArea, false );
 	DetailsView = PropertyEditorModule.CreateDetailView( DetailsViewArgs );
 	DetailsView->SetObject( SoundClass );
 }
diff --git a/Engine/Source/Editor/StaticMeshEditor/Private/StaticMeshEditor.cpp b/Engine/Source/Editor/StaticMeshEditor/Private/StaticMeshEditor.cpp
--- a/Engine/Source/Editor/StaticMeshEditor/Private/StaticMeshEditor.cpp
+++ b/Engine/Source/Editor/StaticMeshEditor/Private/StaticMeshEditor.cpp
@@ -113,8 +113,7 @@ void FStaticMeshEditor::InitStaticMeshEditor( const EToolkitMode::Type Mode, con
 	DetailsViewArgs.bAllowSearch = true;
 	DetailsViewArgs.bLockable = false;
 	DetailsViewArgs.bUpdatesFromSelection = false;
-	DetailsViewArgs.bHideActorNameArea = true;
-	DetailsViewArgs.bObjectsUseNameArea = false;
+	DetailsViewArgs.NameAreaSettings = FDetailsViewArgs::HideNameArea;
 	DetailsViewArgs.NotifyHook = this;
 
 	StaticMeshDetailsView = PropertyEditorModule.CreateDetailView( DetailsViewArgs );
diff --git a/Engine/Source/Editor/UMGEditor/Private/Details/SWidgetDetailsView.cpp b/Engine/Source/Editor/UMGEditor/Private/Details/SWidgetDetailsView.cpp
--- a/Engine/Source/Editor/UMGEditor/Private/Details/SWidgetDetailsView.cpp
+++ b/Engine/Source/Editor/UMGEditor/Private/Details/SWidgetDetailsView.cpp
@@ -38,10 +38,7 @@ void SWidgetDetailsView::Construct(const FArguments& InArgs, TSharedPtr<FWidgetB
 	FPropertyEditorModule& EditModule = FModuleManager::Get().GetModuleChecked<FPropertyEditorModule>("PropertyEditor");
 
 	FNotifyHook* NotifyHook = this;
-	FDetailsViewArgs DetailsViewArgs( /*bUpdateFromSelection=*/ false, /*bLockable=*/ false, /*bAllowSearch=*/ true, /*bObjectsUseNameArea=*/ true, /*bHideSelectionTip=*/ true, /*InNotifyHook=*/ NotifyHook, /*InSearchInitialKeyFocus=*/ false, /*InViewIdentifier=*/ NAME_None);
-
-	// We hide the actor name area because we're providing our own.
-	DetailsViewArgs.bHideActorNameArea = true;
+	FDetailsViewArgs DetailsViewArgs( /*bUpdateFromSelection=*/ false, /*bLockable=*/ false, /*bAllowSearch=*/ true, FDetailsViewArgs::HideNameArea, /*bHideSelectionTip=*/ true, /*InNotifyHook=*/ NotifyHook, /*InSearchInitialKeyFocus=*/ false, /*InViewIdentifier=*/ NAME_None);
 
 	PropertyView = EditModule.CreateDetailView(DetailsViewArgs);
 
diff --git a/Engine/Source/Editor/UnrealEd/Private/SSocketManager.cpp b/Engine/Source/Editor/UnrealEd/Private/SSocketManager.cpp
--- a/Engine/Source/Editor/UnrealEd/Private/SSocketManager.cpp
+++ b/Engine/Source/Editor/UnrealEd/Private/SSocketManager.cpp
@@ -185,9 +185,10 @@ void SSocketManager::Construct(const FArguments& InArgs)
 	FDetailsViewArgs Args;
 	Args.bHideSelectionTip = true;
 	Args.bLockable = false;
-	Args.bAllowSearch = true;
+	Args.bAllowSearch = false;
+	Args.bShowOptions = false;
 	Args.NotifyHook = this;
-	Args.bObjectsUseNameArea = true;
+	Args.NameAreaSettings = FDetailsViewArgs::HideNameArea;
 
 	FPropertyEditorModule& PropertyModule = FModuleManager::LoadModuleChecked<FPropertyEditorModule>("PropertyEditor");
 	SocketDetailsView = PropertyModule.CreateDetailView(Args);
diff --git a/Engine/Source/Editor/UnrealEd/Private/Toolkits/SimpleAssetEditor.cpp b/Engine/Source/Editor/UnrealEd/Private/Toolkits/SimpleAssetEditor.cpp
--- a/Engine/Source/Editor/UnrealEd/Private/Toolkits/SimpleAssetEditor.cpp
+++ b/Engine/Source/Editor/UnrealEd/Private/Toolkits/SimpleAssetEditor.cpp
@@ -49,7 +49,7 @@ void FSimpleAssetEditor::InitEditor( const EToolkitMode::Type Mode, const TShare
 	const bool bIsLockable = false;
 
 	FPropertyEditorModule& PropertyEditorModule = FModuleManager::GetModuleChecked<FPropertyEditorModule>( "PropertyEditor" );
-	const FDetailsViewArgs DetailsViewArgs( bIsUpdatable, bIsLockable, true, true, false );
+	const FDetailsViewArgs DetailsViewArgs( bIsUpdatable, bIsLockable, true, FDetailsViewArgs::ObjectsUseNameArea, false );
 	DetailsView = PropertyEditorModule.CreateDetailView( DetailsViewArgs );
 	const TSharedRef<FTabManager::FLayout> StandaloneDefaultLayout = FTabManager::NewLayout( "Standalone_SimpleAssetEditor_Layout_v3" )
 	->AddArea
diff --git a/Engine/Source/Editor/WorldBrowser/Private/SWorldDetails.cpp b/Engine/Source/Editor/WorldBrowser/Private/SWorldDetails.cpp
--- a/Engine/Source/Editor/WorldBrowser/Private/SWorldDetails.cpp
+++ b/Engine/Source/Editor/WorldBrowser/Private/SWorldDetails.cpp
@@ -58,7 +58,7 @@ void SWorldDetails::OnBrowseWorld(UWorld* InWorld)
 		WorldModel->CollectionChanged.AddSP(this, &SWorldDetails::OnCollectionChanged);
 	
 		FPropertyEditorModule& PropertyModule = FModuleManager::LoadModuleChecked<FPropertyEditorModule>("PropertyEditor");
-		FDetailsViewArgs Args(false, false, false, false, true);
+		FDetailsViewArgs Args(false, false, false, FDetailsViewArgs::HideNameArea, true);
 		Args.bShowActorLabel = false;
 	
 		DetailsView = PropertyModule.CreateDetailView(Args);
