[INTEGRATE] Change 2411255 by [REDACTED_PERSON][REDACTED_HANDLE] on 2015/01/19 15:09:23

	Fix logic bug which was breaking mac compile

diff --git a/Engine/Source/Editor/PropertyEditor/Private/SDetailsView.cpp b/Engine/Source/Editor/PropertyEditor/Private/SDetailsView.cpp
--- a/Engine/Source/Editor/PropertyEditor/Private/SDetailsView.cpp
+++ b/Engine/Source/Editor/PropertyEditor/Private/SDetailsView.cpp
@@ -257,7 +257,7 @@ FReply SDetailsView::OnOpenRawPropertyEditorClicked()
 
 EVisibility SDetailsView::GetActorNameAreaVisibility() const
 {
-	const bool bVisible = !DetailsViewArgs.NameAreaSettings == FDetailsViewArgs::HideNameArea && !bViewingClassDefaultObject;
+	const bool bVisible = DetailsViewArgs.NameAreaSettings != FDetailsViewArgs::HideNameArea && !bViewingClassDefaultObject;
 	return bVisible ? EVisibility::Visible : EVisibility::Collapsed; 
 }
 
