[INTEGRATE] Change 2412372 by [REDACTED_PERSON][REDACTED_HANDLE] on 2015/01/20 12:00:05

	Fixed missing padding between selected actor icon and label text

diff --git a/Engine/Source/Editor/PropertyEditor/Private/SDetailNameArea.cpp b/Engine/Source/Editor/PropertyEditor/Private/SDetailNameArea.cpp
--- a/Engine/Source/Editor/PropertyEditor/Private/SDetailNameArea.cpp
+++ b/Engine/Source/Editor/PropertyEditor/Private/SDetailNameArea.cpp
@@ -109,7 +109,7 @@ TSharedRef< SWidget > SDetailNameArea::BuildObjectNameArea( const TArray< TWeakO
 		.AutoWidth()
 		.HAlign(HAlign_Left)
 		.VAlign(VAlign_Center)
-		.Padding(0,0,0,0)
+		.Padding(0,0,6,0)
 		[
 			SNew(SImage)
 			.Image(ActorIcon)
