[INTEGRATE] Change 2416422 by [REDACTED_PERSON][REDACTED_HANDLE] on 2015/01/22 23:41:13

	A few aesthetic changes for the component tree and details panel:
	- Made header area of component tree darker to match details panel style
	- Made details panel not have empty space where scrollbar should be.  Scrollbar no longer takes up permanent space

	- Removed old components tree

diff --git a/Engine/Source/Editor/EditorStyle/Private/SlateEditorStyle.cpp b/Engine/Source/Editor/EditorStyle/Private/SlateEditorStyle.cpp
--- a/Engine/Source/Editor/EditorStyle/Private/SlateEditorStyle.cpp
+++ b/Engine/Source/Editor/EditorStyle/Private/SlateEditorStyle.cpp
@@ -3918,6 +3918,16 @@ void FSlateEditorStyle::FStyle::SetupGraphEditorStyles()
 
 		Set("SCSEditor.PromoteToBlueprintIcon", new IMAGE_BRUSH("Icons/AssetIcons/Blueprint_16x", Icon16x16) );
 
+		Set("SCSEditor.TopBar.Font", FTextBlockStyle(NormalText)
+			.SetFont(TTF_CORE_FONT("Fonts/Roboto-Bold", 10))
+			.SetColorAndOpacity(FLinearColor(1.0f, 1.0f, 1.0f))
+			.SetHighlightColor(FLinearColor(1.0f, 1.0f, 1.0f))
+			.SetShadowOffset(FVector2D(1, 1))
+			.SetShadowColorAndOpacity(FLinearColor(0, 0, 0, 0.9f)));
+
+		Set( "SCSEditor.TreePanel", new BOX_BRUSH( "Common/GroupBorder_FlatTop", FMargin(4.0f/16.0f) ) );
+
+
 		// Timeline Editor
 		{
 			Set( "TimelineEditor.AddFloatTrack", new IMAGE_BRUSH( "Icons/icon_TrackAddFloat_36x24px", Icon36x24, FLinearColor::Black ) );
diff --git a/Engine/Source/Editor/Kismet/Private/SSCSEditor.cpp b/Engine/Source/Editor/Kismet/Private/SSCSEditor.cpp
--- a/Engine/Source/Editor/Kismet/Private/SSCSEditor.cpp
+++ b/Engine/Source/Editor/Kismet/Private/SSCSEditor.cpp
@@ -2536,12 +2536,13 @@ void SSCSEditor::Construct( const FArguments& InArgs )
 			+ SVerticalBox::Slot()
 			.AutoHeight()
 			.VAlign(VAlign_Top)
-			.Padding(0.0f, 2.0f)
+			.Padding(0)
 			[
 				SNew(SBorder)
 				.Padding(2.0f)
-				.BorderImage(FEditorStyle::GetBrush("ToolPanel.GroupBorder"))
+				.BorderImage(FEditorStyle::GetBrush("DetailsView.CategoryTop"))
 				.AddMetaData<FTagMetaData>(FTagMetaData(TEXT("ComponentsPanel")))
+				.BorderBackgroundColor( FLinearColor( .6,.6,.6, 1.0f ) )
 				[
 					SNew(SHorizontalBox)
 					+ SHorizontalBox::Slot()
@@ -2629,7 +2630,7 @@ void SSCSEditor::Construct( const FArguments& InArgs )
 			[
 				SNew(SBorder)
 				.Padding(2.0f)
-				.BorderImage(FEditorStyle::GetBrush("ToolPanel.GroupBorder"))
+				.BorderImage(FEditorStyle::GetBrush("SCSEditor.TreePanel"))
 				.AddMetaData<FTagMetaData>(FTagMetaData(TEXT("ComponentsPanel")))
 				
 				[
diff --git a/Engine/Source/Editor/LevelEditor/Private/ComponentsTree.cpp b/Engine/Source/Editor/LevelEditor/Private/ComponentsTree.cpp
deleted file mode 100644
diff --git a/Engine/Source/Editor/LevelEditor/Private/ComponentsTree.h b/Engine/Source/Editor/LevelEditor/Private/ComponentsTree.h
deleted file mode 100644
diff --git a/Engine/Source/Editor/LevelEditor/Private/SLevelEditor.cpp b/Engine/Source/Editor/LevelEditor/Private/SLevelEditor.cpp
--- a/Engine/Source/Editor/LevelEditor/Private/SLevelEditor.cpp
+++ b/Engine/Source/Editor/LevelEditor/Private/SLevelEditor.cpp
@@ -38,9 +38,12 @@
 #include "NewsFeed.h"
 #include "TutorialMetaData.h"
 #include "SDockTab.h"
-#include "ComponentsTree.h"
+
+
+
 #include "SSCSEditor.h"
 #include "ComponentEditorUtils.h"
+#include "SExpandableArea.h"
 
 static const FName LevelEditorBuildAndSubmitTab("LevelEditorBuildAndSubmit");
 static const FName LevelEditorStatsViewerTab("LevelEditorStatsViewer");
diff --git a/Engine/Source/Editor/PropertyEditor/Private/DetailAdvancedDropdownNode.cpp b/Engine/Source/Editor/PropertyEditor/Private/DetailAdvancedDropdownNode.cpp
--- a/Engine/Source/Editor/PropertyEditor/Private/DetailAdvancedDropdownNode.cpp
+++ b/Engine/Source/Editor/PropertyEditor/Private/DetailAdvancedDropdownNode.cpp
@@ -36,7 +36,7 @@ public:
 			Content =
 				SNew( SBorder )
 				.BorderImage( FEditorStyle::GetBrush("DetailsView.CategoryMiddle") )
-				.Padding( FMargin( 0.0f, 3.0f, 0.0f, 0.0f ) )
+				.Padding( FMargin( 0.0f, 3.0f, SDetailTableRowBase::ScrollbarPaddingSize, 0.0f ) )
 				[
 					SNew( SImage )
 					.Image( FEditorStyle::GetBrush("DetailsView.AdvancedDropdownBorder.Open") )
@@ -47,7 +47,7 @@ public:
 			Content = 
 				SNew( SBorder )
 				.BorderImage( FEditorStyle::GetBrush("DetailsView.AdvancedDropdownBorder") )
-				.Padding( FMargin( 0.0f, 3.0f, 0.0f, 0.0f ) )
+				.Padding( FMargin( 0.0f, 3.0f, SDetailTableRowBase::ScrollbarPaddingSize, 0.0f ) )
 				[
 					SNew( SVerticalBox )
 					+SVerticalBox::Slot()
@@ -126,7 +126,7 @@ public:
 			Content =
 				SNew( SBorder )
 				.BorderImage( FEditorStyle::GetBrush("DetailsView.CategoryBottom") )
-				.Padding( FMargin( 0.0f, 0.0f, 0.0f, 2.0f ) )
+				.Padding( FMargin( 0.0f, 0.0f, SDetailTableRowBase::ScrollbarPaddingSize, 2.0f ) )
 				[
 					SplitterArea.ToSharedRef()	
 				];
diff --git a/Engine/Source/Editor/PropertyEditor/Private/DetailCategoryGroupNode.cpp b/Engine/Source/Editor/PropertyEditor/Private/DetailCategoryGroupNode.cpp
--- a/Engine/Source/Editor/PropertyEditor/Private/DetailCategoryGroupNode.cpp
+++ b/Engine/Source/Editor/PropertyEditor/Private/DetailCategoryGroupNode.cpp
@@ -16,7 +16,7 @@ void SDetailCategoryTableRow::Construct( const FArguments& InArgs, TSharedRef<ID
 	[	
 		SNew( SBorder )
 		.BorderImage( this, &SDetailCategoryTableRow::GetBackgroundImage )
-		.Padding( FMargin( 0.0f, 3.0f ) )
+		.Padding( FMargin( 0.0f, 3.0f, SDetailTableRowBase::ScrollbarPaddingSize, 3.0f ) )
 		.BorderBackgroundColor( FLinearColor( .6,.6,.6, 1.0f ) )
 		[
 			SAssignNew( MyContent, SHorizontalBox )
diff --git a/Engine/Source/Editor/PropertyEditor/Private/SDetailSingleItemRow.cpp b/Engine/Source/Editor/PropertyEditor/Private/SDetailSingleItemRow.cpp
--- a/Engine/Source/Editor/PropertyEditor/Private/SDetailSingleItemRow.cpp
+++ b/Engine/Source/Editor/PropertyEditor/Private/SDetailSingleItemRow.cpp
@@ -198,7 +198,7 @@ void SDetailSingleItemRow::Construct( const FArguments& InArgs, FDetailLayoutCus
 	[	
 		SNew( SBorder )
 		.BorderImage( this, &SDetailSingleItemRow::GetBorderImage )
-		.Padding( 0.0f )
+		.Padding( FMargin( 0.0f, 0.0f, SDetailTableRowBase::ScrollbarPaddingSize, 0.0f ) )
 		[
 			Widget
 		]
diff --git a/Engine/Source/Editor/PropertyEditor/Private/SDetailTableRowBase.h b/Engine/Source/Editor/PropertyEditor/Private/SDetailTableRowBase.h
--- a/Engine/Source/Editor/PropertyEditor/Private/SDetailTableRowBase.h
+++ b/Engine/Source/Editor/PropertyEditor/Private/SDetailTableRowBase.h
@@ -75,5 +75,8 @@ private:
 		}
 	}
 protected:
+	static float ScrollbarPaddingSize;
 	TWeakPtr<IDetailTreeNode> OwnerTreeNode;
 };
+
+float SDetailTableRowBase::ScrollbarPaddingSize = 16.0f;
\ No newline at end of file
diff --git a/Engine/Source/Editor/PropertyEditor/Private/SDetailsView.cpp b/Engine/Source/Editor/PropertyEditor/Private/SDetailsView.cpp
--- a/Engine/Source/Editor/PropertyEditor/Private/SDetailsView.cpp
+++ b/Engine/Source/Editor/PropertyEditor/Private/SDetailsView.cpp
@@ -209,13 +209,13 @@ void SDetailsView::Construct(const FArguments& InArgs)
 	.FillHeight(1)
 	.Padding(0)
 	[
-		SNew(SHorizontalBox)
-		+ SHorizontalBox::Slot()
+		SNew(SOverlay)
+		+ SOverlay::Slot()
 		[
 			ConstructTreeView(ExternalScrollbar)
 		]
-		+ SHorizontalBox::Slot()
-		.AutoWidth()
+		+ SOverlay::Slot()
+		.HAlign(HAlign_Right)
 		[
 			SNew(SBox)
 			.WidthOverride(16.0f)
