[INTEGRATE] Change 2417626 by [REDACTED_PERSON][REDACTED_HANDLE] on 2015/01/23 21:51:28

	Fix for multiply defined variable in non-unity build.

diff --git a/Engine/Source/Editor/PropertyEditor/Private/SDetailTableRowBase.cpp b/Engine/Source/Editor/PropertyEditor/Private/SDetailTableRowBase.cpp
new file mode 100644
--- /dev/null
+++ b/Engine/Source/Editor/PropertyEditor/Private/SDetailTableRowBase.cpp
@@ -0,0 +1,6 @@
+// Copyright 1998-2015 Epic Games, Inc. All Rights Reserved.
+
+#include "PropertyEditorPrivatePCH.h"
+#include "SDetailTableRowBase.h"
+
+float SDetailTableRowBase::ScrollbarPaddingSize = 16.0f;
diff --git a/Engine/Source/Editor/PropertyEditor/Private/SDetailTableRowBase.h b/Engine/Source/Editor/PropertyEditor/Private/SDetailTableRowBase.h
--- a/Engine/Source/Editor/PropertyEditor/Private/SDetailTableRowBase.h
+++ b/Engine/Source/Editor/PropertyEditor/Private/SDetailTableRowBase.h
@@ -78,5 +78,3 @@ protected:
 	static float ScrollbarPaddingSize;
 	TWeakPtr<IDetailTreeNode> OwnerTreeNode;
 };
-
-float SDetailTableRowBase::ScrollbarPaddingSize = 16.0f;
\ No newline at end of file
