[INTEGRATE] Change 2419380 by [REDACTED_PERSON][REDACTED_HANDLE] on 2015/01/26 18:10:41

	Temporarily disallow creation of new widget blueprints from a widget component

diff --git a/Engine/Source/Editor/PropertyEditor/Private/PropertyEditorHelpers.cpp b/Engine/Source/Editor/PropertyEditor/Private/PropertyEditorHelpers.cpp
--- a/Engine/Source/Editor/PropertyEditor/Private/PropertyEditorHelpers.cpp
+++ b/Engine/Source/Editor/PropertyEditor/Private/PropertyEditorHelpers.cpp
@@ -587,7 +587,7 @@ namespace PropertyEditorHelpers
 
 			UClass* Class = (ClassProp ? ClassProp->MetaClass : FEditorClassUtils::GetClassFromString(NodeProperty->GetMetaData("MetaClass")));
 
-			if (Class && FKismetEditorUtilities::CanCreateBlueprintOfClass(Class))
+			if (Class && FKismetEditorUtilities::CanCreateBlueprintOfClass(Class) && !ClassProp->HasMetaData("DisallowCreateNew"))
 			{
 				OutRequiredButtons.Add( EPropertyButton::NewBlueprint );
 			}
diff --git a/Engine/Source/Runtime/UMG/Public/Components/WidgetComponent.h b/Engine/Source/Runtime/UMG/Public/Components/WidgetComponent.h
--- a/Engine/Source/Runtime/UMG/Public/Components/WidgetComponent.h
+++ b/Engine/Source/Runtime/UMG/Public/Components/WidgetComponent.h
@@ -98,7 +98,7 @@ private:
 	EWidgetSpace Space;
 
 	/** The class of User Widget to create and display an instance of */
-	UPROPERTY(EditAnywhere, Category=UI)
+	UPROPERTY(EditAnywhere, Category=UI, meta=(DisallowCreateNew))
 	TSubclassOf<UUserWidget> WidgetClass;
 	
 	/** The size of the displayed quad. */
