FText key stablization & key editing

FText DisplayString shared ptrs are persisted across property edits. This preserves their entry in the live table and also updates the string that the entry points to. We call GetString() to update the hashed value as persisting the DisplayString means that its location in the table is preserved.
On serialization, a valid key is still found in the live table for the display string, so the key is preserved.
Added funtionality to the property editor to allow key editing and re-generation so any key conflicts can be resolved in the future.
Added new flag to prevent us from being able to edit 'immutable' code-declared (LOCTEXT) text. We dont attempt to preseve the identity of properties that are initialized with these.

UE-5350 - As a translator, I would like the editor tutorial text to preserve it's history between changes in order to greatly improve the translation workflow.

diff --git a/Engine/Source/Editor/EditorStyle/Private/SlateEditorStyle.cpp b/Engine/Source/Editor/EditorStyle/Private/SlateEditorStyle.cpp
--- a/Engine/Source/Editor/EditorStyle/Private/SlateEditorStyle.cpp
+++ b/Engine/Source/Editor/EditorStyle/Private/SlateEditorStyle.cpp
@@ -3060,6 +3060,7 @@ void FSlateEditorStyle::FStyle::SetupPropertyEditorStyles()
 		Set( "PropertyWindow.Button_PickAsset", new IMAGE_BRUSH( "Icons/pillarray_12x", Icon12x12 ) );
 		Set( "PropertyWindow.Button_PickActor", new IMAGE_BRUSH( "Icons/hiererchy_16x", Icon12x12 ) );
 		Set( "PropertyWindow.Button_PickActorInteractive", new IMAGE_BRUSH( "Icons/eyedropper_16px", Icon12x12 ) );
+		Set( "PropertyWindow.Button_Refresh", new IMAGE_BRUSH("Icons/refresh_12x", Icon12x12 ) );
 
 		Set( "PropertyWindow.WindowBorder", new BOX_BRUSH( "Common/GroupBorder", FMargin(4.0f/16.0f) ) );
 		Set( "DetailsView.NameChangeCommitted", new BOX_BRUSH( "Common/EditableTextSelectionBackground", FMargin(4.f/16.f) ) );
diff --git a/Engine/Source/Editor/IntroTutorials/Private/TutorialStructCustomization.cpp b/Engine/Source/Editor/IntroTutorials/Private/TutorialStructCustomization.cpp
--- a/Engine/Source/Editor/IntroTutorials/Private/TutorialStructCustomization.cpp
+++ b/Engine/Source/Editor/IntroTutorials/Private/TutorialStructCustomization.cpp
@@ -7,6 +7,7 @@
 #include "ISlateMetaData.h"
 #include "TutorialMetaData.h"
 #include "SDockTab.h"
+#include "PropertyCustomizationHelpers.h"
 
 #define LOCTEXT_NAMESPACE "TutorialStructCustomization"
 
@@ -165,6 +166,14 @@ void FTutorialContentCustomization::CustomizeHeader( TSharedRef<class IPropertyH
 				.OnTextCommitted(FOnTextCommitted::CreateStatic(&Local::OnTextCommitted, TextProperty))
 				.OnTextChanged(FOnTextChanged::CreateStatic(&Local::OnTextChanged, TextProperty))
 			]
+			+SHorizontalBox::Slot()
+			.AutoWidth()
+			.VAlign(VAlign_Center)
+			.HAlign(HAlign_Center)
+			.Padding(2.0f, 0.0f, 0.0f, 0.0f)
+			[
+				PropertyCustomizationHelpers::MakeTextLocalizationButton(TextProperty.ToSharedRef())
+			]
 		]
 	];
 }
diff --git a/Engine/Source/Editor/PropertyEditor/Private/PropertyCustomizationHelpers.cpp b/Engine/Source/Editor/PropertyEditor/Private/PropertyCustomizationHelpers.cpp
--- a/Engine/Source/Editor/PropertyEditor/Private/PropertyCustomizationHelpers.cpp
+++ b/Engine/Source/Editor/PropertyEditor/Private/PropertyCustomizationHelpers.cpp
@@ -1051,5 +1051,186 @@ void FMaterialList::AddMaterialItem( FDetailWidgetRow& Row, int32 CurrentSlot, c
 	];
 }
 
+
+TSharedRef<SWidget> PropertyCustomizationHelpers::MakeTextLocalizationButton(TSharedRef<IPropertyHandle> InPropertyHandle)
+{
+	struct Local
+	{
+		static void DirtyPropertyPackage(TSharedRef<IPropertyHandle> InPropertyHandle)
+		{
+			// dirty packages as they will need to be re-serialized with the correct key
+			TArray<UObject*> Objects;
+			InPropertyHandle->GetOuterObjects(Objects);
+			if (Objects.Num() > 0 && Objects[0]->GetOutermost() != nullptr)
+			{
+				Objects[0]->GetOutermost()->MarkPackageDirty();
+			}
+		}
+
+		static FReply HandleGenerateKeyClicked(TSharedRef<IPropertyHandle> InPropertyHandle)
+		{
+			FText DisplayText;
+			if (InPropertyHandle->GetValueAsDisplayText(DisplayText) == FPropertyAccess::Success)
+			{
+				TSharedRef< FString, ESPMode::ThreadSafe > DisplayString = FTextInspector::GetSharedDisplayString(DisplayText);
+				TSharedPtr< FString, ESPMode::ThreadSafe > Namespace;
+				TSharedPtr< FString, ESPMode::ThreadSafe > Key;
+				FTextLocalizationManager::Get().FindKeyNamespaceFromDisplayString(DisplayString, Namespace, Key);
+
+				if (Key.IsValid())
+				{
+					*Key.Get() = FGuid::NewGuid().ToString();
+					DirtyPropertyPackage(InPropertyHandle);
+				}
+			}
+
+			return FReply::Unhandled();
+		}
+
+		static void HandleKeyTextCommitted(const FText& InText, ETextCommit::Type InCommitType, TSharedRef<IPropertyHandle> InPropertyHandle)
+		{
+			FText DisplayText;
+			if (InPropertyHandle->GetValueAsDisplayText(DisplayText) == FPropertyAccess::Success)
+			{
+				TSharedRef< FString, ESPMode::ThreadSafe > DisplayString = FTextInspector::GetSharedDisplayString(DisplayText);
+				TSharedPtr< FString, ESPMode::ThreadSafe > Namespace;
+				TSharedPtr< FString, ESPMode::ThreadSafe > Key;
+				FTextLocalizationManager::Get().FindKeyNamespaceFromDisplayString(DisplayString, Namespace, Key);
+
+				if (Key.IsValid())
+				{
+					*Key.Get() = InText.ToString();
+					DirtyPropertyPackage(InPropertyHandle);
+				}
+			}
+		}
+
+		static void HandleNamespaceTextCommitted(const FText& InText, ETextCommit::Type InCommitType, TSharedRef<IPropertyHandle> InPropertyHandle)
+		{
+			FText DisplayText;
+			if (InPropertyHandle->GetValueAsDisplayText(DisplayText) == FPropertyAccess::Success)
+			{
+				TSharedRef< FString, ESPMode::ThreadSafe > DisplayString = FTextInspector::GetSharedDisplayString(DisplayText);
+				TSharedPtr< FString, ESPMode::ThreadSafe > Namespace;
+				TSharedPtr< FString, ESPMode::ThreadSafe > Key;
+				FTextLocalizationManager::Get().FindKeyNamespaceFromDisplayString(DisplayString, Namespace, Key);
+
+				if (Namespace.IsValid())
+				{
+					*Namespace.Get() = InText.ToString();
+					DirtyPropertyPackage(InPropertyHandle);
+				}
+			}
+		}
+
+		static TSharedRef<SWidget> GetLocalizationMenuContent(TSharedRef<IPropertyHandle> InPropertyHandle)
+		{
+			FMenuBuilder MenuContentBuilder(true, NULL);
+			{
+				FText DisplayText;
+				if (InPropertyHandle->GetValueAsDisplayText(DisplayText) == FPropertyAccess::Success)
+				{
+					TSharedRef< FString, ESPMode::ThreadSafe > DisplayString = FTextInspector::GetSharedDisplayString(DisplayText);
+					TSharedPtr< FString, ESPMode::ThreadSafe > Namespace;
+					TSharedPtr< FString, ESPMode::ThreadSafe > Key;
+					FTextLocalizationManager::Get().FindKeyNamespaceFromDisplayString(DisplayString, Namespace, Key);
+
+					MenuContentBuilder.BeginSection(TEXT("Localization"), LOCTEXT("LocalizationSectionHeading", "Localization"));
+					{
+						if (!Namespace.IsValid() && !Key.IsValid())
+						{
+							MenuContentBuilder.AddWidget(SNew(SSpacer), LOCTEXT("NoLocWarning", "No localization available"));
+						}
+						else
+						{
+							TSharedRef<SGridPanel> GridPanel = SNew(SGridPanel);
+
+							if (Key.IsValid())
+							{
+								GridPanel->AddSlot(0, 0)
+									.VAlign(VAlign_Center)
+									.Padding(1.0f)
+									[
+										SNew(STextBlock)
+										.Text(LOCTEXT("KeyLabel", "Key"))
+										.ToolTipText(LOCTEXT("KeyTooltip", "The localization key of the text property"))
+									];
+
+								auto KeyLambda = [=](){ return FText::FromString(*Key.Get()); };
+
+								GridPanel->AddSlot(1, 0)
+									.VAlign(VAlign_Center)
+									.Padding(1.0f)
+									[
+										SNew(SHorizontalBox)
+										+SHorizontalBox::Slot()
+										.FillWidth(1.0f)
+										[
+											SNew(SEditableTextBox)
+											.Text_Lambda(KeyLambda)
+											.OnTextCommitted(FOnTextCommitted::CreateStatic(&Local::HandleKeyTextCommitted, InPropertyHandle))
+										]
+										+ SHorizontalBox::Slot()
+										.AutoWidth()
+										.HAlign(HAlign_Center)
+										.Padding(2.0f, 0.0f, 0.0f, 0.0f)
+										[
+											SNew(SButton)
+											.ToolTipText(LOCTEXT("RefreshKeyTooltip", "Generate a new random key"))
+											.ButtonStyle(FEditorStyle::Get(), "HoverHintOnly")
+											.OnClicked(FOnClicked::CreateStatic(&Local::HandleGenerateKeyClicked, InPropertyHandle))
+											.Content()
+											[
+												SNew(SImage)
+												.Image(FEditorStyle::GetBrush("PropertyWindow.Button_Refresh"))
+											]
+										]
+									];
+							}
+
+							if (Namespace.IsValid())
+							{
+								auto NamespaceLambda = [=](){ return FText::FromString(*Namespace.Get()); };
+
+								GridPanel->AddSlot(0, 1)
+									.VAlign(VAlign_Center)
+									.Padding(1.0f, 1.0f, 5.0f, 1.0f)
+									[
+										SNew(STextBlock)
+										.Text(LOCTEXT("NamespaceLabel", "Namespace"))
+										.ToolTipText(LOCTEXT("NamespaceTooltip", "The localization namespace of the text property"))
+									];
+
+								GridPanel->AddSlot(1, 1)
+									.VAlign(VAlign_Center)
+									.Padding(1.0f)
+									[
+										SNew(SEditableTextBox)
+										.Text_Lambda(NamespaceLambda)
+										.OnTextCommitted(FOnTextCommitted::CreateStatic(&Local::HandleNamespaceTextCommitted, InPropertyHandle))
+									];
+							}
+
+							MenuContentBuilder.AddWidget(GridPanel, FText::GetEmpty());
+						}
+					}
+
+					MenuContentBuilder.EndSection();
+				}
+			}
+
+			return MenuContentBuilder.MakeWidget();
+		}
+	};
+
+	return SNew(SComboButton)
+		.ToolTipText(LOCTEXT("LocalizationUtilitiesTooltip", "Localization Utilities"))
+		.ButtonStyle(FEditorStyle::Get(), "HoverHintOnly")
+		.ContentPadding(2)
+		.ForegroundColor(FSlateColor::UseForeground())
+		.HasDownArrow(true)
+		.OnGetMenuContent(FOnGetContent::CreateStatic(&Local::GetLocalizationMenuContent, InPropertyHandle));
+}
+
 #undef LOCTEXT_NAMESPACE
 
diff --git a/Engine/Source/Editor/PropertyEditor/Private/UserInterface/PropertyEditor/SPropertyEditorText.cpp b/Engine/Source/Editor/PropertyEditor/Private/UserInterface/PropertyEditor/SPropertyEditorText.cpp
--- a/Engine/Source/Editor/PropertyEditor/Private/UserInterface/PropertyEditor/SPropertyEditorText.cpp
+++ b/Engine/Source/Editor/PropertyEditor/Private/UserInterface/PropertyEditor/SPropertyEditorText.cpp
@@ -13,38 +13,67 @@ void SPropertyEditorText::Construct( const FArguments& InArgs, const TSharedRef<
 	PropertyEditor = InPropertyEditor;
 
 	const UProperty* Property = InPropertyEditor->GetProperty();
+
+	TSharedPtr<SHorizontalBox> HorizontalBox;
+
 	bIsMultiLine = Property->GetBoolMetaData("MultiLine");
 	if(bIsMultiLine)
 	{
 		ChildSlot
 		[
-			SAssignNew( PrimaryWidget, SMultiLineEditableTextBox)
-			.Text( InPropertyEditor, &FPropertyEditor::GetValueAsText )
-			.Font( InArgs._Font )
-			.SelectAllTextWhenFocused( false )
-			.ClearKeyboardFocusOnCommit(false)
-			.OnTextCommitted( this, &SPropertyEditorText::OnTextCommitted )
-			.SelectAllTextOnCommit( false )
-			.IsReadOnly(this, &SPropertyEditorText::IsReadOnly)
-			.AutoWrapText(true)
-			.ModiferKeyForNewLine(EModifierKey::Shift)
+			SAssignNew(HorizontalBox, SHorizontalBox)
+			+SHorizontalBox::Slot()
+			.FillWidth(1.0f)
+			[
+				SAssignNew(PrimaryWidget, SMultiLineEditableTextBox)
+				.Text(InPropertyEditor, &FPropertyEditor::GetValueAsText)
+				.Font(InArgs._Font)
+				.SelectAllTextWhenFocused(false)
+				.ClearKeyboardFocusOnCommit(false)
+				.OnTextCommitted(this, &SPropertyEditorText::OnTextCommitted)
+				.SelectAllTextOnCommit(false)
+				.IsReadOnly(this, &SPropertyEditorText::IsReadOnly)
+				.AutoWrapText(true)
+				.ModiferKeyForNewLine(EModifierKey::Shift)
+			]
 		];
 	}
 	else
 	{
 		ChildSlot
 		[
-			SAssignNew( PrimaryWidget, SEditableTextBox )
-			.Text( InPropertyEditor, &FPropertyEditor::GetValueAsText )
-			.Font( InArgs._Font )
-			.SelectAllTextWhenFocused( true )
-			.ClearKeyboardFocusOnCommit(false)
-			.OnTextCommitted( this, &SPropertyEditorText::OnTextCommitted )
-			.SelectAllTextOnCommit( true )
-			.IsReadOnly(this, &SPropertyEditorText::IsReadOnly)
+			SAssignNew(HorizontalBox, SHorizontalBox)
+			+SHorizontalBox::Slot()
+			.FillWidth(1.0f)
+			[
+				SAssignNew( PrimaryWidget, SEditableTextBox )
+				.Text( InPropertyEditor, &FPropertyEditor::GetValueAsText )
+				.Font( InArgs._Font )
+				.SelectAllTextWhenFocused( true )
+				.ClearKeyboardFocusOnCommit(false)
+				.OnTextCommitted( this, &SPropertyEditorText::OnTextCommitted )
+				.SelectAllTextOnCommit( true )
+				.IsReadOnly(this, &SPropertyEditorText::IsReadOnly)
+			]
 		];
 	}
 
+	if (InPropertyEditor->PropertyIsA(UTextProperty::StaticClass()))
+	{
+		const TSharedRef< FPropertyNode > PropertyNode = InPropertyEditor->GetPropertyNode();
+		const TSharedRef< IPropertyHandle > PropertyHandle = PropertyEditor->GetPropertyHandle();
+
+		// text properties require some extra controls to allow for localization functionality
+		HorizontalBox->AddSlot()
+			.AutoWidth()
+			.VAlign(VAlign_Center)
+			.HAlign(HAlign_Center)
+			.Padding(2.0f, 0.0f, 0.0f, 0.0f)
+			[
+				PropertyCustomizationHelpers::MakeTextLocalizationButton(PropertyHandle)
+			];
+	}
+
 	if( InPropertyEditor->PropertyIsA( UObjectPropertyBase::StaticClass() ) )
 	{
 		// Object properties should display their entire text in a tooltip
diff --git a/Engine/Source/Editor/PropertyEditor/Public/PropertyCustomizationHelpers.h b/Engine/Source/Editor/PropertyEditor/Public/PropertyCustomizationHelpers.h
--- a/Engine/Source/Editor/PropertyEditor/Public/PropertyCustomizationHelpers.h
+++ b/Engine/Source/Editor/PropertyEditor/Public/PropertyCustomizationHelpers.h
@@ -40,6 +40,9 @@ namespace PropertyCustomizationHelpers
 
 	/** Returns a list of factories which can be used to create new assets, based on the supplied class */
 	PROPERTYEDITOR_API TArray<UFactory*> GetNewAssetFactoriesForClasses(const TArray<const UClass*>& Classes);
+
+	/** Build a combo button allowing access to text property localization utilities */
+	PROPERTYEDITOR_API TSharedRef<SWidget> MakeTextLocalizationButton(TSharedRef<IPropertyHandle> InPropertyHandle);
 }
 
 /** Delegate used to get a generic object */
diff --git a/Engine/Source/Runtime/Core/Private/Internationalization/Text.cpp b/Engine/Source/Runtime/Core/Private/Internationalization/Text.cpp
--- a/Engine/Source/Runtime/Core/Private/Internationalization/Text.cpp
+++ b/Engine/Source/Runtime/Core/Private/Internationalization/Text.cpp
@@ -54,6 +54,11 @@ const FString& FTextInspector::GetDisplayString(const FText& Text)
 	return Text.DisplayString.Get();
 }
 
+const TSharedRef<FString, ESPMode::ThreadSafe> FTextInspector::GetSharedDisplayString(const FText& Text)
+{
+	return Text.DisplayString;
+}
+
 int32 FTextInspector::GetFlags(const FText& Text)
 {
 	return Text.Flags;
@@ -843,4 +848,40 @@ bool FTextSnapshot::IsDisplayStringEqualTo(const FText& InText) const
 	return HistoryRevision == InHistoryRevision && DisplayStringPtr.IsValid() && DisplayStringPtr->Equals(InText.ToString(), ESearchCase::CaseSensitive);
 }
 
+FScopedTextIdentityPreserver::FScopedTextIdentityPreserver(FText& InTextToPersist)
+	: TextToPersist(InTextToPersist)
+	, Flags(InTextToPersist.Flags)
+{
+	// Empty display strings can't have a namespace or key.
+	if (GIsEditor && !InTextToPersist.DisplayString->IsEmpty())
+	{
+		// Save off namespace and key to be restored later.
+		FTextLocalizationManager::Get().FindKeyNamespaceFromDisplayString(InTextToPersist.DisplayString, Namespace, Key);
+	}
+}
+
+FScopedTextIdentityPreserver::~FScopedTextIdentityPreserver()
+{
+	// Never persist identities in non-editor situations
+	// If we don't have a key, then the old identity wasn't valid and shouldn't be preserved.
+	// Never persist identities for immutable (i.e. code LOCTEXT declared) text.
+	if(GIsEditor && Key.IsValid() && (Flags & ETextFlag::Immutable) == 0)
+	{
+		// Get the text's new source string.
+		const FString* SourceString = FTextInspector::GetSourceString(TextToPersist);
+
+		// Without a source string, we can't possibly preserve the identity. If the text we're preserving identity for can't possibly have an identity anymore, this class shouldn't be used on this text.
+		check(SourceString);
+
+		// Make the history so that it can attempt find and serialize the identity for this source string when the history itself is serialized.
+		TextToPersist.History = MakeShareable(new FTextHistory_Base(*SourceString));
+
+		// Create/update the display string instance for this identity in the text localization manager...
+		const TSharedRef<FString, ESPMode::ThreadSafe> DisplayString = FTextLocalizationManager::Get().GetString(Namespace.IsValid() ? *Namespace : TEXT(""), *Key, SourceString);
+
+		// ... and set it so that the text's history's serialization will properly map this display string instance back to the identity we intended to preserve and restore.
+		TextToPersist.DisplayString = DisplayString;
+	}
+}
+
 #undef LOCTEXT_NAMESPACE
diff --git a/Engine/Source/Runtime/Core/Private/Internationalization/TextLocalizationManager.cpp b/Engine/Source/Runtime/Core/Private/Internationalization/TextLocalizationManager.cpp
--- a/Engine/Source/Runtime/Core/Private/Internationalization/TextLocalizationManager.cpp
+++ b/Engine/Source/Runtime/Core/Private/Internationalization/TextLocalizationManager.cpp
@@ -469,7 +469,7 @@ void FTextLocalizationManager::UpdateLiveTable(const TArray<FLocalizationEntryTr
 						);
 					LiveKeyTable.Add( Key, NewLiveEntry );
 
-					ReverseLiveTable.Add(NewLiveEntry.String, FNamespaceKeyEntry( NamespaceName.IsEmpty()? nullptr : MakeShareable( new FString( NamespaceName ) ), Key.IsEmpty()? nullptr : MakeShareable( new FString( Key ) )));
+					ReverseLiveTable.Add(NewLiveEntry.String, FNamespaceKeyEntry( MakeShareable( new FString( NamespaceName ) ), MakeShareable( new FString( Key ) )));
 				}
 			}
 		}
@@ -640,7 +640,7 @@ TSharedRef<FString, ESPMode::ThreadSafe> FTextLocalizationManager::GetString(con
 
 			LiveKeyTable->Add( Key, NewEntry );
 
-			ReverseLiveTable.Add(NewEntry.String, FNamespaceKeyEntry( Namespace.IsEmpty()? nullptr : MakeShareable( new FString( Namespace ) ), Key.IsEmpty()? nullptr : MakeShareable( new FString( Key ) )));
+			ReverseLiveTable.Add(NewEntry.String, FNamespaceKeyEntry( MakeShareable( new FString( Namespace ) ), MakeShareable( new FString( Key ) )));
 
 		}
 
diff --git a/Engine/Source/Runtime/Core/Public/Internationalization/Internationalization.h b/Engine/Source/Runtime/Core/Public/Internationalization/Internationalization.h
--- a/Engine/Source/Runtime/Core/Public/Internationalization/Internationalization.h
+++ b/Engine/Source/Runtime/Core/Public/Internationalization/Internationalization.h
@@ -23,7 +23,7 @@ public:
 
 	static CORE_API FText ForUseOnlyByLocMacroAndGraphNodeTextLiterals_CreateText( const TCHAR* InTextLiteral, const TCHAR* Namespace, const TCHAR* Key )
 	{
-		return FText( InTextLiteral, Namespace, Key );
+		return FText( InTextLiteral, Namespace, Key, ETextFlag::Immutable );
 	}
 
 	//Set the current culture by name
diff --git a/Engine/Source/Runtime/Core/Public/Internationalization/Text.h b/Engine/Source/Runtime/Core/Public/Internationalization/Text.h
--- a/Engine/Source/Runtime/Core/Public/Internationalization/Text.h
+++ b/Engine/Source/Runtime/Core/Public/Internationalization/Text.h
@@ -18,9 +18,10 @@ namespace ETextFlag
 {
 	enum Type
 	{
-		Transient = (1<<0),
-		CultureInvariant = (1<<1),
-		ConvertedProperty = (1<<2)
+		Transient = (1 << 0),
+		CultureInvariant = (1 << 1),
+		ConvertedProperty = (1 << 2),
+		Immutable = (1 << 3),
 	};
 }
 
@@ -381,6 +382,7 @@ public:
 	friend class FTextHistory_NamedFormat;
 	friend class FTextHistory_ArgumentDataFormat;
 	friend class FTextHistory_OrderedFormat;
+	friend class FScopedTextIdentityPreserver;
 
 #if !UE_ENABLE_ICU
 	friend class FLegacyTextHelper;
@@ -424,6 +426,7 @@ public:
 	static const FString* GetKey(const FText& Text);
 	static const FString* GetSourceString(const FText& Text);
 	static const FString& GetDisplayString(const FText& Text);
+	static const TSharedRef<FString, ESPMode::ThreadSafe> GetSharedDisplayString(const FText& Text);
 	static int32 GetFlags(const FText& Text);
 };
 
@@ -563,4 +566,17 @@ private:
 	int32 IndentCount;
 };
 
+class CORE_API FScopedTextIdentityPreserver
+{
+public:
+	FScopedTextIdentityPreserver(FText& InTextToPersist);
+	~FScopedTextIdentityPreserver();
+
+private:
+	TSharedPtr< FString, ESPMode::ThreadSafe > Namespace;
+	TSharedPtr< FString, ESPMode::ThreadSafe > Key;
+	int32 Flags;
+	FText& TextToPersist;
+};
+
 Expose_TNameOf(FText)
diff --git a/Engine/Source/Runtime/Core/Public/Internationalization/TextLocalizationManager.h b/Engine/Source/Runtime/Core/Public/Internationalization/TextLocalizationManager.h
--- a/Engine/Source/Runtime/Core/Public/Internationalization/TextLocalizationManager.h
+++ b/Engine/Source/Runtime/Core/Public/Internationalization/TextLocalizationManager.h
@@ -79,8 +79,8 @@ private:
 public:
 	static CORE_API FTextLocalizationManager& Get();
 
-	TSharedPtr<FString, ESPMode::ThreadSafe> FindString( const FString& Namespace, const FString& Key, const FString* const SourceString = nullptr );
-	TSharedRef<FString, ESPMode::ThreadSafe> GetString(const FString& Namespace, const FString& Key, const FString* const SourceString);
+	TSharedPtr<FString, ESPMode::ThreadSafe> CORE_API FindString(const FString& Namespace, const FString& Key, const FString* const SourceString = nullptr);
+	TSharedRef<FString, ESPMode::ThreadSafe> CORE_API GetString(const FString& Namespace, const FString& Key, const FString* const SourceString);
 
 	/**
 	 * Using an FText SourceString's, returns the associated Namespace and Key
diff --git a/Engine/Source/Runtime/CoreUObject/Private/UObject/UTextProperty.cpp b/Engine/Source/Runtime/CoreUObject/Private/UObject/UTextProperty.cpp
--- a/Engine/Source/Runtime/CoreUObject/Private/UObject/UTextProperty.cpp
+++ b/Engine/Source/Runtime/CoreUObject/Private/UObject/UTextProperty.cpp
@@ -106,7 +106,8 @@ const TCHAR* UTextProperty::ImportText_Internal( const TCHAR* Buffer, void* Data
 		BufferAsString = FString(Temp);
 	}
 
-	*Text = FText::FromString( FString(BufferAsString) );
+	FScopedTextIdentityPreserver Preserver(*Text);
+	*Text = FText::FromString(FString(BufferAsString));
 
 	return Buffer;
 }
