Merging Fix crash clicking on objects with certian class properties (MattK).
using UE4-Releases-4.7-To-UE4

diff --git a/Engine/Source/Editor/PropertyEditor/Private/PropertyEditorHelpers.cpp b/Engine/Source/Editor/PropertyEditor/Private/PropertyEditorHelpers.cpp
--- a/Engine/Source/Editor/PropertyEditor/Private/PropertyEditorHelpers.cpp
+++ b/Engine/Source/Editor/PropertyEditor/Private/PropertyEditorHelpers.cpp
@@ -587,7 +587,7 @@ namespace PropertyEditorHelpers
 
 			UClass* Class = (ClassProp ? ClassProp->MetaClass : FEditorClassUtils::GetClassFromString(NodeProperty->GetMetaData("MetaClass")));
 
-			if (Class && FKismetEditorUtilities::CanCreateBlueprintOfClass(Class) && !ClassProp->HasMetaData("DisallowCreateNew"))
+			if (Class && FKismetEditorUtilities::CanCreateBlueprintOfClass(Class) && !NodeProperty->HasMetaData("DisallowCreateNew"))
 			{
 				OutRequiredButtons.Add( EPropertyButton::NewBlueprint );
 			}
