Optimizations to details panel when there are many large array properties.

diff --git a/Engine/Source/Editor/BehaviorTreeEditor/Private/BehaviorTreeEditor.cpp b/Engine/Source/Editor/BehaviorTreeEditor/Private/BehaviorTreeEditor.cpp
--- a/Engine/Source/Editor/BehaviorTreeEditor/Private/BehaviorTreeEditor.cpp
+++ b/Engine/Source/Editor/BehaviorTreeEditor/Private/BehaviorTreeEditor.cpp
@@ -954,9 +954,9 @@ void FBehaviorTreeEditor::CreateInternalWidgets()
 	FPropertyEditorModule& PropertyEditorModule = FModuleManager::GetModuleChecked<FPropertyEditorModule>( "PropertyEditor" );
 	FDetailsViewArgs DetailsViewArgs( false, false, true, FDetailsViewArgs::HideNameArea, false );
 	DetailsViewArgs.NotifyHook = this;
+	DetailsViewArgs.bShowDisableEditOnInstanceNodes = false;
 	DetailsView = PropertyEditorModule.CreateDetailView( DetailsViewArgs );
 	DetailsView->SetObject( NULL );
-	DetailsView->SetIsPropertyVisibleDelegate(FIsPropertyVisible::CreateSP(this, &FBehaviorTreeEditor::IsPropertyVisible));
 	DetailsView->SetIsPropertyEditingEnabledDelegate(FIsPropertyEditingEnabled::CreateSP(this, &FBehaviorTreeEditor::IsPropertyEditable));
 	DetailsView->OnFinishedChangingProperties().AddSP(this, &FBehaviorTreeEditor::OnFinishedChangingProperties);
 }
@@ -1109,11 +1109,6 @@ void FBehaviorTreeEditor::BindDebuggerToolbarCommands()
 		FExecuteAction::CreateStatic(&FBehaviorTreeDebugger::StopPlaySession));
 }
 
-bool FBehaviorTreeEditor::IsPropertyVisible( const FPropertyAndParent& PropertyAndParent) const
-{
-	return !PropertyAndParent.Property.HasAnyPropertyFlags(CPF_DisableEditOnInstance);
-}
-
 bool FBehaviorTreeEditor::IsPropertyEditable() const
 {
 	if (FBehaviorTreeDebugger::IsPIESimulating() || bForceDisablePropertyEdit)
diff --git a/Engine/Source/Editor/BehaviorTreeEditor/Private/BehaviorTreeEditor.h b/Engine/Source/Editor/BehaviorTreeEditor/Private/BehaviorTreeEditor.h
--- a/Engine/Source/Editor/BehaviorTreeEditor/Private/BehaviorTreeEditor.h
+++ b/Engine/Source/Editor/BehaviorTreeEditor/Private/BehaviorTreeEditor.h
@@ -87,7 +87,6 @@ public:
 
 	void JumpToNode(const UEdGraphNode* Node);
 
-	bool IsPropertyVisible( const struct FPropertyAndParent& PropertyAndParent ) const;
 	bool IsPropertyEditable() const;
 	void OnPackageSaved(const FString& PackageFileName, UObject* Outer);
 	void OnFinishedChangingProperties(const FPropertyChangedEvent& PropertyChangedEvent);
diff --git a/Engine/Source/Editor/BehaviorTreeEditor/Private/SBehaviorTreeDiff.cpp b/Engine/Source/Editor/BehaviorTreeEditor/Private/SBehaviorTreeDiff.cpp
--- a/Engine/Source/Editor/BehaviorTreeEditor/Private/SBehaviorTreeDiff.cpp
+++ b/Engine/Source/Editor/BehaviorTreeEditor/Private/SBehaviorTreeDiff.cpp
@@ -395,10 +395,10 @@ void SBehaviorTreeDiff::FBehaviorTreeDiffPanel::GeneratePanel(UEdGraph* Graph, U
 		];
 
 	FPropertyEditorModule& PropertyEditorModule = FModuleManager::GetModuleChecked<FPropertyEditorModule>( "PropertyEditor" );
-	const FDetailsViewArgs DetailsViewArgs( false, false, true, FDetailsViewArgs::ObjectsUseNameArea, false );
+	FDetailsViewArgs DetailsViewArgs( false, false, true, FDetailsViewArgs::ObjectsUseNameArea, false );
+	DetailsViewArgs.bShowDisableEditOnInstanceNodes = false;
 	DetailsView = PropertyEditorModule.CreateDetailView( DetailsViewArgs );
 	DetailsView->SetObject( NULL );
-	DetailsView->SetIsPropertyVisibleDelegate(FIsPropertyVisible::CreateRaw(this, &SBehaviorTreeDiff::FBehaviorTreeDiffPanel::IsPropertyVisible));
 	DetailsView->SetIsPropertyEditingEnabledDelegate(FIsPropertyEditingEnabled::CreateRaw(this, &SBehaviorTreeDiff::FBehaviorTreeDiffPanel::IsPropertyEditable));
 	
 	if(Graph)
@@ -584,11 +584,6 @@ void SBehaviorTreeDiff::FBehaviorTreeDiffPanel::OnSelectionChanged( const FGraph
 	}
 }
 
-bool SBehaviorTreeDiff::FBehaviorTreeDiffPanel::IsPropertyVisible(const FPropertyAndParent& PropertyAndParent) const
-{
-	return !PropertyAndParent.Property.HasAnyPropertyFlags(CPF_DisableEditOnInstance);
-}
-
 bool SBehaviorTreeDiff::FBehaviorTreeDiffPanel::IsPropertyEditable()
 {
 	return false;
diff --git a/Engine/Source/Editor/BehaviorTreeEditor/Private/SBehaviorTreeDiff.h b/Engine/Source/Editor/BehaviorTreeEditor/Private/SBehaviorTreeDiff.h
--- a/Engine/Source/Editor/BehaviorTreeEditor/Private/SBehaviorTreeDiff.h
+++ b/Engine/Source/Editor/BehaviorTreeEditor/Private/SBehaviorTreeDiff.h
@@ -53,11 +53,6 @@ private:
 		 */
 		void OnSelectionChanged( const FGraphPanelSelectionSet& Selection );
 
-		/**
-		 * Delegate to say if a node property should be visible
-		 */
-		bool IsPropertyVisible(const struct FPropertyAndParent& PropertyAndParen) const;
-
 		/**
 		 * Delegate to say if a node property should be editable
 		 */
diff --git a/Engine/Source/Editor/Blutility/Private/GlobalBlutilityDialog.cpp b/Engine/Source/Editor/Blutility/Private/GlobalBlutilityDialog.cpp
--- a/Engine/Source/Editor/Blutility/Private/GlobalBlutilityDialog.cpp
+++ b/Engine/Source/Editor/Blutility/Private/GlobalBlutilityDialog.cpp
@@ -126,29 +126,11 @@ FLinearColor FGlobalBlutilityDialog::GetWorldCentricTabColorScale() const
 
 void FGlobalBlutilityDialog::CreateInternalWidgets()
 {
-	// Helper for hiding private blueprint variables
-	struct Local
-	{
-		static bool IsPropertyVisible(const FPropertyAndParent& PropertyAndParent)
-		{
-			// For details views in the level editor all properties are the instanced versions
-			if (PropertyAndParent.Property.HasAllPropertyFlags(CPF_DisableEditOnInstance))
-			{
-				return false;
-			}
-			else
-			{
-				return true;
-			}
-		}
-	};
-
 	// Create a details view
 	FPropertyEditorModule& PropertyEditorModule = FModuleManager::GetModuleChecked<FPropertyEditorModule>("PropertyEditor");
-	const FDetailsViewArgs DetailsViewArgs(/*bUpdateFromSelection=*/ false, /*bLockable=*/ false, /*bAllowSearch=*/ false, FDetailsViewArgs::HideNameArea, /*bHideSelectionTip=*/ true);
+	FDetailsViewArgs DetailsViewArgs(/*bUpdateFromSelection=*/ false, /*bLockable=*/ false, /*bAllowSearch=*/ false, FDetailsViewArgs::HideNameArea, /*bHideSelectionTip=*/ true);
+	DetailsViewArgs.bShowDisableEditOnInstanceNodes = false;
 	DetailsView = PropertyEditorModule.CreateDetailView(DetailsViewArgs);
-
-	DetailsView->SetIsPropertyVisibleDelegate(FIsPropertyVisible::CreateStatic(&Local::IsPropertyVisible));
 }
 
 void FGlobalBlutilityDialog::UpdatePropertyWindow(const TArray<UObject*>& SelectedObjects)
diff --git a/Engine/Source/Editor/LevelEditor/Private/SActorDetails.cpp b/Engine/Source/Editor/LevelEditor/Private/SActorDetails.cpp
--- a/Engine/Source/Editor/LevelEditor/Private/SActorDetails.cpp
+++ b/Engine/Source/Editor/LevelEditor/Private/SActorDetails.cpp
@@ -10,20 +10,6 @@
 
 void SActorDetails::Construct(const FArguments& InArgs, const FName TabIdentifier)
 {
-	struct Local
-	{
-		static bool IsPropertyVisible(const FPropertyAndParent& PropertyAndParent)
-		{
-			// For details views in the level editor all properties are the instanced versions
-			if(PropertyAndParent.Property.HasAllPropertyFlags(CPF_DisableEditOnInstance))
-			{
-				return false;
-			}
-
-			return true;
-		}
-	};
-
 	bSelectionGuard = false;
 	bShowingRootActorNodeSelected = false;
 
@@ -39,11 +25,10 @@ void SActorDetails::Construct(const FArguments& InArgs, const FName TabIdentifie
 	DetailsViewArgs.ViewIdentifier = TabIdentifier;
 	DetailsViewArgs.bCustomNameAreaLocation = true;
 	DetailsViewArgs.bCustomFilterAreaLocation = true;
+	DetailsViewArgs.bShowDisableEditOnInstanceNodes = false;
 
 	DetailsView = PropPlugin.CreateDetailView(DetailsViewArgs);
 
-	DetailsView->SetIsPropertyVisibleDelegate(FIsPropertyVisible::CreateStatic(&Local::IsPropertyVisible));
-
 	auto IsPropertyEditingEnabled = [&]()
 	{
 		bool bIsEditable = true;
diff --git a/Engine/Source/Editor/PropertyEditor/Private/CategoryPropertyNode.cpp b/Engine/Source/Editor/PropertyEditor/Private/CategoryPropertyNode.cpp
--- a/Engine/Source/Editor/PropertyEditor/Private/CategoryPropertyNode.cpp
+++ b/Engine/Source/Editor/PropertyEditor/Private/CategoryPropertyNode.cpp
@@ -53,6 +53,7 @@ void FCategoryPropertyNode::InitBeforeNodeFlags()
 void FCategoryPropertyNode::InitChildNodes()
 {
 	const bool bShowHiddenProperties = !!HasNodeFlags( EPropertyNodeFlags::ShouldShowHiddenProperties );
+	const bool bShouldShowDisableEditOnInstance = !!HasNodeFlags(EPropertyNodeFlags::ShouldShowDisableEditOnInstance);
 
 	TArray<UProperty*> Properties;
 	// The parent of a category window has to be an object window.
@@ -63,19 +64,25 @@ void FCategoryPropertyNode::InitChildNodes()
 		for (TFieldIterator<UProperty> It(ComplexNode->GetBaseStructure()); It; ++It)
 		{
 			bool bMetaDataAllowVisible = true;
-			FString MetaDataVisibilityCheckString = It->GetMetaData(TEXT("bShowOnlyWhenTrue"));
-			if (MetaDataVisibilityCheckString.Len())
+			if (!bShowHiddenProperties)
 			{
-				//ensure that the metadata visibility string is actually set to true in order to show this property
-				// @todo Remove this
-				GConfig->GetBool(TEXT("UnrealEd.PropertyFilters"), *MetaDataVisibilityCheckString, bMetaDataAllowVisible, GEditorUserSettingsIni);
+				FString MetaDataVisibilityCheckString = It->GetMetaData(TEXT("bShowOnlyWhenTrue"));
+				if (MetaDataVisibilityCheckString.Len())
+				{
+					//ensure that the metadata visibility string is actually set to true in order to show this property
+					// @todo Remove this
+					GConfig->GetBool(TEXT("UnrealEd.PropertyFilters"), *MetaDataVisibilityCheckString, bMetaDataAllowVisible, GEditorUserSettingsIni);
+				}
 			}
 
 			if (bMetaDataAllowVisible)
 			{
+				const bool bShowIfEditableProperty = (*It)->HasAnyPropertyFlags(CPF_Edit);
+				const bool bShowIfDisableEditOnInstance = !(*It)->HasAnyPropertyFlags(CPF_DisableEditOnInstance) || bShouldShowDisableEditOnInstance;
+
 				// Add if we are showing non-editable props and this is the 'None' category, 
 				// or if this is the right category, and we are either showing non-editable
-				if (FObjectEditorUtils::GetCategoryFName(*It) == CategoryName && (bShowHiddenProperties || (It->PropertyFlags & CPF_Edit)))
+				if (FObjectEditorUtils::GetCategoryFName(*It) == CategoryName && (bShowHiddenProperties || (bShowIfEditableProperty && bShowIfDisableEditOnInstance)))
 				{
 					Properties.Add(*It);
 				}
@@ -94,6 +101,7 @@ void FCategoryPropertyNode::InitChildNodes()
 		InitParams.ArrayIndex = INDEX_NONE;
 		InitParams.bAllowChildren = true;
 		InitParams.bForceHiddenPropertyVisibility = bShowHiddenProperties;
+		InitParams.bCreateDisableEditOnInstanceNodes = bShouldShowDisableEditOnInstance;
 
 		NewItemNode->InitNode( InitParams );
 
diff --git a/Engine/Source/Editor/PropertyEditor/Private/ItemPropertyNode.cpp b/Engine/Source/Editor/PropertyEditor/Private/ItemPropertyNode.cpp
--- a/Engine/Source/Editor/PropertyEditor/Private/ItemPropertyNode.cpp
+++ b/Engine/Source/Editor/PropertyEditor/Private/ItemPropertyNode.cpp
@@ -112,11 +112,11 @@ void FItemPropertyNode::InitChildNodes()
 	UArrayProperty* ArrayProperty = Cast<UArrayProperty>(Property);
 	UObjectPropertyBase* ObjectProperty = Cast<UObjectPropertyBase>(Property);
 
-	bool bShouldShowHiddenProperties = !!HasNodeFlags(EPropertyNodeFlags::ShouldShowHiddenProperties);
+	const bool bShouldShowHiddenProperties = !!HasNodeFlags(EPropertyNodeFlags::ShouldShowHiddenProperties);
+	const bool bShouldShowDisableEditOnInstance = !!HasNodeFlags(EPropertyNodeFlags::ShouldShowDisableEditOnInstance);
 
 	if( Property->ArrayDim > 1 && ArrayIndex == -1 )
 	{
-
 		// Do not add array children which are defined by an enum but the enum at the array index is hidden
 		// This only applies to static arrays
 		static const FName NAME_ArraySizeEnum("ArraySizeEnum");
@@ -126,11 +126,9 @@ void FItemPropertyNode::InitChildNodes()
 			ArraySizeEnum	= FindObject<UEnum>(NULL, *Property->GetMetaData(NAME_ArraySizeEnum));
 		}
 
-	
 		// Expand array.
 		for( int32 ArrayIndex = 0 ; ArrayIndex < Property->ArrayDim ; ArrayIndex++ )
 		{
-		
 			bool bShouldBeHidden = false;
 			if( ArraySizeEnum )
 			{
@@ -148,6 +146,7 @@ void FItemPropertyNode::InitChildNodes()
 				InitParams.ArrayIndex = ArrayIndex;
 				InitParams.bAllowChildren = true;
 				InitParams.bForceHiddenPropertyVisibility = bShouldShowHiddenProperties;
+				InitParams.bCreateDisableEditOnInstanceNodes = bShouldShowDisableEditOnInstance;
 
 				NewItemNode->InitNode( InitParams );
 				AddChildNode(NewItemNode);
@@ -176,6 +175,7 @@ void FItemPropertyNode::InitChildNodes()
 				InitParams.ArrayIndex = ArrayIndex;
 				InitParams.bAllowChildren = true;
 				InitParams.bForceHiddenPropertyVisibility = bShouldShowHiddenProperties;
+				InitParams.bCreateDisableEditOnInstanceNodes = bShouldShowDisableEditOnInstance;
 
 				NewItemNode->InitNode( InitParams );
 				AddChildNode(NewItemNode);
@@ -188,7 +188,9 @@ void FItemPropertyNode::InitChildNodes()
 		for( TFieldIterator<UProperty> It(StructProperty->Struct); It; ++It )
 		{
 			UProperty* StructMember = *It;
-			if( bShouldShowHiddenProperties || (StructMember->PropertyFlags & CPF_Edit) )
+			const bool bShowIfEditableProperty = StructMember->HasAnyPropertyFlags(CPF_Edit);
+			const bool bShowIfDisableEditOnInstance = !StructMember->HasAnyPropertyFlags(CPF_DisableEditOnInstance) || bShouldShowDisableEditOnInstance;
+			if (bShouldShowHiddenProperties || (bShowIfEditableProperty && bShowIfDisableEditOnInstance))
 			{
 				TSharedPtr<FItemPropertyNode> NewItemNode( new FItemPropertyNode );//;//CreatePropertyItem(StructMember,INDEX_NONE,this);
 		
@@ -199,6 +201,7 @@ void FItemPropertyNode::InitChildNodes()
 				InitParams.ArrayIndex = INDEX_NONE;
 				InitParams.bAllowChildren = true;
 				InitParams.bForceHiddenPropertyVisibility = bShouldShowHiddenProperties;
+				InitParams.bCreateDisableEditOnInstanceNodes = bShouldShowDisableEditOnInstance;
 
 				NewItemNode->InitNode( InitParams );
 				AddChildNode(NewItemNode);
@@ -261,27 +264,13 @@ void FItemPropertyNode::InitChildNodes()
 				InitParams.ArrayIndex = INDEX_NONE;
 				InitParams.bAllowChildren = true;
 				InitParams.bForceHiddenPropertyVisibility = bShouldShowHiddenProperties;
+				InitParams.bCreateDisableEditOnInstanceNodes = bShouldShowDisableEditOnInstance;
 
 				NewObjectNode->InitNode( InitParams );
 				AddChildNode(NewObjectNode);
 			}
 		}
 	}
-
-	//needs to be after all the children are created
-	if ( FPropertySettings::Get().ExpandDistributions() == true)
-	{
-		// auto-expand distribution structs
-		if (Property->IsA(UStructProperty::StaticClass()))
-		{
-			FName StructName = ((UStructProperty*)Property)->Struct->GetFName();
-			if (StructName == NAME_RawDistributionFloat || StructName == NAME_RawDistributionVector)
-			{
-				const bool bExpand = true;
-				const bool bRecurse = true;
-			}
-		}
-	}
 }
 
 void FItemPropertyNode::SetDisplayNameOverride( const FText& InDisplayNameOverride )
diff --git a/Engine/Source/Editor/PropertyEditor/Private/ObjectPropertyNode.cpp b/Engine/Source/Editor/PropertyEditor/Private/ObjectPropertyNode.cpp
--- a/Engine/Source/Editor/PropertyEditor/Private/ObjectPropertyNode.cpp
+++ b/Engine/Source/Editor/PropertyEditor/Private/ObjectPropertyNode.cpp
@@ -304,27 +304,28 @@ void FObjectPropertyNode::InternalInitChildNodes( FName SinglePropertyName )
 	// Assemble a list of category names by iterating over all fields of BaseClass.
 
 	// build a list of classes that we need to look at
-	TArray<UClass*> ClassesToConsider;
+	TSet<UClass*> ClassesToConsider;
 	for( int32 i = 0; i < GetNumObjects(); ++i )
 	{
 		UObject* TempObject = GetUObject( i );
 		if( TempObject )
 		{
-			ClassesToConsider.AddUnique( TempObject->GetClass() );
+			ClassesToConsider.Add( TempObject->GetClass() );
 		}
 	}
 
-	TArray<FName> Categories;
+	const bool bShouldShowHiddenProperties = !!HasNodeFlags(EPropertyNodeFlags::ShouldShowHiddenProperties);
+	const bool bShouldShowDisableEditOnInstance = !!HasNodeFlags(EPropertyNodeFlags::ShouldShowDisableEditOnInstance);
+
+	TSet<FName> Categories;
 	for( TFieldIterator<UProperty> It(BaseClass.Get()); It; ++It )
 	{
 		bool bHidden = false;
 
 		FName CategoryName = FObjectEditorUtils::GetCategoryFName(*It);
 
-		for( int32 ClassIndex = 0; ClassIndex < ClassesToConsider.Num(); ClassIndex++ )
+		for( UClass* Class : ClassesToConsider )
 		{
-			UClass* Class = ClassesToConsider[ ClassIndex ];
-		
 			if( FEditorCategoryUtils::IsCategoryHiddenFromClass(Class, CategoryName.ToString()) )
 			{
 				HiddenCategories.Add( CategoryName );
@@ -344,9 +345,11 @@ void FObjectPropertyNode::InternalInitChildNodes( FName SinglePropertyName )
 
 		if (bMetaDataAllowVisible)
 		{
-			if( HasNodeFlags( EPropertyNodeFlags::ShouldShowHiddenProperties ) || (It->HasAnyPropertyFlags(CPF_Edit) && !bHidden) )
+			const bool bShowIfNonHiddenEditableProperty = (*It)->HasAnyPropertyFlags(CPF_Edit) && !bHidden;
+			const bool bShowIfDisableEditOnInstance = !(*It)->HasAnyPropertyFlags(CPF_DisableEditOnInstance) || bShouldShowDisableEditOnInstance;
+			if( bShouldShowHiddenProperties || (bShowIfNonHiddenEditableProperty && bShowIfDisableEditOnInstance) )
 			{
-				Categories.AddUnique( CategoryName );
+				Categories.Add( CategoryName );
 			}
 		}
 	}
@@ -362,10 +365,8 @@ void FObjectPropertyNode::InternalInitChildNodes( FName SinglePropertyName )
 
 		TArray< FPropertyNode* > ParentNodesToSort;
 
-		for( int32 x = 0 ; x < Categories.Num() ; ++x )
+		for( const FName& FullCategoryPath : Categories )
 		{
-			const FName& FullCategoryPath = Categories[x];
-
 			// Figure out the nesting level for this category
 			TArray< FString > FullCategoryPathStrings;
 			FullCategoryPath.ToString().ParseIntoArray( &FullCategoryPathStrings, *CategoryDelimiterString, true );
@@ -420,11 +421,11 @@ void FObjectPropertyNode::InternalInitChildNodes( FName SinglePropertyName )
 						InitParams.ArrayOffset = 0;
 						InitParams.ArrayIndex = INDEX_NONE;
 						InitParams.bAllowChildren = true;
-						InitParams.bForceHiddenPropertyVisibility = !!HasNodeFlags(EPropertyNodeFlags::ShouldShowHiddenProperties);
+						InitParams.bForceHiddenPropertyVisibility = bShouldShowHiddenProperties;
+						InitParams.bCreateDisableEditOnInstanceNodes = bShouldShowDisableEditOnInstance;
 
 						NewCategoryNode->InitNode( InitParams );
 
-
 						// Recursively expand category properties if the category has been flagged for auto-expansion.
 						if (BaseClass->IsAutoExpandCategory(*CategoryName.ToString())
 							&&	!BaseClass->IsAutoCollapseCategory(*CategoryName.ToString()))
@@ -447,15 +448,12 @@ void FObjectPropertyNode::InternalInitChildNodes( FName SinglePropertyName )
 	}
 	else
 	{
-		// For now we show noneditable properties. 
-		const bool bShouldShowHiddenProperties = !!HasNodeFlags( EPropertyNodeFlags::ShouldShowHiddenProperties );
-
-	
 		// Iterate over all fields, creating items.
 		for( TFieldIterator<UProperty> It(BaseClass.Get()); It; ++It )
 		{
-			if( bShouldShowHiddenProperties ||
-				( ( It->PropertyFlags&CPF_Edit ) && ( FEditorCategoryUtils::IsCategoryHiddenFromClass(BaseClass.Get(), FObjectEditorUtils::GetCategory(*It)) == false ) ) )
+			const bool bShowIfNonHiddenEditableProperty = (*It)->HasAnyPropertyFlags(CPF_Edit) && !FEditorCategoryUtils::IsCategoryHiddenFromClass(BaseClass.Get(), FObjectEditorUtils::GetCategory(*It));
+			const bool bShowIfDisableEditOnInstance = !(*It)->HasAnyPropertyFlags(CPF_DisableEditOnInstance) || bShouldShowDisableEditOnInstance;
+			if (bShouldShowHiddenProperties || (bShowIfNonHiddenEditableProperty && bShowIfDisableEditOnInstance))
 			{
 				UProperty* CurProp = *It;
 				if( SinglePropertyName == NAME_None || CurProp->GetFName() == SinglePropertyName )
@@ -469,6 +467,7 @@ void FObjectPropertyNode::InternalInitChildNodes( FName SinglePropertyName )
 					InitParams.ArrayIndex = INDEX_NONE;
 					InitParams.bAllowChildren = SinglePropertyName == NAME_None;
 					InitParams.bForceHiddenPropertyVisibility = bShouldShowHiddenProperties;
+					InitParams.bCreateDisableEditOnInstanceNodes = bShouldShowDisableEditOnInstance;
 
 					NewItemNode->InitNode( InitParams );
 
diff --git a/Engine/Source/Editor/PropertyEditor/Private/PropertyNode.cpp b/Engine/Source/Editor/PropertyEditor/Private/PropertyNode.cpp
--- a/Engine/Source/Editor/PropertyEditor/Private/PropertyNode.cpp
+++ b/Engine/Source/Editor/PropertyEditor/Private/PropertyNode.cpp
@@ -103,10 +103,8 @@ void FPropertyNode::InitNode( const FPropertyNodeInitParams& InitParams )
 		SetNodeFlags(EPropertyNodeFlags::ShowCategories, InitParams.bCreateCategoryNodes );
 	}
 
-	if (InitParams.bForceHiddenPropertyVisibility)
-	{
-		SetNodeFlags(EPropertyNodeFlags::ShouldShowHiddenProperties, true );
-	}
+	SetNodeFlags(EPropertyNodeFlags::ShouldShowHiddenProperties, InitParams.bForceHiddenPropertyVisibility);
+	SetNodeFlags(EPropertyNodeFlags::ShouldShowDisableEditOnInstance, InitParams.bCreateDisableEditOnInstanceNodes);
 
 	//Custom code run prior to setting property flags
 	//needs to happen after the above SetNodeFlags calls so that ObjectPropertyNode can properly respond to CollapseCategories
@@ -321,6 +319,8 @@ FObjectPropertyNode* FPropertyNode::FindRootObjectItemParent()
  */
 FPropertyNode::DataValidationResult FPropertyNode::EnsureDataIsValid()
 {
+	bool bValidateChildren = true;
+
 	// The root must always be validated
 	if( GetParentNode() == NULL || HasNodeFlags(EPropertyNodeFlags::RequiresValidation) != 0 )
 	{
@@ -347,6 +347,14 @@ FPropertyNode::DataValidationResult FPropertyNode::EnsureDataIsValid()
 
 			bool bArrayHasNewItem = false;
 
+			if (ArrayProperty)
+			{
+				if (!ArrayProperty->Inner->IsA(UObjectProperty::StaticClass()) && !ArrayProperty->Inner->IsA(UStructProperty::StaticClass()))
+				{
+					bValidateChildren = false;
+				}
+			}
+
 			//verify that the number of object children are the same too
 			UObjectPropertyBase* ObjectProperty = Cast<UObjectPropertyBase>(MyProperty);
 			//check to see, if this an object property, whether the contents are NULL or not.
@@ -453,20 +461,23 @@ FPropertyNode::DataValidationResult FPropertyNode::EnsureDataIsValid()
 	FPropertyNode::DataValidationResult FinalResult = DataValid;
 
 	//go through my children
-	for (int32 Scan = 0; Scan < ChildNodes.Num(); ++Scan)
+	if (bValidateChildren)
 	{
-		TSharedPtr<FPropertyNode>& ChildNode = ChildNodes[Scan];
-		check( ChildNode.IsValid() );
-
-		// @todo Slate Property Window 
-		//if (ChildNode->HasNodeFlags(EPropertyNodeFlags::IsSeen))
-		//{
-		FPropertyNode::DataValidationResult ChildDataResult = ChildNode->EnsureDataIsValid();
-		if (FinalResult == DataValid && ChildDataResult != DataValid )
+		for (int32 Scan = 0; Scan < ChildNodes.Num(); ++Scan)
 		{
-			FinalResult = ChildDataResult;
+			TSharedPtr<FPropertyNode>& ChildNode = ChildNodes[Scan];
+			check(ChildNode.IsValid());
+
+			// @todo Slate Property Window 
+			//if (ChildNode->HasNodeFlags(EPropertyNodeFlags::IsSeen))
+			//{
+			FPropertyNode::DataValidationResult ChildDataResult = ChildNode->EnsureDataIsValid();
+			if (FinalResult == DataValid && ChildDataResult != DataValid)
+			{
+				FinalResult = ChildDataResult;
+			}
+			//}
 		}
-		//}
 	}
 
 	return FinalResult;
diff --git a/Engine/Source/Editor/PropertyEditor/Private/PropertyNode.h b/Engine/Source/Editor/PropertyEditor/Private/PropertyNode.h
--- a/Engine/Source/Editor/PropertyEditor/Private/PropertyNode.h
+++ b/Engine/Source/Editor/PropertyEditor/Private/PropertyNode.h
@@ -38,6 +38,8 @@ namespace EPropertyNodeFlags
 	const Type	IsCustomized					= 1 << 17;	/** true if this node's visual representation has been customized by the editor */
 	
 	const Type	RequiresValidation				= 1 << 18; /** true if this node could unexpectedly change (array changes, editinlinenew changes) */
+
+	const Type	ShouldShowDisableEditOnInstance = 1 << 19; /** true if this node should show child properties marked CPF_DisableEditOnInstance */
 	const Type 	NoFlags							= 0;
 
 };
@@ -173,6 +175,9 @@ struct FPropertyNodeInitParams
 	bool bForceHiddenPropertyVisibility;
 	/** Whether or not to create category nodes (note: this setting is only valid for the root node of a property tree. The setting will propagate to children) */
 	bool bCreateCategoryNodes;
+	/** Whether or not to create nodes for properties marked CPF_DisableEditOnInstance */
+	bool bCreateDisableEditOnInstanceNodes;
+
 	FPropertyNodeInitParams()
 		: ParentNode( NULL )
 		, Property( NULL )
@@ -181,6 +186,7 @@ struct FPropertyNodeInitParams
 		, bAllowChildren( true )
 		, bForceHiddenPropertyVisibility( false )
 		, bCreateCategoryNodes( true )
+		, bCreateDisableEditOnInstanceNodes( true )
 	{}
 };
 
diff --git a/Engine/Source/Editor/PropertyEditor/Private/SDetailsView.cpp b/Engine/Source/Editor/PropertyEditor/Private/SDetailsView.cpp
--- a/Engine/Source/Editor/PropertyEditor/Private/SDetailsView.cpp
+++ b/Engine/Source/Editor/PropertyEditor/Private/SDetailsView.cpp
@@ -579,6 +579,7 @@ void SDetailsView::PostSetObject()
 	InitParams.ArrayIndex = INDEX_NONE;
 	InitParams.bAllowChildren = true;
 	InitParams.bForceHiddenPropertyVisibility =  FPropertySettings::Get().ShowHiddenProperties();
+	InitParams.bCreateDisableEditOnInstanceNodes = DetailsViewArgs.bShowDisableEditOnInstanceNodes;
 
 	RootPropertyNode->InitNode( InitParams );
 
diff --git a/Engine/Source/Editor/PropertyEditor/Private/StructurePropertyNode.cpp b/Engine/Source/Editor/PropertyEditor/Private/StructurePropertyNode.cpp
--- a/Engine/Source/Editor/PropertyEditor/Private/StructurePropertyNode.cpp
+++ b/Engine/Source/Editor/PropertyEditor/Private/StructurePropertyNode.cpp
@@ -11,12 +11,17 @@
 void FStructurePropertyNode::InitChildNodes()
 {
 	const bool bShouldShowHiddenProperties = !!HasNodeFlags(EPropertyNodeFlags::ShouldShowHiddenProperties);
+	const bool bShouldShowDisableEditOnInstance = !!HasNodeFlags(EPropertyNodeFlags::ShouldShowDisableEditOnInstance);
+
 	const UStruct* Struct = StructData.IsValid() ? StructData->GetStruct() : NULL;
 
 	for (TFieldIterator<UProperty> It(Struct); It; ++It)
 	{
 		UProperty* StructMember = *It;
-		if (StructMember && (bShouldShowHiddenProperties || (StructMember->PropertyFlags & CPF_Edit)))
+		const bool bShowIfEditableProperty = StructMember->HasAnyPropertyFlags(CPF_Edit);
+		const bool bShowIfDisableEditOnInstance = !StructMember->HasAnyPropertyFlags(CPF_DisableEditOnInstance) || bShouldShowDisableEditOnInstance;
+
+		if (StructMember && (bShouldShowHiddenProperties || (bShowIfEditableProperty && bShowIfDisableEditOnInstance)))
 		{
 			TSharedPtr<FItemPropertyNode> NewItemNode(new FItemPropertyNode);//;//CreatePropertyItem(StructMember,INDEX_NONE,this);
 
@@ -26,7 +31,8 @@ void FStructurePropertyNode::InitChildNodes()
 			InitParams.ArrayOffset = 0;
 			InitParams.ArrayIndex = INDEX_NONE;
 			InitParams.bAllowChildren = true;
-			InitParams.bForceHiddenPropertyVisibility = !!HasNodeFlags(EPropertyNodeFlags::ShouldShowHiddenProperties);
+			InitParams.bForceHiddenPropertyVisibility = bShouldShowHiddenProperties;
+			InitParams.bCreateDisableEditOnInstanceNodes = bShouldShowDisableEditOnInstance;
 			InitParams.bCreateCategoryNodes = false;
 
 			NewItemNode->InitNode(InitParams);
diff --git a/Engine/Source/Editor/PropertyEditor/Public/IDetailsView.h b/Engine/Source/Editor/PropertyEditor/Public/IDetailsView.h
--- a/Engine/Source/Editor/PropertyEditor/Public/IDetailsView.h
+++ b/Engine/Source/Editor/PropertyEditor/Public/IDetailsView.h
@@ -44,10 +44,13 @@ struct FDetailsViewArgs
 	uint32 bShowActorLabel : 1;
 	/** Bind this delegate to hide differing properties */
 	uint32 bShowDifferingPropertiesOption : 1;
-	/** If true the name area will be created but will not be displayed so it can be placed in a custom location.  */
+	/** If true, the name area will be created but will not be displayed so it can be placed in a custom location.  */
 	uint32 bCustomNameAreaLocation : 1;
-	/** If true the filter area will be created but will not be displayed so it can be placed in a custom location.  */
+	/** If true, the filter area will be created but will not be displayed so it can be placed in a custom location.  */
 	uint32 bCustomFilterAreaLocation : 1;
+	/** If true, nodes marked CPF_DisableEditOnInstance will be shown */
+	uint32 bShowDisableEditOnInstanceNodes : 1;
+
 	/** Default constructor */
 	FDetailsViewArgs( const bool InUpdateFromSelection = false
 					, const bool InLockable = false
@@ -71,6 +74,7 @@ struct FDetailsViewArgs
 		, bShowDifferingPropertiesOption(false)
 		, bCustomNameAreaLocation(false)
 		, bCustomFilterAreaLocation(false)
+		, bShowDisableEditOnInstanceNodes(true)
 	{
 	}
 };
