Fix some nonunity fails

diff --git a/Engine/Source/Editor/PropertyEditor/Public/IDetailPropertyRow.h b/Engine/Source/Editor/PropertyEditor/Public/IDetailPropertyRow.h
--- a/Engine/Source/Editor/PropertyEditor/Public/IDetailPropertyRow.h
+++ b/Engine/Source/Editor/PropertyEditor/Public/IDetailPropertyRow.h
@@ -2,6 +2,9 @@
 
 #pragma once
 
+// Forward declarations
+class FDetailWidgetRow;
+
 /**
  * A single row for a property in a details panel                                                              
  */
diff --git a/Engine/Source/Runtime/Landscape/Private/LandscapeGrass.cpp b/Engine/Source/Runtime/Landscape/Private/LandscapeGrass.cpp
--- a/Engine/Source/Runtime/Landscape/Private/LandscapeGrass.cpp
+++ b/Engine/Source/Runtime/Landscape/Private/LandscapeGrass.cpp
@@ -19,6 +19,7 @@
 #include "EdGraph/EdGraphNode.h"
 #include "Engine/TextureRenderTarget2D.h"
 #include "ContentStreaming.h"
+#include "LandscapeDataAccess.h"
 
 #define LOCTEXT_NAMESPACE "Landscape"
 
