[INTEGRATE] Change 2433594 by [REDACTED_PERSON][REDACTED_HANDLE] on 2015/02/05 10:21:50

	UE-8999 crash editing a public variable in my blueprint

diff --git a/Engine/Source/Editor/PropertyEditor/Private/PropertyHandleImpl.cpp b/Engine/Source/Editor/PropertyEditor/Private/PropertyHandleImpl.cpp
--- a/Engine/Source/Editor/PropertyEditor/Private/PropertyHandleImpl.cpp
+++ b/Engine/Source/Editor/PropertyEditor/Private/PropertyHandleImpl.cpp
@@ -338,8 +338,7 @@ FPropertyAccess::Result FPropertyValueImpl::ImportText( const TArray<FObjectBase
 			if (Cur.BaseAddress == NULL)
 			{
 				//Fully abort this procedure.  The data has changed out from under the object
-				check(false);
-				Result = FPropertyAccess::Fail;;
+				Result = FPropertyAccess::Fail;
 				break;
 			}
 
