[INTEGRATE] Change 2435462 by [REDACTED_PERSON][REDACTED_HANDLE] on 2015/02/06 10:42:03

	Editor - The asset pickers now all save their changes to size so that you don't have to resize it every time you go to pick an asset.

diff --git a/Engine/Source/Editor/PropertyEditor/Private/UserInterface/PropertyEditor/SPropertyEditorAsset.cpp b/Engine/Source/Editor/PropertyEditor/Private/UserInterface/PropertyEditor/SPropertyEditorAsset.cpp
--- a/Engine/Source/Editor/PropertyEditor/Private/UserInterface/PropertyEditor/SPropertyEditorAsset.cpp
+++ b/Engine/Source/Editor/PropertyEditor/Private/UserInterface/PropertyEditor/SPropertyEditorAsset.cpp
@@ -192,6 +192,7 @@ void SPropertyEditorAsset::Construct( const FArguments& InArgs, const TSharedPtr
 				.ButtonStyle( FEditorStyle::Get(), "PropertyEditor.AssetComboStyle" )
 				.ForegroundColor(FEditorStyle::GetColor("PropertyEditor.AssetName.ColorAndOpacity"))
 				.OnGetMenuContent( this, &SPropertyEditorAsset::OnGetMenuContent )
+				.OnMenuOpenChanged( this, &SPropertyEditorAsset::OnMenuOpenChanged )
 				.ContentPadding(2.0f)
 				.ButtonContent()
 				[
@@ -344,6 +345,14 @@ TSharedRef<SWidget> SPropertyEditorAsset::OnGetMenuContent()
 	}
 }
 
+void SPropertyEditorAsset::OnMenuOpenChanged(bool bOpen)
+{
+	if ( bOpen == false )
+	{
+		AssetComboButton->SetMenuContent(SNullWidget::NullWidget);
+	}
+}
+
 bool SPropertyEditorAsset::IsFilteredActor( const AActor* const Actor ) const
 {
 	return Actor->IsA( ObjectClass );
diff --git a/Engine/Source/Editor/PropertyEditor/Private/UserInterface/PropertyEditor/SPropertyEditorAsset.h b/Engine/Source/Editor/PropertyEditor/Private/UserInterface/PropertyEditor/SPropertyEditorAsset.h
--- a/Engine/Source/Editor/PropertyEditor/Private/UserInterface/PropertyEditor/SPropertyEditorAsset.h
+++ b/Engine/Source/Editor/PropertyEditor/Private/UserInterface/PropertyEditor/SPropertyEditorAsset.h
@@ -102,6 +102,12 @@ private:
 	 */
 	TSharedRef<SWidget> OnGetMenuContent();
 
+	/**
+	 * Called when the asset menu is closed, we handle this to force the destruction of the asset menu to
+	 * ensure any settings the user set are saved.
+	 */
+	void OnMenuOpenChanged(bool bOpen);
+
 	/**
 	 * Returns whether the actor should be filtered out from selection.
 	 */
diff --git a/Engine/Source/Editor/PropertyEditor/Private/UserInterface/PropertyEditor/SPropertyMenuAssetPicker.cpp b/Engine/Source/Editor/PropertyEditor/Private/UserInterface/PropertyEditor/SPropertyMenuAssetPicker.cpp
--- a/Engine/Source/Editor/PropertyEditor/Private/UserInterface/PropertyEditor/SPropertyMenuAssetPicker.cpp
+++ b/Engine/Source/Editor/PropertyEditor/Private/UserInterface/PropertyEditor/SPropertyMenuAssetPicker.cpp
@@ -97,8 +97,8 @@ void SPropertyMenuAssetPicker::Construct( const FArguments& InArgs )
 		AssetPickerConfig.Filter.bRecursiveClasses = true;
 		// Set a delegate for setting the asset from the picker
 		AssetPickerConfig.OnAssetSelected = FOnAssetSelected::CreateSP(this, &SPropertyMenuAssetPicker::OnAssetSelected);
-		// Use the smallest size thumbnails
-		AssetPickerConfig.ThumbnailScale = 0;
+		// Use small size thumbnails by default
+		AssetPickerConfig.ThumbnailScale = 0.1;
 		// Use the list view by default
 		AssetPickerConfig.InitialAssetViewType = EAssetViewType::List;
 		// The initial selection should be the current value
@@ -111,6 +111,8 @@ void SPropertyMenuAssetPicker::Construct( const FArguments& InArgs )
 		AssetPickerConfig.OnShouldFilterAsset = OnShouldFilterAsset;
 		// Don't allow dragging
 		AssetPickerConfig.bAllowDragging = false;
+		// Save the settings into a special section for asset pickers for properties
+		AssetPickerConfig.SaveSettingsName = TEXT("AssetPropertyPicker");
 
 		MenuContent =
 			SNew(SBox)
diff --git a/Engine/Source/Editor/UnrealEd/Private/Toolkits/GlobalEditorCommonCommands.cpp b/Engine/Source/Editor/UnrealEd/Private/Toolkits/GlobalEditorCommonCommands.cpp
--- a/Engine/Source/Editor/UnrealEd/Private/Toolkits/GlobalEditorCommonCommands.cpp
+++ b/Engine/Source/Editor/UnrealEd/Private/Toolkits/GlobalEditorCommonCommands.cpp
@@ -25,12 +25,13 @@ public:
 		FAssetPickerConfig AssetPickerConfig;
 		AssetPickerConfig.OnAssetDoubleClicked = FOnAssetSelected::CreateSP(this, &SGlobalOpenAssetDialog::OnAssetSelectedFromPicker);
 		AssetPickerConfig.OnAssetEnterPressed = FOnAssetEnterPressed::CreateSP(this, &SGlobalOpenAssetDialog::OnPressedEnterOnAssetsInPicker);
-		AssetPickerConfig.ThumbnailScale = 0.0f;
+		AssetPickerConfig.ThumbnailScale = 0.1f;
 		AssetPickerConfig.InitialAssetViewType = EAssetViewType::List;
 		AssetPickerConfig.bAllowNullSelection = false;
 		AssetPickerConfig.bShowBottomToolbar = true;
 		AssetPickerConfig.bAutohideSearchBar = false;
 		AssetPickerConfig.bCanShowClasses = false;
+		AssetPickerConfig.SaveSettingsName = TEXT("GlobalAssetPicker");
 
 		ChildSlot
 		[
