[INTEGRATE] Change 2436814 by [REDACTED_PERSON][REDACTED_HANDLE] on 2015/02/07 13:41:18

	Fix BSP surface selection not appearing if you had a brush selected and then selected a surface

diff --git a/Engine/Source/Editor/PropertyEditor/Private/SDetailsView.cpp b/Engine/Source/Editor/PropertyEditor/Private/SDetailsView.cpp
--- a/Engine/Source/Editor/PropertyEditor/Private/SDetailsView.cpp
+++ b/Engine/Source/Editor/PropertyEditor/Private/SDetailsView.cpp
@@ -353,7 +353,14 @@ void SDetailsView::RemoveInvalidObjects()
 bool SDetailsView::ShouldSetNewObjects( const TArray< TWeakObjectPtr< UObject > >& InObjects ) const
 {
 	bool bShouldSetObjects = false;
-	if( InObjects.Num() != RootPropertyNode->GetNumObjects() )
+
+	const bool bHadBSPBrushSelected = SelectedActorInfo.bHaveBSPBrush;
+	if( bHadBSPBrushSelected == true )
+	{
+		// If a BSP brush was selected we need to refresh because surface could have been selected and the object set not updated
+		bShouldSetObjects = true;
+	}
+	else if( InObjects.Num() != RootPropertyNode->GetNumObjects() )
 	{
 		// If the object arrys differ in size then at least one object is different so we must reset
 		bShouldSetObjects = true;
diff --git a/Engine/Source/Editor/UnrealEd/Private/AssetSelection.cpp b/Engine/Source/Editor/UnrealEd/Private/AssetSelection.cpp
--- a/Engine/Source/Editor/UnrealEd/Private/AssetSelection.cpp
+++ b/Engine/Source/Editor/UnrealEd/Private/AssetSelection.cpp
@@ -92,6 +92,7 @@ namespace AssetSelectionUtils
 							ActorInfo.bHaveBuilderBrush = FActorEditorUtils::IsABuilderBrush(Brush);
 						}
 						ActorInfo.bHaveBrush |= true;
+						ActorInfo.bHaveBSPBrush |= (!Brush->IsVolumeBrush());
 						ActorInfo.bHaveVolume |= Brush->IsVolumeBrush();
 					}
 
diff --git a/Engine/Source/Editor/UnrealEd/Public/AssetSelection.h b/Engine/Source/Editor/UnrealEd/Public/AssetSelection.h
--- a/Engine/Source/Editor/UnrealEd/Public/AssetSelection.h
+++ b/Engine/Source/Editor/UnrealEd/Public/AssetSelection.h
@@ -85,6 +85,8 @@ struct FSelectedActorInfo
 	uint32 bAllSelectedStaticMeshesHaveCollisionModels:1;
 	/** true if a brush is in the selection */
 	uint32 bHaveBrush:1;
+	/** true if there is a bsp brush in the selection */
+	uint32 bHaveBSPBrush:1;
 	/** true if a volume brush is in the selection */
 	uint32 bHaveVolume:1;
 	/** true if a builder brush is in the selection */
@@ -131,6 +133,7 @@ struct FSelectedActorInfo
 		, bAllSelectedActorsBelongToSameWorld(true)
 		, bAllSelectedStaticMeshesHaveCollisionModels(true)
 		, bHaveBrush(false)
+		, bHaveBSPBrush(false)
 		, bHaveVolume(false)
 		, bHaveBuilderBrush(false)
 		, bHaveStaticMeshComponent(false)
