Merging CL#2453328 using UT_To_UE4. Fixing a crash when trying to convert an enum to a string or text when the value of the enum is outside it's bounds.

diff --git a/Engine/Source/Editor/PropertyEditor/Private/PropertyHandleImpl.cpp b/Engine/Source/Editor/PropertyEditor/Private/PropertyHandleImpl.cpp
--- a/Engine/Source/Editor/PropertyEditor/Private/PropertyHandleImpl.cpp
+++ b/Engine/Source/Editor/PropertyEditor/Private/PropertyHandleImpl.cpp
@@ -70,11 +70,18 @@ FPropertyAccess::Result FPropertyValueImpl::GetPropertyValueString( FString& Out
 				{
 					const uint8 EnumValueIndex = ByteProperty->GetPropertyValue(ValueAddress);
 
-					// See if we specified an alternate name for this value using metadata
-					OutString = ByteProperty->Enum->GetDisplayNameText(EnumValueIndex).ToString();
-					if(!bAllowAlternateDisplayValue || OutString.Len() == 0) 
+					if (EnumValueIndex >= 0 && EnumValueIndex < ByteProperty->Enum->NumEnums())
 					{
-						OutString = ByteProperty->Enum->GetEnumName(EnumValueIndex);
+						// See if we specified an alternate name for this value using metadata
+						OutString = ByteProperty->Enum->GetDisplayNameText(EnumValueIndex).ToString();
+						if(!bAllowAlternateDisplayValue || OutString.Len() == 0) 
+						{
+							OutString = ByteProperty->Enum->GetEnumName(EnumValueIndex);
+						}
+					}
+					else
+					{
+						Result = FPropertyAccess::Fail;
 					}
 				}
 			}
@@ -123,11 +130,18 @@ FPropertyAccess::Result FPropertyValueImpl::GetPropertyValueText( FText& OutText
 				{
 					const uint8 EnumValueIndex = ByteProperty->GetPropertyValue(ValueAddress);
 					
-					// See if we specified an alternate name for this value using metadata
-					OutText = ByteProperty->Enum->GetDisplayNameText(EnumValueIndex);
-					if(!bAllowAlternateDisplayValue || OutText.IsEmpty()) 
+					if (EnumValueIndex >= 0 && EnumValueIndex < ByteProperty->Enum->NumEnums())
+					{
+						// See if we specified an alternate name for this value using metadata
+						OutText = ByteProperty->Enum->GetDisplayNameText(EnumValueIndex);
+						if(!bAllowAlternateDisplayValue || OutText.IsEmpty()) 
+						{
+							OutText = ByteProperty->Enum->GetEnumText(EnumValueIndex);
+						}
+					}
+					else
 					{
-						OutText = ByteProperty->Enum->GetEnumText(EnumValueIndex);
+						Result = FPropertyAccess::Fail;
 					}
 				}
 				else
