Fixed color picker constantly closing when editing some blueprint properties

diff --git a/Engine/Source/Editor/PropertyEditor/Private/PropertyEditorModule.cpp b/Engine/Source/Editor/PropertyEditor/Private/PropertyEditorModule.cpp
--- a/Engine/Source/Editor/PropertyEditor/Private/PropertyEditorModule.cpp
+++ b/Engine/Source/Editor/PropertyEditor/Private/PropertyEditorModule.cpp
@@ -520,8 +520,6 @@ void FPropertyEditorModule::UpdatePropertyViews( const TArray<UObject*>& NewObje
 
 void FPropertyEditorModule::ReplaceViewedObjects( const TMap<UObject*, UObject*>& OldToNewObjectMap )
 {
-	DestroyColorPicker();
-
 	// Replace objects from detail views
 	for( int32 ViewIndex = 0; ViewIndex < AllDetailViews.Num(); ++ViewIndex )
 	{
