Editor/Engine - The asset picker property editor now supports filtering on interfaces.  It expands any AllowedClasses that are interfaces into the classes that implement the interface that are in memory.  Now using it on the Blendables array in the post process volume to limit what is displayed.  The blendables are also no longer in the Misc category.

diff --git a/Engine/Source/Editor/PropertyEditor/Private/UserInterface/PropertyEditor/SPropertyEditorAsset.cpp b/Engine/Source/Editor/PropertyEditor/Private/UserInterface/PropertyEditor/SPropertyEditorAsset.cpp
--- a/Engine/Source/Editor/PropertyEditor/Private/UserInterface/PropertyEditor/SPropertyEditorAsset.cpp
+++ b/Engine/Source/Editor/PropertyEditor/Private/UserInterface/PropertyEditor/SPropertyEditorAsset.cpp
@@ -83,8 +83,17 @@ void SPropertyEditorAsset::Construct( const FArguments& InArgs, const TSharedPtr
 		ObjectClass = ObjectProperty->PropertyClass;
 		bIsActor = ObjectProperty->PropertyClass->IsChildOf( AActor::StaticClass() );
 
-		FString ClassFilterString = NodeProperty->GetMetaData("AllowedClasses");
-		if(ClassFilterString.IsEmpty())
+		FString ClassFilterString;
+		if ( UArrayProperty* ArrayParent = Cast<UArrayProperty>(NodeProperty->GetOuter()) )
+		{
+			ClassFilterString = ArrayParent->GetMetaData("AllowedClasses");
+		}
+		else
+		{
+			ClassFilterString = NodeProperty->GetMetaData("AllowedClasses");
+		}
+
+		if ( ClassFilterString.IsEmpty() )
 		{
 			CustomClassFilters.Add(ObjectClass);
 		}
@@ -98,13 +107,30 @@ void SPropertyEditorAsset::Construct( const FArguments& InArgs, const TSharedPtr
 				const FString& ClassName = *It;
 
 				UClass* Class = FindObject<UClass>(ANY_PACKAGE, *ClassName);
-				if(!Class)
+				
+				if ( !Class )
 				{
 					Class = LoadObject<UClass>(nullptr, *ClassName);
 				}
-				if(Class)
+
+				if ( Class )
 				{
-					CustomClassFilters.Add(Class);
+					// If the class is an interface, expand it to be all classes in memory that implement the class.
+					if ( Class->HasAnyClassFlags(CLASS_Interface) )
+					{
+						for ( TObjectIterator<UClass> ClassIt; ClassIt; ++ClassIt )
+						{
+							UClass* const ClassWithInterface = ( *ClassIt );
+							if ( ClassWithInterface->ImplementsInterface(Class) )
+							{
+								CustomClassFilters.Add(ClassWithInterface);
+							}
+						}
+					}
+					else
+					{
+						CustomClassFilters.Add(Class);
+					}
 				}
 			}
 		}
@@ -124,7 +150,7 @@ void SPropertyEditorAsset::Construct( const FArguments& InArgs, const TSharedPtr
 	{
 		NewAssetFactories = InArgs._NewAssetFactories.GetValue();
 	}
-	else if (ObjectClass->GetClass() != UObject::StaticClass())
+	else if (ObjectClass != UObject::StaticClass())
 	{
 		TArray<const UClass*> AllowedClasses;
 		AllowedClasses.Add(ObjectClass);
diff --git a/Engine/Source/Runtime/Engine/Classes/Engine/Scene.h b/Engine/Source/Runtime/Engine/Classes/Engine/Scene.h
--- a/Engine/Source/Runtime/Engine/Classes/Engine/Scene.h
+++ b/Engine/Source/Runtime/Engine/Classes/Engine/Scene.h
@@ -883,7 +883,7 @@ struct FPostProcessSettings
 	 * Allows custom post process materials to be defined, using a MaterialInstance with the same Material as its parent to allow blending.
 	 * Make sure you use the "PostProcess" domain type. This can be used for any UObject object implementing the IBlendableInterface (e.g. could be used to fade weather settings).
 	 */
-	UPROPERTY(EditAnywhere, BlueprintReadWrite, Category=Misc)
+	UPROPERTY(EditAnywhere, BlueprintReadWrite, Category="Custom Post Process Effects", meta=( AllowedClasses="BlendableInterface" ))
 	TArray<UObject*> Blendables;
 
 	// good start values for a new volume, by default no value is overriding
