Allow property node to be overriden as read-only

diff --git a/Engine/Source/Editor/PropertyEditor/Private/IDetailsViewPrivate.h b/Engine/Source/Editor/PropertyEditor/Private/IDetailsViewPrivate.h
--- a/Engine/Source/Editor/PropertyEditor/Private/IDetailsViewPrivate.h
+++ b/Engine/Source/Editor/PropertyEditor/Private/IDetailsViewPrivate.h
@@ -77,6 +77,11 @@ public:
 	 */
 	virtual bool IsPropertyVisible( const struct FPropertyAndParent& PropertyAndParent ) const = 0;
 
+	/**
+	 * @return True if the property is visible
+	 */
+	virtual bool IsPropertyReadOnly( const struct FPropertyAndParent& PropertyAndParent ) const = 0;
+
 	virtual TSharedPtr<IDetailKeyframeHandler> GetKeyframeHandler() = 0;
 
 	virtual TSharedPtr<IDetailPropertyExtensionHandler> GetExtensionHandler() = 0;
diff --git a/Engine/Source/Editor/PropertyEditor/Private/PropertyNode.cpp b/Engine/Source/Editor/PropertyEditor/Private/PropertyNode.cpp
--- a/Engine/Source/Editor/PropertyEditor/Private/PropertyNode.cpp
+++ b/Engine/Source/Editor/PropertyEditor/Private/PropertyNode.cpp
@@ -535,8 +535,8 @@ bool FPropertyNode::IsEditConst() const
 	// Ask the objects whether this property can be changed
 	const FObjectPropertyNode* ObjectPropertyNode = FindObjectItemParent();
 
-	bool bIsEditConst = false;
-	if (Property != NULL && ObjectPropertyNode)
+	bool bIsEditConst = (HasNodeFlags(EPropertyNodeFlags::IsReadOnly) != 0);
+	if (!bIsEditConst && Property != NULL && ObjectPropertyNode)
 	{
 		bIsEditConst = (Property->PropertyFlags & CPF_EditConst) ? true : false;
 		if (!bIsEditConst)
diff --git a/Engine/Source/Editor/PropertyEditor/Private/PropertyNode.h b/Engine/Source/Editor/PropertyEditor/Private/PropertyNode.h
--- a/Engine/Source/Editor/PropertyEditor/Private/PropertyNode.h
+++ b/Engine/Source/Editor/PropertyEditor/Private/PropertyNode.h
@@ -40,6 +40,9 @@ namespace EPropertyNodeFlags
 	const Type	RequiresValidation				= 1 << 18; /** true if this node could unexpectedly change (array changes, editinlinenew changes) */
 
 	const Type	ShouldShowDisableEditOnInstance = 1 << 19; /** true if this node should show child properties marked CPF_DisableEditOnInstance */
+
+	const Type	IsReadOnly						= 1 << 20; /** true if this node is overriden to appear as read-only */
+
 	const Type 	NoFlags							= 0;
 
 };
diff --git a/Engine/Source/Editor/PropertyEditor/Private/SDetailsViewBase.cpp b/Engine/Source/Editor/PropertyEditor/Private/SDetailsViewBase.cpp
--- a/Engine/Source/Editor/PropertyEditor/Private/SDetailsViewBase.cpp
+++ b/Engine/Source/Editor/PropertyEditor/Private/SDetailsViewBase.cpp
@@ -298,6 +298,11 @@ void SDetailsViewBase::SetIsPropertyVisibleDelegate(FIsPropertyVisible InIsPrope
 	IsPropertyVisibleDelegate = InIsPropertyVisible;
 }
 
+void SDetailsViewBase::SetIsPropertyReadOnlyDelegate(FIsPropertyReadOnly InIsPropertyReadOnly)
+{
+	IsPropertyReadOnlyDelegate = InIsPropertyReadOnly;
+}
+
 void SDetailsViewBase::SetIsPropertyEditingEnabledDelegate(FIsPropertyEditingEnabled IsPropertyEditingEnabled)
 {
 	IsPropertyEditingEnabledDelegate = IsPropertyEditingEnabled;
@@ -391,6 +396,11 @@ bool SDetailsViewBase::IsPropertyVisible( const FPropertyAndParent& PropertyAndP
 	return IsPropertyVisibleDelegate.IsBound() ? IsPropertyVisibleDelegate.Execute(PropertyAndParent) : true;
 }
 
+bool SDetailsViewBase::IsPropertyReadOnly( const FPropertyAndParent& PropertyAndParent ) const
+{
+	return IsPropertyReadOnlyDelegate.IsBound() ? IsPropertyReadOnlyDelegate.Execute(PropertyAndParent) : false;
+}
+
 TSharedPtr<IPropertyUtilities> SDetailsViewBase::GetPropertyUtilities()
 {
 	return PropertyUtilities;
@@ -1118,6 +1128,11 @@ void SDetailsViewBase::UpdatePropertyMapRecursive(FPropertyNode& InNode, FDetail
 						CategoryName = PropertyCatagoryName;
 					}
 
+					if (IsPropertyReadOnly(PropertyAndParent))
+					{
+						ChildNode.SetNodeFlags(EPropertyNodeFlags::IsReadOnly, true);
+					}
+
 					// Add a property to the default category
 					FDetailCategoryImpl& CategoryImpl = InDetailLayout.DefaultCategory(CategoryName);
 					CategoryImpl.AddPropertyNode(ChildNodePtr.ToSharedRef(), InstanceName);
diff --git a/Engine/Source/Editor/PropertyEditor/Private/SDetailsViewBase.h b/Engine/Source/Editor/PropertyEditor/Private/SDetailsViewBase.h
--- a/Engine/Source/Editor/PropertyEditor/Private/SDetailsViewBase.h
+++ b/Engine/Source/Editor/PropertyEditor/Private/SDetailsViewBase.h
@@ -146,6 +146,11 @@ public:
 	 */
 	virtual void SetIsPropertyVisibleDelegate(FIsPropertyVisible InIsPropertyVisible) override;
 
+	/**
+	 * Sets a delegate to call to determine if a specific property should be read-only in this instance of the details view
+	 */ 
+	virtual void SetIsPropertyReadOnlyDelegate( FIsPropertyReadOnly InIsPropertyReadOnly ) override;
+
 	/**
 	 * Sets a delegate to call to determine if the properties  editing is enabled
 	 */ 
@@ -233,6 +238,11 @@ public:
 	 */
 	virtual bool IsPropertyVisible( const struct FPropertyAndParent& PropertyAndParent ) const override;
 
+	/**
+	 * @return True if the property is visible
+	 */
+	virtual bool IsPropertyReadOnly( const struct FPropertyAndParent& PropertyAndParent ) const override;
+
 	/**
 	* Sets a delegate which is regardless of the objects being viewed to lay out generic details not specific to any object
 	*/
@@ -420,6 +430,8 @@ protected:
 	TArray< TSharedRef<IDetailTreeNode> > RootTreeNodes;
 	/** Delegate executed to determine if a property should be visible */
 	FIsPropertyVisible IsPropertyVisibleDelegate;
+	/** Delegate executed to determine if a property should be read-only */
+	FIsPropertyReadOnly IsPropertyReadOnlyDelegate;
 	/** Delegate called to see if a property editing is enabled */
 	FIsPropertyEditingEnabled IsPropertyEditingEnabledDelegate;
 	/** Delegate called when the details panel finishes editing a property (after post edit change is called) */
diff --git a/Engine/Source/Editor/PropertyEditor/Public/IDetailsView.h b/Engine/Source/Editor/PropertyEditor/Public/IDetailsView.h
--- a/Engine/Source/Editor/PropertyEditor/Public/IDetailsView.h
+++ b/Engine/Source/Editor/PropertyEditor/Public/IDetailsView.h
@@ -160,7 +160,12 @@ public:
 	 * Sets a delegate to call to determine if a specific property should be visible in this instance of the details view
 	 */ 
 	virtual void SetIsPropertyVisibleDelegate( FIsPropertyVisible InIsPropertyVisible ) = 0;
-		
+
+	/**
+	 * Sets a delegate to call to determine if a specific property should be read-only in this instance of the details view
+	 */ 
+	virtual void SetIsPropertyReadOnlyDelegate( FIsPropertyReadOnly InIsPropertyReadOnly ) = 0;
+
 	/**
 	 * Sets a delegate to call to layout generic details not specific to an object being viewed
 	 */ 
diff --git a/Engine/Source/Editor/PropertyEditor/Public/PropertyEditorDelegates.h b/Engine/Source/Editor/PropertyEditor/Public/PropertyEditorDelegates.h
--- a/Engine/Source/Editor/PropertyEditor/Public/PropertyEditorDelegates.h
+++ b/Engine/Source/Editor/PropertyEditor/Public/PropertyEditorDelegates.h
@@ -21,6 +21,9 @@ struct FPropertyAndParent
 /** Delegate called to see if a property should be visible */
 DECLARE_DELEGATE_RetVal_OneParam( bool, FIsPropertyVisible, const FPropertyAndParent& );
 
+/** Delegate called to see if a property should be visible */
+DECLARE_DELEGATE_RetVal_OneParam( bool, FIsPropertyReadOnly, const FPropertyAndParent& );
+
 /** Delegate called to get a detail layout for a specific object class */
 DECLARE_DELEGATE_RetVal( TSharedRef<class IDetailCustomization>, FOnGetDetailCustomizationInstance );
 
