Fix empty categories in UMG

diff --git a/Engine/Source/Editor/PropertyEditor/Private/DetailCategoryBuilderImpl.cpp b/Engine/Source/Editor/PropertyEditor/Private/DetailCategoryBuilderImpl.cpp
--- a/Engine/Source/Editor/PropertyEditor/Private/DetailCategoryBuilderImpl.cpp
+++ b/Engine/Source/Editor/PropertyEditor/Private/DetailCategoryBuilderImpl.cpp
@@ -600,12 +600,15 @@ void FDetailCategoryImpl::GenerateNodesFromCustomizations( const FCustomizationL
 			TSharedRef<FDetailItemNode> NewNode = MakeShareable( new FDetailItemNode( Customization, AsShared(), IsParentEnabled ) );
 			NewNode->Initialize();
 
-			if( CustomizationIndex == InCustomizationList.Num()-1 )
+			if( !NewNode->ShouldShowOnlyChildren() || NewNode->HasGeneratedChildren() )
 			{
-				bOutLastItemHasMultipleColumns = NewNode->HasMultiColumnWidget();
-			}
+				if(CustomizationIndex == InCustomizationList.Num()-1)
+				{
+					bOutLastItemHasMultipleColumns = NewNode->HasMultiColumnWidget();
+				}
 
-			OutNodeList.Add( NewNode );
+				OutNodeList.Add(NewNode);
+			}
 		}
 	}
 }
diff --git a/Engine/Source/Editor/PropertyEditor/Private/DetailItemNode.h b/Engine/Source/Editor/PropertyEditor/Private/DetailItemNode.h
--- a/Engine/Source/Editor/PropertyEditor/Private/DetailItemNode.h
+++ b/Engine/Source/Editor/PropertyEditor/Private/DetailItemNode.h
@@ -30,6 +30,11 @@ public:
 	 */
 	bool HasMultiColumnWidget() const;
 
+	/**
+	 * @return true if this node has any children (regardless of child visibility)
+	 */
+	bool HasGeneratedChildren() const { return Children.Num() > 0;}
+
 	/** IDetailTreeNode interface */
 	virtual IDetailsViewPrivate& GetDetailsView() const override{ return ParentCategory.Pin()->GetDetailsView(); }
 	virtual TSharedRef< ITableRow > GenerateNodeWidget( const TSharedRef<STableViewBase>& OwnerTable, const FDetailColumnSizeData& ColumnSizeData, const TSharedRef<IPropertyUtilities>& PropertyUtilities ) override;
