UMG - Don't allow adding keyframes from the details panel when there is no animation selected.

Jira: UE-11420

diff --git a/Engine/Source/Editor/PropertyEditor/Private/SDetailSingleItemRow.cpp b/Engine/Source/Editor/PropertyEditor/Private/SDetailSingleItemRow.cpp
--- a/Engine/Source/Editor/PropertyEditor/Private/SDetailSingleItemRow.cpp
+++ b/Engine/Source/Editor/PropertyEditor/Private/SDetailSingleItemRow.cpp
@@ -110,7 +110,6 @@ void SDetailSingleItemRow::Construct( const FArguments& InArgs, FDetailLayoutCus
 		if( bHasMultipleColumns )
 		{
 			NameWidget->SetEnabled(IsPropertyEditingEnabled);
-			KeyFrameButton->SetEnabled(IsPropertyEditingEnabled);
 			Widget = 
 				SNew( SSplitter )
 				.Style( FEditorStyle::Get(), "DetailsView.Splitter" )
@@ -164,7 +163,6 @@ void SDetailSingleItemRow::Construct( const FArguments& InArgs, FDetailLayoutCus
 		else
 		{
 			Row.WholeRowWidget.Widget->SetEnabled(IsPropertyEditingEnabled);
-			KeyFrameButton->SetEnabled(IsPropertyEditingEnabled);
 			Widget =
 				SNew( SHorizontalBox )
 				+ SHorizontalBox::Slot()
@@ -359,21 +357,24 @@ TSharedRef<SWidget> SDetailSingleItemRow::CreateKeyframeButton( FDetailLayoutCus
 		
 	}
 
-	return 
+	return
 		SNew(SButton)
 		.HAlign(HAlign_Left)
 		.VAlign(VAlign_Center)
 		.ContentPadding(0.0f)
 		.ButtonStyle(FEditorStyle::Get(), "Sequencer.AddKey.Details")
 		.Visibility(SetKeyVisibility)
-		.IsEnabled( this, &SDetailSingleItemRow::IsKeyframeButtonEnabled )
+		.IsEnabled( this, &SDetailSingleItemRow::IsKeyframeButtonEnabled, InTreeNode )
 		.ToolTipText( NSLOCTEXT("PropertyView", "AddKeyframeButton_ToolTip", "Adds a keyframe for this property to the current animation") )
 		.OnClicked( this, &SDetailSingleItemRow::OnAddKeyframeClicked );
 }
 
-bool SDetailSingleItemRow::IsKeyframeButtonEnabled() const
+bool SDetailSingleItemRow::IsKeyframeButtonEnabled(TSharedRef<IDetailTreeNode> InTreeNode) const
 {
-	return KeyframeHandler.IsValid() ? KeyframeHandler.Pin()->IsPropertyKeyingEnabled() : false;
+	return
+		InTreeNode->IsPropertyEditingEnabled().Get() &&
+		KeyframeHandler.IsValid() &&
+		KeyframeHandler.Pin()->IsPropertyKeyingEnabled();
 }
 
 FReply SDetailSingleItemRow::OnAddKeyframeClicked()
diff --git a/Engine/Source/Editor/PropertyEditor/Private/SDetailSingleItemRow.h b/Engine/Source/Editor/PropertyEditor/Private/SDetailSingleItemRow.h
--- a/Engine/Source/Editor/PropertyEditor/Private/SDetailSingleItemRow.h
+++ b/Engine/Source/Editor/PropertyEditor/Private/SDetailSingleItemRow.h
@@ -31,7 +31,7 @@ private:
 	const FSlateBrush* GetBorderImage() const;
 	TSharedRef<SWidget> CreateExtensionWidget( TSharedRef<SWidget> ValueWidget, FDetailLayoutCustomization& InCustomization, TSharedRef<IDetailTreeNode> InTreeNode );
 	TSharedRef<SWidget> CreateKeyframeButton( FDetailLayoutCustomization& InCustomization, TSharedRef<IDetailTreeNode> InTreeNode );
-	bool IsKeyframeButtonEnabled() const;
+	bool IsKeyframeButtonEnabled(TSharedRef<IDetailTreeNode> InTreeNode) const;
 	FReply OnAddKeyframeClicked();
 	bool IsHighlighted() const;
 private:
