Reset to default now propagates changes to subclasses

diff --git a/Engine/Source/Editor/PropertyEditor/Private/PropertyNode.cpp b/Engine/Source/Editor/PropertyEditor/Private/PropertyNode.cpp
--- a/Engine/Source/Editor/PropertyEditor/Private/PropertyNode.cpp
+++ b/Engine/Source/Editor/PropertyEditor/Private/PropertyNode.cpp
@@ -1509,12 +1509,14 @@ void FPropertyNode::ResetToDefault( FNotifyHook* InNotifyHook )
 				{
 					FPropertyItemValueDataTrackerSlate& ValueTracker = *ValueTrackerPtr;
 
+					bool bIsGameWorld = false;
 					// If the object we are modifying is in the PIE world, than make the PIE world the active
 					// GWorld.  Assumes all objects managed by this property window belong to the same world.
 					UWorld* OldGWorld = NULL;
 					if ( GUnrealEd && GUnrealEd->PlayWorld && !GUnrealEd->bIsSimulatingInEditor && Object->IsIn(GUnrealEd->PlayWorld))
 					{
 						OldGWorld = SetPlayInEditorWorld(GUnrealEd->PlayWorld);
+						bIsGameWorld = true;
 					}
 
 					if( !bNotifiedPreChange )
@@ -1524,6 +1526,13 @@ void FPropertyNode::ResetToDefault( FNotifyHook* InNotifyHook )
 						bNotifiedPreChange = true;
 					}
 
+					// Cache the value of the property before modifying it.
+					FString PreviousValue;
+					TheProperty->ExportText_Direct(PreviousValue, ValueTracker.GetPropertyValueAddress(), ValueTracker.GetPropertyValueAddress(), NULL, 0);
+
+			
+					FString PreviousArrayValue;
+
 					if( ValueTracker.GetPropertyDefaultAddress() != NULL )
 					{
 						UObject* RootObject = ValueTracker.GetTopLevelObject();
@@ -1534,6 +1543,20 @@ void FPropertyNode::ResetToDefault( FNotifyHook* InNotifyHook )
 						// be used in a static array
 						UArrayProperty* ArrayProp = Cast<UArrayProperty>(TheProperty);
 
+
+					
+						FPropertyNode* ParentNode = GetParentNode();
+						if(ParentNode != NULL && ParentNode->GetProperty() && ParentNode->GetProperty()->IsA(UArrayProperty::StaticClass()))
+						{
+							UArrayProperty* ArrayProp = Cast<UArrayProperty>(ParentNode->GetProperty());
+							if(ArrayProp->Inner == TheProperty)
+							{
+								uint8* Addr = ParentNode->GetValueBaseAddress((uint8*)Object);
+
+								ArrayProp->ExportText_Direct(PreviousArrayValue, Addr, Addr, NULL, 0);
+							}
+						}
+
 						if( ArrayProp != NULL )
 						{
 							TheProperty->CopyCompleteValue(ValueTracker.GetPropertyValueAddress(), ValueTracker.GetPropertyDefaultAddress());
@@ -1601,6 +1624,17 @@ void FPropertyNode::ResetToDefault( FNotifyHook* InNotifyHook )
 						TheProperty->ClearValue(ValueTracker.GetPropertyValueAddress());
 					}
 
+					// Cache the value of the property after having modified it.
+					FString ValueAfterImport;
+					Property->ExportText_Direct(ValueAfterImport, ValueTracker.GetPropertyValueAddress(), ValueTracker.GetPropertyValueAddress(), NULL, 0);
+
+					if((Object->HasAnyFlags(RF_ClassDefaultObject | RF_ArchetypeObject) ||
+						(Object->HasAnyFlags(RF_DefaultSubObject) && Object->GetOuter()->HasAnyFlags(RF_ClassDefaultObject | RF_ArchetypeObject))) &&
+						!bIsGameWorld)
+					{
+						PropagatePropertyChange(Object, *ValueAfterImport, PreviousArrayValue.IsEmpty() ? PreviousValue : PreviousArrayValue);
+					}
+
 					if(OldGWorld)
 					{
 						// restore the original (editor) GWorld
