Editor - Giving the DetailsView more capability to allow it to be set to automatically handle the EditDefaultsOnly flag for showing/hiding nodes when selecting a CDO instead of needing to monitor it, and show/hide them manually.

diff --git a/Engine/Source/Editor/BehaviorTreeEditor/Private/BehaviorTreeEditor.cpp b/Engine/Source/Editor/BehaviorTreeEditor/Private/BehaviorTreeEditor.cpp
--- a/Engine/Source/Editor/BehaviorTreeEditor/Private/BehaviorTreeEditor.cpp
+++ b/Engine/Source/Editor/BehaviorTreeEditor/Private/BehaviorTreeEditor.cpp
@@ -906,7 +906,7 @@ void FBehaviorTreeEditor::CreateInternalWidgets()
 	FPropertyEditorModule& PropertyEditorModule = FModuleManager::GetModuleChecked<FPropertyEditorModule>( "PropertyEditor" );
 	FDetailsViewArgs DetailsViewArgs( false, false, true, FDetailsViewArgs::HideNameArea, false );
 	DetailsViewArgs.NotifyHook = this;
-	DetailsViewArgs.bShowDisableEditOnInstanceNodes = false;
+	DetailsViewArgs.DefaultsOnlyVisibility = FDetailsViewArgs::EEditDefaultsOnlyNodeVisibility::Hide;
 	DetailsView = PropertyEditorModule.CreateDetailView( DetailsViewArgs );
 	DetailsView->SetObject( NULL );
 	DetailsView->SetIsPropertyEditingEnabledDelegate(FIsPropertyEditingEnabled::CreateSP(this, &FBehaviorTreeEditor::IsPropertyEditable));
diff --git a/Engine/Source/Editor/BehaviorTreeEditor/Private/SBehaviorTreeDiff.cpp b/Engine/Source/Editor/BehaviorTreeEditor/Private/SBehaviorTreeDiff.cpp
--- a/Engine/Source/Editor/BehaviorTreeEditor/Private/SBehaviorTreeDiff.cpp
+++ b/Engine/Source/Editor/BehaviorTreeEditor/Private/SBehaviorTreeDiff.cpp
@@ -396,7 +396,7 @@ void SBehaviorTreeDiff::FBehaviorTreeDiffPanel::GeneratePanel(UEdGraph* Graph, U
 
 	FPropertyEditorModule& PropertyEditorModule = FModuleManager::GetModuleChecked<FPropertyEditorModule>( "PropertyEditor" );
 	FDetailsViewArgs DetailsViewArgs( false, false, true, FDetailsViewArgs::ObjectsUseNameArea, false );
-	DetailsViewArgs.bShowDisableEditOnInstanceNodes = false;
+	DetailsViewArgs.DefaultsOnlyVisibility = FDetailsViewArgs::EEditDefaultsOnlyNodeVisibility::Hide;
 	DetailsView = PropertyEditorModule.CreateDetailView( DetailsViewArgs );
 	DetailsView->SetObject( NULL );
 	DetailsView->SetIsPropertyEditingEnabledDelegate(FIsPropertyEditingEnabled::CreateRaw(this, &SBehaviorTreeDiff::FBehaviorTreeDiffPanel::IsPropertyEditable));
diff --git a/Engine/Source/Editor/Blutility/Private/GlobalBlutilityDialog.cpp b/Engine/Source/Editor/Blutility/Private/GlobalBlutilityDialog.cpp
--- a/Engine/Source/Editor/Blutility/Private/GlobalBlutilityDialog.cpp
+++ b/Engine/Source/Editor/Blutility/Private/GlobalBlutilityDialog.cpp
@@ -129,7 +129,7 @@ void FGlobalBlutilityDialog::CreateInternalWidgets()
 	// Create a details view
 	FPropertyEditorModule& PropertyEditorModule = FModuleManager::GetModuleChecked<FPropertyEditorModule>("PropertyEditor");
 	FDetailsViewArgs DetailsViewArgs(/*bUpdateFromSelection=*/ false, /*bLockable=*/ false, /*bAllowSearch=*/ false, FDetailsViewArgs::HideNameArea, /*bHideSelectionTip=*/ true);
-	DetailsViewArgs.bShowDisableEditOnInstanceNodes = false;
+	DetailsViewArgs.DefaultsOnlyVisibility = FDetailsViewArgs::EEditDefaultsOnlyNodeVisibility::Hide;
 	DetailsView = PropertyEditorModule.CreateDetailView(DetailsViewArgs);
 }
 
diff --git a/Engine/Source/Editor/LevelEditor/Private/SActorDetails.cpp b/Engine/Source/Editor/LevelEditor/Private/SActorDetails.cpp
--- a/Engine/Source/Editor/LevelEditor/Private/SActorDetails.cpp
+++ b/Engine/Source/Editor/LevelEditor/Private/SActorDetails.cpp
@@ -79,7 +79,7 @@ void SActorDetails::Construct(const FArguments& InArgs, const FName TabIdentifie
 	DetailsViewArgs.ViewIdentifier = TabIdentifier;
 	DetailsViewArgs.bCustomNameAreaLocation = true;
 	DetailsViewArgs.bCustomFilterAreaLocation = true;
-	DetailsViewArgs.bShowDisableEditOnInstanceNodes = false;
+	DetailsViewArgs.DefaultsOnlyVisibility = FDetailsViewArgs::EEditDefaultsOnlyNodeVisibility::Hide;
 
 	DetailsView = PropPlugin.CreateDetailView(DetailsViewArgs);
 
diff --git a/Engine/Source/Editor/PropertyEditor/Private/SDetailsView.cpp b/Engine/Source/Editor/PropertyEditor/Private/SDetailsView.cpp
--- a/Engine/Source/Editor/PropertyEditor/Private/SDetailsView.cpp
+++ b/Engine/Source/Editor/PropertyEditor/Private/SDetailsView.cpp
@@ -299,7 +299,7 @@ void SDetailsView::SetObjects( const TArray<UObject*>& InObjects, bool bForceRef
 			ObjectWeakPtrs.Add( *ObjectIter );
 		}
 
- 		if( bForceRefresh || ShouldSetNewObjects( ObjectWeakPtrs ) )
+		if( bForceRefresh || ShouldSetNewObjects( ObjectWeakPtrs ) )
 		{
 			SetObjectArrayPrivate( ObjectWeakPtrs );
 		}
@@ -402,15 +402,15 @@ void SDetailsView::SetObjectArrayPrivate( const TArray< TWeakObjectPtr< UObject
 	// Selected actors for building SelectedActorInfo
 	TArray<AActor*> SelectedRawActors;
 
-	bViewingClassDefaultObject = false;
+	bViewingClassDefaultObject = InObjects.Num() > 0 ? true : false;
 	bool bOwnedByLockedLevel = false;
-	for( int32 ObjectIndex = 0 ; ObjectIndex < InObjects.Num() ; ++ObjectIndex )
+	for( int32 ObjectIndex = 0 ; ObjectIndex < InObjects.Num(); ++ObjectIndex )
 	{
 		TWeakObjectPtr< UObject > Object = InObjects[ObjectIndex];
 
 		if( Object.IsValid() )
 		{
-			bViewingClassDefaultObject |= Object->HasAnyFlags( RF_ClassDefaultObject );
+			bViewingClassDefaultObject &= Object->HasAnyFlags( RF_ClassDefaultObject );
 
 			RootPropertyNode->AddObject( Object.Get() );
 			SelectedObjects.Add( Object );
@@ -586,7 +586,21 @@ void SDetailsView::PostSetObject()
 	InitParams.ArrayIndex = INDEX_NONE;
 	InitParams.bAllowChildren = true;
 	InitParams.bForceHiddenPropertyVisibility =  FPropertySettings::Get().ShowHiddenProperties();
-	InitParams.bCreateDisableEditOnInstanceNodes = DetailsViewArgs.bShowDisableEditOnInstanceNodes;
+
+	switch ( DetailsViewArgs.DefaultsOnlyVisibility )
+	{
+	case FDetailsViewArgs::EEditDefaultsOnlyNodeVisibility::Hide:
+		InitParams.bCreateDisableEditOnInstanceNodes = false;
+		break;
+	case FDetailsViewArgs::EEditDefaultsOnlyNodeVisibility::Show:
+		InitParams.bCreateDisableEditOnInstanceNodes = true;
+		break;
+	case FDetailsViewArgs::EEditDefaultsOnlyNodeVisibility::Automatic:
+		InitParams.bCreateDisableEditOnInstanceNodes = HasClassDefaultObject();
+		break;
+	default:
+		check(false);
+	}
 
 	RootPropertyNode->InitNode( InitParams );
 
diff --git a/Engine/Source/Editor/PropertyEditor/Public/IDetailsView.h b/Engine/Source/Editor/PropertyEditor/Public/IDetailsView.h
--- a/Engine/Source/Editor/PropertyEditor/Public/IDetailsView.h
+++ b/Engine/Source/Editor/PropertyEditor/Public/IDetailsView.h
@@ -20,6 +20,17 @@ struct FDetailsViewArgs
 		/** Components and actors use the name area. Components will display their actor owner as the name */
 		ComponentsAndActorsUseNameArea,
 	};
+
+	enum class EEditDefaultsOnlyNodeVisibility : uint8
+	{
+		/** Always show nodes that have the EditDefaultsOnly aka CPF_DisableEditOnInstance flag. */
+		Show,
+		/** Always hide nodes that have the EditDefaultsOnly aka CPF_DisableEditOnInstance flag. */
+		Hide,
+		/** Let the details panel control it. If the CDO is selected EditDefaultsOnly nodes will be visible, otherwise false. */
+		Automatic,
+	};
+
 	/** Identifier for this details view; NAME_None if this view is anonymous */
 	FName ViewIdentifier;
 	/** Notify hook to call when properties are changed */
@@ -48,8 +59,8 @@ struct FDetailsViewArgs
 	uint32 bCustomNameAreaLocation : 1;
 	/** If true, the filter area will be created but will not be displayed so it can be placed in a custom location.  */
 	uint32 bCustomFilterAreaLocation : 1;
-	/** If true, nodes marked CPF_DisableEditOnInstance will be shown */
-	uint32 bShowDisableEditOnInstanceNodes : 1;
+	/** Controls how CPF_DisableEditOnInstance nodes will be treated */
+	EEditDefaultsOnlyNodeVisibility DefaultsOnlyVisibility;
 
 	/** Default constructor */
 	FDetailsViewArgs( const bool InUpdateFromSelection = false
@@ -74,7 +85,7 @@ struct FDetailsViewArgs
 		, bShowDifferingPropertiesOption(false)
 		, bCustomNameAreaLocation(false)
 		, bCustomFilterAreaLocation(false)
-		, bShowDisableEditOnInstanceNodes(true)
+		, DefaultsOnlyVisibility(EEditDefaultsOnlyNodeVisibility::Show)
 	{
 	}
 };
