Property Editor: Added support for passing in a command list when creating a details view that can be queried by customizations [UE-13643]

diff --git a/Engine/Source/Editor/PropertyEditor/Private/SDetailsViewBase.cpp b/Engine/Source/Editor/PropertyEditor/Private/SDetailsViewBase.cpp
--- a/Engine/Source/Editor/PropertyEditor/Private/SDetailsViewBase.cpp
+++ b/Engine/Source/Editor/PropertyEditor/Private/SDetailsViewBase.cpp
@@ -457,6 +457,11 @@ TSharedPtr<SWidget> SDetailsViewBase::GetFilterAreaWidget()
 	return DetailsViewArgs.bCustomFilterAreaLocation ? FilterRow : nullptr;
 }
 
+TSharedPtr<class FUICommandList> SDetailsViewBase::GetHostCommandList() const
+{
+	return DetailsViewArgs.HostCommandList;
+}
+
 /** 
  * Hides or shows properties based on the passed in filter text
  * 
diff --git a/Engine/Source/Editor/PropertyEditor/Private/SDetailsViewBase.h b/Engine/Source/Editor/PropertyEditor/Private/SDetailsViewBase.h
--- a/Engine/Source/Editor/PropertyEditor/Private/SDetailsViewBase.h
+++ b/Engine/Source/Editor/PropertyEditor/Private/SDetailsViewBase.h
@@ -391,6 +391,7 @@ protected:
 
 	virtual TSharedPtr<SWidget> GetNameAreaWidget() override;
 	virtual TSharedPtr<SWidget> GetFilterAreaWidget() override;
+	virtual TSharedPtr<class FUICommandList> GetHostCommandList() const override;
 
 	/** 
 	 * Hides or shows properties based on the passed in filter text
diff --git a/Engine/Source/Editor/PropertyEditor/Public/IDetailsView.h b/Engine/Source/Editor/PropertyEditor/Public/IDetailsView.h
--- a/Engine/Source/Editor/PropertyEditor/Public/IDetailsView.h
+++ b/Engine/Source/Editor/PropertyEditor/Public/IDetailsView.h
@@ -61,7 +61,10 @@ struct FDetailsViewArgs
 	uint32 bCustomFilterAreaLocation : 1;
 	/** Controls how CPF_DisableEditOnInstance nodes will be treated */
 	EEditDefaultsOnlyNodeVisibility DefaultsOnlyVisibility;
+	/** The command list from the host of the details view, allowing child widgets to bind actions with a bound chord */
+	TSharedPtr<class FUICommandList> HostCommandList;
 
+public:
 	/** Default constructor */
 	FDetailsViewArgs( const bool InUpdateFromSelection = false
 					, const bool InLockable = false
@@ -243,4 +246,7 @@ public:
 
 	/** Returns the search area widget used to display search and view options so it can be placed in a custom location.  Note FDetailsViewArgs.bCustomFilterAreaLocation must be true */
 	virtual TSharedPtr<SWidget> GetFilterAreaWidget() = 0;
+
+	/** Returns the command list of the hosting toolkit (can be nullptr if the widget that contains the details panel didn't route a command list in) */
+	virtual TSharedPtr<class FUICommandList> GetHostCommandList() const = 0;
 };
