Updated SPropertyEditorAsset to account for AllowedClasses metadata on properties when being created without a valid property editor

diff --git a/Engine/Source/Editor/DetailCustomizations/Private/SlateBrushCustomization.cpp b/Engine/Source/Editor/DetailCustomizations/Private/SlateBrushCustomization.cpp
--- a/Engine/Source/Editor/DetailCustomizations/Private/SlateBrushCustomization.cpp
+++ b/Engine/Source/Editor/DetailCustomizations/Private/SlateBrushCustomization.cpp
@@ -1060,10 +1060,6 @@ class SBrushResourceObjectBox : public SCompoundWidget
 		FSimpleDelegate OnBrushResourceChangedDelegate = FSimpleDelegate::CreateSP(this, &SBrushResourceObjectBox::OnBrushResourceChanged);
 		ResourceObjectProperty->SetOnPropertyValueChanged(OnBrushResourceChangedDelegate);
 
-		TArray<const UClass*> SupportedClasses;
-		SupportedClasses.Add(UTexture2D::StaticClass());
-		SupportedClasses.Add(UMaterialInterface::StaticClass());
-
 		ChildSlot
 		[
 			SNew(SVerticalBox)
@@ -1073,8 +1069,6 @@ class SBrushResourceObjectBox : public SCompoundWidget
 				SNew(SObjectPropertyEntryBox)
 				.PropertyHandle(InResourceObjectProperty)
 				.ThumbnailPool(StructCustomizationUtils->GetThumbnailPool())
-				.NewAssetFactories(PropertyCustomizationHelpers::GetNewAssetFactoriesForClasses(SupportedClasses))
-				.OnShouldFilterAsset(this, &SBrushResourceObjectBox::OnFilterAssetPicker)
 				.OnObjectChanged(this, &SBrushResourceObjectBox::OnAssetPicked)
 			]
 			+ SVerticalBox::Slot()
@@ -1136,13 +1130,6 @@ private:
 		}
 	}
 
-	/** Called when the asset picker needs to be filtered */
-	bool OnFilterAssetPicker(const FAssetData& InAssetData) const
-	{
-		UClass* Class = InAssetData.GetClass();
-		return !Class->IsChildOf(UTexture2D::StaticClass()) && !Class->IsChildOf(UMaterialInterface::StaticClass());
-	}
-
 	/**
 	 * When the asset is picked if it's a material being used in the UI we 
 	 * automatically set the bUsedWithUI flag if it isn't already set.
diff --git a/Engine/Source/Editor/PropertyEditor/Private/UserInterface/PropertyEditor/SPropertyEditorAsset.cpp b/Engine/Source/Editor/PropertyEditor/Private/UserInterface/PropertyEditor/SPropertyEditorAsset.cpp
--- a/Engine/Source/Editor/PropertyEditor/Private/UserInterface/PropertyEditor/SPropertyEditorAsset.cpp
+++ b/Engine/Source/Editor/PropertyEditor/Private/UserInterface/PropertyEditor/SPropertyEditorAsset.cpp
@@ -72,28 +72,50 @@ void SPropertyEditorAsset::Construct( const FArguments& InArgs, const TSharedPtr
 	PropertyEditor = InPropertyEditor;
 	PropertyHandle = InArgs._PropertyHandle;
 	OnSetObject = InArgs._OnSetObject;
+	OnShouldFilterAsset = InArgs._OnShouldFilterAsset;
 
+	UProperty* Property = nullptr;
 	if(PropertyEditor.IsValid())
 	{
-		UProperty* NodeProperty = PropertyEditor->GetPropertyNode()->GetProperty();
-		UObjectPropertyBase* ObjectProperty = Cast<UObjectPropertyBase>( NodeProperty );
+		Property = PropertyEditor->GetPropertyNode()->GetProperty();
+		UObjectPropertyBase* ObjectProperty = Cast<UObjectPropertyBase>(Property);
 		check(ObjectProperty);
 
-		bAllowClear = !(NodeProperty->PropertyFlags & CPF_NoClear);
+		bAllowClear = !(Property->PropertyFlags & CPF_NoClear);
 		ObjectClass = ObjectProperty->PropertyClass;
 		bIsActor = ObjectProperty->PropertyClass->IsChildOf( AActor::StaticClass() );
+	}
+	else
+	{
+		bAllowClear = InArgs._AllowClear;
+		ObjectPath = InArgs._ObjectPath;
+		ObjectClass = InArgs._Class;
+		bIsActor = ObjectClass->IsChildOf( AActor::StaticClass() );
 
+		if (PropertyHandle.IsValid() && PropertyHandle->IsValidHandle())
+		{
+			Property = PropertyHandle->GetProperty();
+		}
+		else
+		{
+			CustomClassFilters.Add(ObjectClass);
+		}
+	}
+
+	// Account for the allowed classes specified in the property metadata
+	if (Property)
+	{
 		FString ClassFilterString;
-		if ( UArrayProperty* ArrayParent = Cast<UArrayProperty>(NodeProperty->GetOuter()) )
+		if (UArrayProperty* ArrayParent = Cast<UArrayProperty>(Property->GetOuter()))
 		{
 			ClassFilterString = ArrayParent->GetMetaData("AllowedClasses");
 		}
 		else
 		{
-			ClassFilterString = NodeProperty->GetMetaData("AllowedClasses");
+			ClassFilterString = Property->GetMetaData("AllowedClasses");
 		}
 
-		if ( ClassFilterString.IsEmpty() )
+		if (ClassFilterString.IsEmpty())
 		{
 			CustomClassFilters.Add(ObjectClass);
 		}
@@ -102,26 +124,26 @@ void SPropertyEditorAsset::Construct( const FArguments& InArgs, const TSharedPtr
 			TArray<FString> CustomClassFilterNames;
 			ClassFilterString.ParseIntoArray(CustomClassFilterNames, TEXT(","), true);
 
-			for(auto It = CustomClassFilterNames.CreateConstIterator(); It; ++It)
+			for (auto It = CustomClassFilterNames.CreateConstIterator(); It; ++It)
 			{
 				const FString& ClassName = *It;
 
 				UClass* Class = FindObject<UClass>(ANY_PACKAGE, *ClassName);
-				
-				if ( !Class )
+
+				if (!Class)
 				{
 					Class = LoadObject<UClass>(nullptr, *ClassName);
 				}
 
-				if ( Class )
+				if (Class)
 				{
 					// If the class is an interface, expand it to be all classes in memory that implement the class.
-					if ( Class->HasAnyClassFlags(CLASS_Interface) )
+					if (Class->HasAnyClassFlags(CLASS_Interface))
 					{
-						for ( TObjectIterator<UClass> ClassIt; ClassIt; ++ClassIt )
+						for (TObjectIterator<UClass> ClassIt; ClassIt; ++ClassIt)
 						{
-							UClass* const ClassWithInterface = ( *ClassIt );
-							if ( ClassWithInterface->ImplementsInterface(Class) )
+							UClass* const ClassWithInterface = (*ClassIt);
+							if (ClassWithInterface->ImplementsInterface(Class))
 							{
 								CustomClassFilters.Add(ClassWithInterface);
 							}
@@ -135,26 +157,14 @@ void SPropertyEditorAsset::Construct( const FArguments& InArgs, const TSharedPtr
 			}
 		}
 	}
-	else
-	{
-		bAllowClear = InArgs._AllowClear;
-		ObjectPath = InArgs._ObjectPath;
-		ObjectClass = InArgs._Class;
-		bIsActor = ObjectClass->IsChildOf( AActor::StaticClass() );
-
-		CustomClassFilters.Add(ObjectClass);
-	}
-	OnShouldFilterAsset = InArgs._OnShouldFilterAsset;
 
 	if (InArgs._NewAssetFactories.IsSet())
 	{
 		NewAssetFactories = InArgs._NewAssetFactories.GetValue();
 	}
-	else if (ObjectClass != UObject::StaticClass())
+	else if (CustomClassFilters.Num() > 1 || !CustomClassFilters.Contains(UObject::StaticClass()))
 	{
-		TArray<const UClass*> AllowedClasses;
-		AllowedClasses.Add(ObjectClass);
-		NewAssetFactories = PropertyCustomizationHelpers::GetNewAssetFactoriesForClasses(AllowedClasses);
+		NewAssetFactories = PropertyCustomizationHelpers::GetNewAssetFactoriesForClasses(CustomClassFilters);
 	}
 	
 	TSharedPtr<SHorizontalBox> HorizontalBox = NULL;
diff --git a/Engine/Source/Runtime/SlateCore/Public/Styling/SlateBrush.h b/Engine/Source/Runtime/SlateCore/Public/Styling/SlateBrush.h
--- a/Engine/Source/Runtime/SlateCore/Public/Styling/SlateBrush.h
+++ b/Engine/Source/Runtime/SlateCore/Public/Styling/SlateBrush.h
@@ -269,7 +269,7 @@ protected:
 	/**
 	 * The image to render for this brush, can be a UTexture2D or Material.
 	 */
-	UPROPERTY(EditAnywhere, BlueprintReadWrite, Category=Brush, meta=( DisplayThumbnail="true", ThumbnailSize="X=40 Y=40", DisplayName="Image" ))
+	UPROPERTY(EditAnywhere, BlueprintReadWrite, Category=Brush, meta=( DisplayThumbnail="true", ThumbnailSize="X=40 Y=40", DisplayName="Image", AllowedClasses="Texture2D,MaterialInterface" ))
 	UObject* ResourceObject;
 
 	/** The name of the rendering resource to use */
