Cleaned up FText property namespace/key menu to use lambdas instead of a locally nested struct and static methods.

diff --git a/Engine/Source/Editor/PropertyEditor/Private/PropertyCustomizationHelpers.cpp b/Engine/Source/Editor/PropertyEditor/Private/PropertyCustomizationHelpers.cpp
--- a/Engine/Source/Editor/PropertyEditor/Private/PropertyCustomizationHelpers.cpp
+++ b/Engine/Source/Editor/PropertyEditor/Private/PropertyCustomizationHelpers.cpp
@@ -1052,175 +1052,172 @@ void FMaterialList::AddMaterialItem( FDetailWidgetRow& Row, int32 CurrentSlot, c
 }
 
 
-TSharedRef<SWidget> PropertyCustomizationHelpers::MakeTextLocalizationButton(TSharedRef<IPropertyHandle> InPropertyHandle)
+TSharedRef<SWidget> PropertyCustomizationHelpers::MakeTextLocalizationButton(const TSharedRef<IPropertyHandle>& InPropertyHandle)
 {
-	struct Local
+	const auto& DirtyPropertyPackage = [=]()
 	{
-		static void DirtyPropertyPackage(TSharedRef<IPropertyHandle> InPropHandle)
+		// dirty packages as they will need to be re-serialized with the correct key
+		TArray<UObject*> Objects;
+		InPropertyHandle->GetOuterObjects(Objects);
+		if (Objects.Num() > 0 && Objects[0]->GetOutermost() != nullptr)
 		{
-			// dirty packages as they will need to be re-serialized with the correct key
-			TArray<UObject*> Objects;
-			InPropHandle->GetOuterObjects(Objects);
-			if (Objects.Num() > 0 && Objects[0]->GetOutermost() != nullptr)
+			Objects[0]->GetOutermost()->MarkPackageDirty();
+		}
+	};
+
+	const auto& HandleGenerateKeyClicked = [=]() -> FReply 
+	{
+		FText DisplayText;
+		if (InPropertyHandle->GetValueAsDisplayText(DisplayText) == FPropertyAccess::Success)
+		{
+			TSharedRef< FString, ESPMode::ThreadSafe > DisplayString = FTextInspector::GetSharedDisplayString(DisplayText);
+			TSharedPtr< FString, ESPMode::ThreadSafe > Namespace;
+			TSharedPtr< FString, ESPMode::ThreadSafe > Key;
+			FTextLocalizationManager::Get().FindKeyNamespaceFromDisplayString(DisplayString, Namespace, Key);
+
+			if (Key.IsValid())
 			{
-				Objects[0]->GetOutermost()->MarkPackageDirty();
+				*Key.Get() = FGuid::NewGuid().ToString();
+				DirtyPropertyPackage();
 			}
 		}
 
-		static FReply HandleGenerateKeyClicked(TSharedRef<IPropertyHandle> InPropHandle)
+		return FReply::Unhandled();
+	};
+
+	const auto& HandleKeyTextCommitted = [=](const FText& InText, ETextCommit::Type InCommitType)
+	{
+		FText DisplayText;
+		if (InPropertyHandle->GetValueAsDisplayText(DisplayText) == FPropertyAccess::Success)
 		{
-			FText DisplayText;
-			if (InPropHandle->GetValueAsDisplayText(DisplayText) == FPropertyAccess::Success)
-			{
-				TSharedRef< FString, ESPMode::ThreadSafe > DisplayString = FTextInspector::GetSharedDisplayString(DisplayText);
-				TSharedPtr< FString, ESPMode::ThreadSafe > Namespace;
-				TSharedPtr< FString, ESPMode::ThreadSafe > Key;
-				FTextLocalizationManager::Get().FindKeyNamespaceFromDisplayString(DisplayString, Namespace, Key);
+			TSharedRef< FString, ESPMode::ThreadSafe > DisplayString = FTextInspector::GetSharedDisplayString(DisplayText);
+			TSharedPtr< FString, ESPMode::ThreadSafe > Namespace;
+			TSharedPtr< FString, ESPMode::ThreadSafe > Key;
+			FTextLocalizationManager::Get().FindKeyNamespaceFromDisplayString(DisplayString, Namespace, Key);
 
-				if (Key.IsValid())
-				{
-					*Key.Get() = FGuid::NewGuid().ToString();
-					DirtyPropertyPackage(InPropHandle);
-				}
+			if (Key.IsValid())
+			{
+				*Key.Get() = InText.ToString();
+				DirtyPropertyPackage();
 			}
-
-			return FReply::Unhandled();
 		}
+	};
 
-		static void HandleKeyTextCommitted(const FText& InText, ETextCommit::Type InCommitType, TSharedRef<IPropertyHandle> InPropHandle)
+	const auto& HandleNamespaceTextCommitted = [=](const FText& InText, ETextCommit::Type InCommitType)
+	{
+		FText DisplayText;
+		if (InPropertyHandle->GetValueAsDisplayText(DisplayText) == FPropertyAccess::Success)
 		{
-			FText DisplayText;
-			if (InPropHandle->GetValueAsDisplayText(DisplayText) == FPropertyAccess::Success)
-			{
-				TSharedRef< FString, ESPMode::ThreadSafe > DisplayString = FTextInspector::GetSharedDisplayString(DisplayText);
-				TSharedPtr< FString, ESPMode::ThreadSafe > Namespace;
-				TSharedPtr< FString, ESPMode::ThreadSafe > Key;
-				FTextLocalizationManager::Get().FindKeyNamespaceFromDisplayString(DisplayString, Namespace, Key);
+			TSharedRef< FString, ESPMode::ThreadSafe > DisplayString = FTextInspector::GetSharedDisplayString(DisplayText);
+			TSharedPtr< FString, ESPMode::ThreadSafe > Namespace;
+			TSharedPtr< FString, ESPMode::ThreadSafe > Key;
+			FTextLocalizationManager::Get().FindKeyNamespaceFromDisplayString(DisplayString, Namespace, Key);
 
-				if (Key.IsValid())
-				{
-					*Key.Get() = InText.ToString();
-					DirtyPropertyPackage(InPropHandle);
-				}
+			if (Namespace.IsValid())
+			{
+				*Namespace.Get() = InText.ToString();
+				DirtyPropertyPackage();
 			}
 		}
+	};
 
-		static void HandleNamespaceTextCommitted(const FText& InText, ETextCommit::Type InCommitType, TSharedRef<IPropertyHandle> InPropHandle)
+	const auto& GetLocalizationMenuContent = [=]() -> TSharedRef<SWidget>
+	{
+		FMenuBuilder MenuContentBuilder(true, NULL);
 		{
 			FText DisplayText;
-			if (InPropHandle->GetValueAsDisplayText(DisplayText) == FPropertyAccess::Success)
+			if (InPropertyHandle->GetValueAsDisplayText(DisplayText) == FPropertyAccess::Success)
 			{
 				TSharedRef< FString, ESPMode::ThreadSafe > DisplayString = FTextInspector::GetSharedDisplayString(DisplayText);
 				TSharedPtr< FString, ESPMode::ThreadSafe > Namespace;
 				TSharedPtr< FString, ESPMode::ThreadSafe > Key;
 				FTextLocalizationManager::Get().FindKeyNamespaceFromDisplayString(DisplayString, Namespace, Key);
 
-				if (Namespace.IsValid())
+				MenuContentBuilder.BeginSection(TEXT("Localization"), LOCTEXT("LocalizationSectionHeading", "Localization"));
 				{
-					*Namespace.Get() = InText.ToString();
-					DirtyPropertyPackage(InPropHandle);
-				}
-			}
-		}
-
-		static TSharedRef<SWidget> GetLocalizationMenuContent(TSharedRef<IPropertyHandle> InPropHandle)
-		{
-			FMenuBuilder MenuContentBuilder(true, NULL);
-			{
-				FText DisplayText;
-				if (InPropHandle->GetValueAsDisplayText(DisplayText) == FPropertyAccess::Success)
-				{
-					TSharedRef< FString, ESPMode::ThreadSafe > DisplayString = FTextInspector::GetSharedDisplayString(DisplayText);
-					TSharedPtr< FString, ESPMode::ThreadSafe > Namespace;
-					TSharedPtr< FString, ESPMode::ThreadSafe > Key;
-					FTextLocalizationManager::Get().FindKeyNamespaceFromDisplayString(DisplayString, Namespace, Key);
-
-					MenuContentBuilder.BeginSection(TEXT("Localization"), LOCTEXT("LocalizationSectionHeading", "Localization"));
+					if (!Namespace.IsValid() && !Key.IsValid())
 					{
-						if (!Namespace.IsValid() && !Key.IsValid())
-						{
-							MenuContentBuilder.AddWidget(SNew(SSpacer), LOCTEXT("NoLocWarning", "No localization available"));
-						}
-						else
-						{
-							TSharedRef<SGridPanel> GridPanel = SNew(SGridPanel);
+						MenuContentBuilder.AddWidget(SNew(SSpacer), LOCTEXT("NoLocWarning", "No localization available"));
+					}
+					else
+					{
+						TSharedRef<SGridPanel> GridPanel = SNew(SGridPanel);
 
-							if (Key.IsValid())
-							{
-								GridPanel->AddSlot(0, 0)
-									.VAlign(VAlign_Center)
-									.Padding(1.0f)
+						if (Key.IsValid())
+						{
+							GridPanel->AddSlot(0, 0)
+								.VAlign(VAlign_Center)
+								.Padding(1.0f)
+								[
+									SNew(STextBlock)
+									.Text(LOCTEXT("KeyLabel", "Key"))
+									.ToolTipText(LOCTEXT("KeyTooltip", "The localization key of the text property"))
+								];
+
+							auto KeyLambda = [=](){ return FText::FromString(*Key.Get()); };
+
+							GridPanel->AddSlot(1, 0)
+								.VAlign(VAlign_Center)
+								.Padding(1.0f)
+								[
+									SNew(SHorizontalBox)
+									+SHorizontalBox::Slot()
+									.FillWidth(1.0f)
 									[
-										SNew(STextBlock)
-										.Text(LOCTEXT("KeyLabel", "Key"))
-										.ToolTipText(LOCTEXT("KeyTooltip", "The localization key of the text property"))
-									];
-
-								auto KeyLambda = [=](){ return FText::FromString(*Key.Get()); };
-
-								GridPanel->AddSlot(1, 0)
-									.VAlign(VAlign_Center)
-									.Padding(1.0f)
+										SNew(SEditableTextBox)
+										.Text_Lambda(KeyLambda)
+										.OnTextCommitted(FOnTextCommitted::CreateLambda(HandleKeyTextCommitted))
+									]
+									+ SHorizontalBox::Slot()
+									.AutoWidth()
+									.HAlign(HAlign_Center)
+									.Padding(2.0f, 0.0f, 0.0f, 0.0f)
 									[
-										SNew(SHorizontalBox)
-										+SHorizontalBox::Slot()
-										.FillWidth(1.0f)
-										[
-											SNew(SEditableTextBox)
-											.Text_Lambda(KeyLambda)
-											.OnTextCommitted(FOnTextCommitted::CreateStatic(&Local::HandleKeyTextCommitted, InPropHandle))
-										]
-										+ SHorizontalBox::Slot()
-										.AutoWidth()
-										.HAlign(HAlign_Center)
-										.Padding(2.0f, 0.0f, 0.0f, 0.0f)
+										SNew(SButton)
+										.ToolTipText(LOCTEXT("RefreshKeyTooltip", "Generate a new random key"))
+										.ButtonStyle(FEditorStyle::Get(), "HoverHintOnly")
+										.OnClicked(FOnClicked::CreateLambda(HandleGenerateKeyClicked))
+										.Content()
 										[
-											SNew(SButton)
-											.ToolTipText(LOCTEXT("RefreshKeyTooltip", "Generate a new random key"))
-											.ButtonStyle(FEditorStyle::Get(), "HoverHintOnly")
-											.OnClicked(FOnClicked::CreateStatic(&Local::HandleGenerateKeyClicked, InPropHandle))
-											.Content()
-											[
-												SNew(SImage)
-												.Image(FEditorStyle::GetBrush("PropertyWindow.Button_Refresh"))
-											]
+											SNew(SImage)
+											.Image(FEditorStyle::GetBrush("PropertyWindow.Button_Refresh"))
 										]
-									];
-							}
-
-							if (Namespace.IsValid())
-							{
-								auto NamespaceLambda = [=](){ return FText::FromString(*Namespace.Get()); };
-
-								GridPanel->AddSlot(0, 1)
-									.VAlign(VAlign_Center)
-									.Padding(1.0f, 1.0f, 5.0f, 1.0f)
-									[
-										SNew(STextBlock)
-										.Text(LOCTEXT("NamespaceLabel", "Namespace"))
-										.ToolTipText(LOCTEXT("NamespaceTooltip", "The localization namespace of the text property"))
-									];
-
-								GridPanel->AddSlot(1, 1)
-									.VAlign(VAlign_Center)
-									.Padding(1.0f)
-									[
-										SNew(SEditableTextBox)
-										.Text_Lambda(NamespaceLambda)
-										.OnTextCommitted(FOnTextCommitted::CreateStatic(&Local::HandleNamespaceTextCommitted, InPropHandle))
-									];
-							}
+									]
+								];
+						}
 
-							MenuContentBuilder.AddWidget(GridPanel, FText::GetEmpty());
+						if (Namespace.IsValid())
+						{
+							auto NamespaceLambda = [=](){ return FText::FromString(*Namespace.Get()); };
+
+							GridPanel->AddSlot(0, 1)
+								.VAlign(VAlign_Center)
+								.Padding(1.0f, 1.0f, 5.0f, 1.0f)
+								[
+									SNew(STextBlock)
+									.Text(LOCTEXT("NamespaceLabel", "Namespace"))
+									.ToolTipText(LOCTEXT("NamespaceTooltip", "The localization namespace of the text property"))
+								];
+
+							GridPanel->AddSlot(1, 1)
+								.VAlign(VAlign_Center)
+								.Padding(1.0f)
+								[
+									SNew(SEditableTextBox)
+									.Text_Lambda(NamespaceLambda)
+									.OnTextCommitted(FOnTextCommitted::CreateLambda(HandleNamespaceTextCommitted))
+								];
 						}
-					}
 
-					MenuContentBuilder.EndSection();
+						MenuContentBuilder.AddWidget(GridPanel, FText::GetEmpty());
+					}
 				}
-			}
 
-			return MenuContentBuilder.MakeWidget();
+				MenuContentBuilder.EndSection();
+			}
 		}
+
+		return MenuContentBuilder.MakeWidget();
 	};
 
 	return SNew(SComboButton)
@@ -1229,7 +1226,7 @@ TSharedRef<SWidget> PropertyCustomizationHelpers::MakeTextLocalizationButton(TSh
 		.ContentPadding(2)
 		.ForegroundColor(FSlateColor::UseForeground())
 		.HasDownArrow(true)
-		.OnGetMenuContent(FOnGetContent::CreateStatic(&Local::GetLocalizationMenuContent, InPropertyHandle));
+		.OnGetMenuContent(FOnGetContent::CreateLambda(GetLocalizationMenuContent));
 }
 
 #undef LOCTEXT_NAMESPACE
diff --git a/Engine/Source/Editor/PropertyEditor/Public/PropertyCustomizationHelpers.h b/Engine/Source/Editor/PropertyEditor/Public/PropertyCustomizationHelpers.h
--- a/Engine/Source/Editor/PropertyEditor/Public/PropertyCustomizationHelpers.h
+++ b/Engine/Source/Editor/PropertyEditor/Public/PropertyCustomizationHelpers.h
@@ -42,7 +42,7 @@ namespace PropertyCustomizationHelpers
 	PROPERTYEDITOR_API TArray<UFactory*> GetNewAssetFactoriesForClasses(const TArray<const UClass*>& Classes);
 
 	/** Build a combo button allowing access to text property localization utilities */
-	PROPERTYEDITOR_API TSharedRef<SWidget> MakeTextLocalizationButton(TSharedRef<IPropertyHandle> InPropertyHandle);
+	PROPERTYEDITOR_API TSharedRef<SWidget> MakeTextLocalizationButton(const TSharedRef<IPropertyHandle>& InPropertyHandle);
 }
 
 /** Delegate used to get a generic object */
