Removed faulty package dirtying logic from FText property namespace/key menu logic. Opted for using the property handle to notify pre-change and post-change, which properly causes package dirtying and marking a blueprint as needing to be recompiled.

diff --git a/Engine/Source/Editor/PropertyEditor/Private/PropertyCustomizationHelpers.cpp b/Engine/Source/Editor/PropertyEditor/Private/PropertyCustomizationHelpers.cpp
--- a/Engine/Source/Editor/PropertyEditor/Private/PropertyCustomizationHelpers.cpp
+++ b/Engine/Source/Editor/PropertyEditor/Private/PropertyCustomizationHelpers.cpp
@@ -1054,17 +1054,6 @@ void FMaterialList::AddMaterialItem( FDetailWidgetRow& Row, int32 CurrentSlot, c
 
 TSharedRef<SWidget> PropertyCustomizationHelpers::MakeTextLocalizationButton(const TSharedRef<IPropertyHandle>& InPropertyHandle)
 {
-	const auto& DirtyPropertyPackage = [=]()
-	{
-		// dirty packages as they will need to be re-serialized with the correct key
-		TArray<UObject*> Objects;
-		InPropertyHandle->GetOuterObjects(Objects);
-		if (Objects.Num() > 0 && Objects[0]->GetOutermost() != nullptr)
-		{
-			Objects[0]->GetOutermost()->MarkPackageDirty();
-		}
-	};
-
 	const auto& HandleGenerateKeyClicked = [=]() -> FReply 
 	{
 		FText DisplayText;
@@ -1077,8 +1066,9 @@ TSharedRef<SWidget> PropertyCustomizationHelpers::MakeTextLocalizationButton(con
 
 			if (Key.IsValid())
 			{
+				InPropertyHandle->NotifyPreChange();
 				*Key.Get() = FGuid::NewGuid().ToString();
-				DirtyPropertyPackage();
+				InPropertyHandle->NotifyPostChange();
 			}
 		}
 
@@ -1097,8 +1087,9 @@ TSharedRef<SWidget> PropertyCustomizationHelpers::MakeTextLocalizationButton(con
 
 			if (Key.IsValid())
 			{
+				InPropertyHandle->NotifyPreChange();
 				*Key.Get() = InText.ToString();
-				DirtyPropertyPackage();
+				InPropertyHandle->NotifyPostChange();
 			}
 		}
 	};
@@ -1115,8 +1106,9 @@ TSharedRef<SWidget> PropertyCustomizationHelpers::MakeTextLocalizationButton(con
 
 			if (Namespace.IsValid())
 			{
+				InPropertyHandle->NotifyPreChange();
 				*Namespace.Get() = InText.ToString();
-				DirtyPropertyPackage();
+				InPropertyHandle->NotifyPostChange();
 			}
 		}
 	};
