Fixed not being able to add fields to enums

diff --git a/Engine/Source/Editor/PropertyEditor/Private/DetailCategoryBuilderImpl.cpp b/Engine/Source/Editor/PropertyEditor/Private/DetailCategoryBuilderImpl.cpp
--- a/Engine/Source/Editor/PropertyEditor/Private/DetailCategoryBuilderImpl.cpp
+++ b/Engine/Source/Editor/PropertyEditor/Private/DetailCategoryBuilderImpl.cpp
@@ -600,7 +600,8 @@ void FDetailCategoryImpl::GenerateNodesFromCustomizations( const FCustomizationL
 			TSharedRef<FDetailItemNode> NewNode = MakeShareable( new FDetailItemNode( Customization, AsShared(), IsParentEnabled ) );
 			NewNode->Initialize();
 
-			if( !NewNode->ShouldShowOnlyChildren() || NewNode->HasGeneratedChildren() )
+			// Add the node unless only its children should be visible or it didnt generate any children or if it is a custom builder which can generate children at any point
+			if( !NewNode->ShouldShowOnlyChildren() || NewNode->HasGeneratedChildren() || Customization.HasCustomBuilder() )
 			{
 				if(CustomizationIndex == InCustomizationList.Num()-1)
 				{
