Allowed FPropertyNode to notify its FNotifyHook of a pre/post edit change even if the node wasn't contained within an object

diff --git a/Engine/Source/Editor/PropertyEditor/Private/PropertyNode.cpp b/Engine/Source/Editor/PropertyEditor/Private/PropertyNode.cpp
--- a/Engine/Source/Editor/PropertyEditor/Private/PropertyNode.cpp
+++ b/Engine/Source/Editor/PropertyEditor/Private/PropertyNode.cpp
@@ -1851,25 +1851,23 @@ void FPropertyNode::NotifyPreChange( UProperty* PropertyAboutToChange, FNotifyHo
 {
 	TSharedRef<FEditPropertyChain> PropertyChain = BuildPropertyChain( PropertyAboutToChange );
 
-	FObjectPropertyNode* ObjectNode = FindObjectItemParent();
-
-	if( ObjectNode )
+	// Call through to the property window's notify hook.
+	if( InNotifyHook )
 	{
-		UProperty* CurProperty = PropertyAboutToChange;
-
-		// Call through to the property window's notify hook.
-		if( InNotifyHook )
+		if ( PropertyChain->Num() == 0 )
 		{
-			if ( PropertyChain->Num() == 0 )
-			{
-				InNotifyHook->NotifyPreChange( PropertyAboutToChange );
-			}
-			else
-			{
-				InNotifyHook->NotifyPreChange( &PropertyChain.Get() );
-			}
+			InNotifyHook->NotifyPreChange( PropertyAboutToChange );
 		}
+		else
+		{
+			InNotifyHook->NotifyPreChange( &PropertyChain.Get() );
+		}
+	}
 
+	FObjectPropertyNode* ObjectNode = FindObjectItemParent();
+	if( ObjectNode )
+	{
+		UProperty* CurProperty = PropertyAboutToChange;
 
 		// Call PreEditChange on the object chain.
 		while ( true )
@@ -1927,9 +1925,11 @@ void FPropertyNode::NotifyPreChange( UProperty* PropertyAboutToChange, FNotifyHo
 void FPropertyNode::NotifyPostChange( FPropertyChangedEvent& InPropertyChangedEvent, class FNotifyHook* InNotifyHook )
 {
 	TSharedRef<FEditPropertyChain> PropertyChain = BuildPropertyChain( InPropertyChangedEvent.Property );
+	
+	// remember the property that was the chain's original active property; this will correspond to the outermost property of struct/array that was modified
+	UProperty* const OriginalActiveProperty = PropertyChain->GetActiveMemberNode()->GetValue();
 
 	FObjectPropertyNode* ObjectNode = FindObjectItemParent();
-	
 	if( ObjectNode )
 	{
 		UProperty* CurProperty = InPropertyChangedEvent.Property;
@@ -1937,9 +1937,6 @@ void FPropertyNode::NotifyPostChange( FPropertyChangedEvent& InPropertyChangedEv
 		// Fire ULevel::LevelDirtiedEvent when falling out of scope.
 		FScopedLevelDirtied	LevelDirtyCallback;
 
-		// remember the property that was the chain's original active property; this will correspond to the outermost property of struct/array that was modified
-		UProperty* OriginalActiveProperty = PropertyChain->GetActiveMemberNode()->GetValue();
-
 		// Call PostEditChange on the object chain.
 		while ( true )
 		{
@@ -2009,44 +2006,43 @@ void FPropertyNode::NotifyPostChange( FPropertyChangedEvent& InPropertyChangedEv
 				}
 			}
 		}
+	}
 
+	// Broadcast the change to any listeners
+	BroadcastPropertyValueChanged();
 
-		// Broadcast the change to any listeners
-		BroadcastPropertyValueChanged();
-
-		// Call through to the property window's notify hook.
-		if( InNotifyHook )
+	// Call through to the property window's notify hook.
+	if( InNotifyHook )
+	{
+		if ( PropertyChain->Num() == 0 )
 		{
-			if ( PropertyChain->Num() == 0 )
-			{
-				InNotifyHook->NotifyPostChange( InPropertyChangedEvent, InPropertyChangedEvent.Property );
-			}
-			else
-			{
-				PropertyChain->SetActiveMemberPropertyNode( OriginalActiveProperty );
-				PropertyChain->SetActivePropertyNode( InPropertyChangedEvent.Property);
-				InNotifyHook->NotifyPostChange( InPropertyChangedEvent, &PropertyChain.Get() );
-			}
+			InNotifyHook->NotifyPostChange( InPropertyChangedEvent, InPropertyChangedEvent.Property );
+		}
+		else
+		{
+			PropertyChain->SetActiveMemberPropertyNode( OriginalActiveProperty );
+			PropertyChain->SetActivePropertyNode( InPropertyChangedEvent.Property);
+			InNotifyHook->NotifyPostChange( InPropertyChangedEvent, &PropertyChain.Get() );
 		}
+	}
 
 
-		if( OriginalActiveProperty )
-		{
-			//if i have metadata forcing other property windows to rebuild
-			FString MetaData = OriginalActiveProperty->GetMetaData(TEXT("ForceRebuildProperty"));
+	if( ObjectNode && OriginalActiveProperty )
+	{
+		//if i have metadata forcing other property windows to rebuild
+		FString MetaData = OriginalActiveProperty->GetMetaData(TEXT("ForceRebuildProperty"));
 
-			if( MetaData.Len() > 0 )
-			{
-				// We need to find the property node beginning at the root/parent, not at our own node.
-				ObjectNode = FindObjectItemParent();
-				check(ObjectNode != NULL);
+		if( MetaData.Len() > 0 )
+		{
+			// We need to find the property node beginning at the root/parent, not at our own node.
+			ObjectNode = FindObjectItemParent();
+			check(ObjectNode != NULL);
 
-				TSharedPtr<FPropertyNode> ForceRebuildNode = ObjectNode->FindChildPropertyNode( FName(*MetaData), true );
+			TSharedPtr<FPropertyNode> ForceRebuildNode = ObjectNode->FindChildPropertyNode( FName(*MetaData), true );
 
-				if( ForceRebuildNode.IsValid() )
-				{
-					ForceRebuildNode->RequestRebuildChildren();
-				}
+			if( ForceRebuildNode.IsValid() )
+			{
+				ForceRebuildNode->RequestRebuildChildren();
 			}
 		}
 	}
