Removed text localization menu button on text properties. Not currently functional enough to ship.

diff --git a/Engine/Source/Editor/IntroTutorials/Private/TutorialStructCustomization.cpp b/Engine/Source/Editor/IntroTutorials/Private/TutorialStructCustomization.cpp
--- a/Engine/Source/Editor/IntroTutorials/Private/TutorialStructCustomization.cpp
+++ b/Engine/Source/Editor/IntroTutorials/Private/TutorialStructCustomization.cpp
@@ -166,14 +166,6 @@ void FTutorialContentCustomization::CustomizeHeader( TSharedRef<class IPropertyH
 				.OnTextCommitted(FOnTextCommitted::CreateStatic(&Local::OnTextCommitted, TextProperty))
 				.OnTextChanged(FOnTextChanged::CreateStatic(&Local::OnTextChanged, TextProperty))
 			]
-			+SHorizontalBox::Slot()
-			.AutoWidth()
-			.VAlign(VAlign_Center)
-			.HAlign(HAlign_Center)
-			.Padding(2.0f, 0.0f, 0.0f, 0.0f)
-			[
-				PropertyCustomizationHelpers::MakeTextLocalizationButton(TextProperty.ToSharedRef())
-			]
 		]
 	];
 }
diff --git a/Engine/Source/Editor/PropertyEditor/Private/UserInterface/PropertyEditor/SPropertyEditorText.cpp b/Engine/Source/Editor/PropertyEditor/Private/UserInterface/PropertyEditor/SPropertyEditorText.cpp
--- a/Engine/Source/Editor/PropertyEditor/Private/UserInterface/PropertyEditor/SPropertyEditorText.cpp
+++ b/Engine/Source/Editor/PropertyEditor/Private/UserInterface/PropertyEditor/SPropertyEditorText.cpp
@@ -58,22 +58,6 @@ void SPropertyEditorText::Construct( const FArguments& InArgs, const TSharedRef<
 		];
 	}
 
-	if (InPropertyEditor->PropertyIsA(UTextProperty::StaticClass()))
-	{
-		const TSharedRef< FPropertyNode > PropertyNode = InPropertyEditor->GetPropertyNode();
-		const TSharedRef< IPropertyHandle > PropertyHandle = PropertyEditor->GetPropertyHandle();
-
-		// text properties require some extra controls to allow for localization functionality
-		HorizontalBox->AddSlot()
-			.AutoWidth()
-			.VAlign(VAlign_Center)
-			.HAlign(HAlign_Center)
-			.Padding(2.0f, 0.0f, 0.0f, 0.0f)
-			[
-				PropertyCustomizationHelpers::MakeTextLocalizationButton(PropertyHandle)
-			];
-	}
-
 	if( InPropertyEditor->PropertyIsA( UObjectPropertyBase::StaticClass() ) )
 	{
 		// Object properties should display their entire text in a tooltip
