Property Editor: Changed SPropertyEditorAsset so that the drop-down button with the asset name stays visible longer as horizontal width reduces (still need to do something more agressive like shrink or omit the thumbnail at very small widths)

diff --git a/Engine/Source/Editor/PropertyEditor/Private/UserInterface/PropertyEditor/SPropertyEditorAsset.cpp b/Engine/Source/Editor/PropertyEditor/Private/UserInterface/PropertyEditor/SPropertyEditorAsset.cpp
--- a/Engine/Source/Editor/PropertyEditor/Private/UserInterface/PropertyEditor/SPropertyEditorAsset.cpp
+++ b/Engine/Source/Editor/PropertyEditor/Private/UserInterface/PropertyEditor/SPropertyEditorAsset.cpp
@@ -208,7 +208,31 @@ void SPropertyEditorAsset::Construct( const FArguments& InArgs, const TSharedPtr
 		}
 	}
 
-	TSharedPtr<SHorizontalBox> ButtonBox;
+	AssetComboButton = SNew(SComboButton)
+		.ToolTipText( this, &SPropertyEditorAsset::OnGetToolTip )
+		.ButtonStyle( FEditorStyle::Get(), "PropertyEditor.AssetComboStyle" )
+		.ForegroundColor(FEditorStyle::GetColor("PropertyEditor.AssetName.ColorAndOpacity"))
+		.OnGetMenuContent( this, &SPropertyEditorAsset::OnGetMenuContent )
+		.OnMenuOpenChanged( this, &SPropertyEditorAsset::OnMenuOpenChanged )
+		.ContentPadding(2.0f)
+		.ButtonContent()
+		[
+			// Show the name of the asset or actor
+			SNew(STextBlock)
+			.TextStyle( FEditorStyle::Get(), "PropertyEditor.AssetClass" )
+			.Font( FEditorStyle::GetFontStyle( PropertyEditorConstants::PropertyFontStyle ) )
+			.Text(this,&SPropertyEditorAsset::OnGetAssetName)
+		];
+
+
+	TSharedRef<SHorizontalBox> ButtonBox = SNew( SHorizontalBox )
+		.IsEnabled( this, &SPropertyEditorAsset::CanEdit )
+		+ SHorizontalBox::Slot()
+		.AutoWidth()
+		[
+			AssetComboButton.ToSharedRef()
+		];
+
 	
 	HorizontalBox->AddSlot()
 	.FillWidth(1.0f)
@@ -219,26 +243,7 @@ void SPropertyEditorAsset::Construct( const FArguments& InArgs, const TSharedPtr
 		+SVerticalBox::Slot()
 		.VAlign( VAlign_Center )
 		[
-			SAssignNew( ButtonBox, SHorizontalBox )
-			.IsEnabled( this, &SPropertyEditorAsset::CanEdit )
-			+ SHorizontalBox::Slot()
-			[
-				SAssignNew( AssetComboButton, SComboButton )
-				.ToolTipText( this, &SPropertyEditorAsset::OnGetToolTip )
-				.ButtonStyle( FEditorStyle::Get(), "PropertyEditor.AssetComboStyle" )
-				.ForegroundColor(FEditorStyle::GetColor("PropertyEditor.AssetName.ColorAndOpacity"))
-				.OnGetMenuContent( this, &SPropertyEditorAsset::OnGetMenuContent )
-				.OnMenuOpenChanged( this, &SPropertyEditorAsset::OnMenuOpenChanged )
-				.ContentPadding(2.0f)
-				.ButtonContent()
-				[
-					// Show the name of the asset or actor
-					SNew(STextBlock)
-					.TextStyle( FEditorStyle::Get(), "PropertyEditor.AssetClass" )
-					.Font( FEditorStyle::GetFontStyle( PropertyEditorConstants::PropertyFontStyle ) )
-					.Text(this,&SPropertyEditorAsset::OnGetAssetName)
-				]
-			]
+			ButtonBox
 		]
 		+SVerticalBox::Slot()
 		.AutoHeight()
