Swap order of PostEditChange on CDO/archetype and propagating property change to instances.
CapsuleComponent Shape sometimes reverts to previous value after entering a new value.

diff --git a/Engine/Source/Editor/PropertyEditor/Private/PropertyHandleImpl.cpp b/Engine/Source/Editor/PropertyEditor/Private/PropertyHandleImpl.cpp
--- a/Engine/Source/Editor/PropertyEditor/Private/PropertyHandleImpl.cpp
+++ b/Engine/Source/Editor/PropertyEditor/Private/PropertyHandleImpl.cpp
@@ -346,6 +346,13 @@ FPropertyAccess::Result FPropertyValueImpl::ImportText( const TArray<FObjectBase
 		const bool bTransactable = (Flags & EPropertyValueSetFlags::NotTransactable) == 0;
 		const bool bFinished = ( Flags & EPropertyValueSetFlags::InteractiveChange) == 0;
 		
+		// Cache previous values to compare against them after all potential modifications to property.
+		TArray<FString> PreviousValues;
+		PreviousValues.Reserve(InObjects.Num());
+
+		TArray<FString> PreviousArrayValues;
+		PreviousArrayValues.Reserve(InObjects.Num());
+
 		for ( int32 ObjectIndex = 0 ; ObjectIndex < InObjects.Num() ; ++ObjectIndex )
 		{	
 			const FObjectBaseAddress& Cur = InObjects[ ObjectIndex ];
@@ -359,6 +366,7 @@ FPropertyAccess::Result FPropertyValueImpl::ImportText( const TArray<FObjectBase
 			// Cache the value of the property before modifying it.
 			FString PreviousValue;
 			NodeProperty->ExportText_Direct(PreviousValue, Cur.BaseAddress, Cur.BaseAddress, NULL, 0 );
+			PreviousValues.Add(PreviousValue);
 
 			// If this property is the inner-property of an array, cache the current value as well
 			FString PreviousArrayValue;
@@ -376,6 +384,7 @@ FPropertyAccess::Result FPropertyValueImpl::ImportText( const TArray<FObjectBase
 					}
 				}
 			}
+			PreviousArrayValues.Add(PreviousArrayValue);
 
 			// Check if we need to call PreEditChange on all objects.
 			// Remove quotes from the original value because FName properties  
@@ -404,26 +413,6 @@ FPropertyAccess::Result FPropertyValueImpl::ImportText( const TArray<FObjectBase
 			const TCHAR* NewValue = *InValues[ObjectIndex];
 			NodeProperty->ImportText( NewValue, Cur.BaseAddress, 0, Cur.Object );
 
-			if (Cur.Object)
-			{
-				// Cache the value of the property after having modified it.
-				FString ValueAfterImport;
-				NodeProperty->ExportText_Direct(ValueAfterImport, Cur.BaseAddress, Cur.BaseAddress, NULL, 0);
-
-				if ((Cur.Object->HasAnyFlags(RF_ClassDefaultObject | RF_ArchetypeObject) ||
-					(Cur.Object->HasAnyFlags(RF_DefaultSubObject) && Cur.Object->GetOuter()->HasAnyFlags(RF_ClassDefaultObject | RF_ArchetypeObject))) &&
-					!bIsGameWorld)
-				{
-					InPropertyNode->PropagatePropertyChange(Cur.Object, NewValue, PreviousArrayValue.IsEmpty() ? PreviousValue : PreviousArrayValue);
-				}
-
-				// If the values before and after setting the property differ, mark the object dirty.
-				if (FCString::Strcmp(*PreviousValue, *ValueAfterImport) != 0)
-				{
-					Cur.Object->MarkPackageDirty();
-				}
-			}
-
 			//add on array index so we can tell which entry just changed
 			ArrayIndicesPerObject.Add(TMap<FString,int32>());
 			FPropertyValueImpl::GenerateArrayIndexMapToObjectNode( ArrayIndicesPerObject[ObjectIndex], InPropertyNode );
@@ -450,6 +439,41 @@ FPropertyAccess::Result FPropertyValueImpl::ImportText( const TArray<FObjectBase
 			}
 		}
 
+		for (int32 ObjectIndex = 0; ObjectIndex < InObjects.Num(); ++ObjectIndex)
+		{
+			const FObjectBaseAddress& Cur = InObjects[ObjectIndex];
+
+			if (Cur.Object == nullptr)
+			{
+				continue;
+			}
+
+			// Cache the value of the property after having modified it (either by entering value or via PostEditChangeProperty).
+			FString ValueAfterImport;
+			NodeProperty->ExportText_Direct(ValueAfterImport, Cur.BaseAddress, Cur.BaseAddress, NULL, 0);
+
+			const FString& PreviousValue = PreviousValues[ObjectIndex];
+
+			// If the values before and after setting the property are the same, we don't need to do anything.
+			if (FCString::Strcmp(*PreviousValue, *ValueAfterImport) == 0)
+			{
+				continue;
+			}
+
+			const FString& PreviousArrayValue = PreviousArrayValues[ObjectIndex];
+
+			// Propagate property change to all objects using this archetype/CDO (or parent's archetype/CDE if this is default subobject)
+			if ((Cur.Object->HasAnyFlags(RF_ClassDefaultObject | RF_ArchetypeObject) ||
+				(Cur.Object->HasAnyFlags(RF_DefaultSubObject) && Cur.Object->GetOuter()->HasAnyFlags(RF_ClassDefaultObject | RF_ArchetypeObject))) &&
+				!bIsGameWorld)
+			{
+				InPropertyNode->PropagatePropertyChange(Cur.Object, *ValueAfterImport, PreviousArrayValue.IsEmpty() ? PreviousValue : PreviousArrayValue);
+			}
+
+			// Property changed, mark package as dirty.
+			Cur.Object->MarkPackageDirty();
+		}
+
 		if (OldGWorld)
 		{
 			// restore the original (editor) GWorld
