Improved the property filtering options when creating a structure detail view

Previously you could only control whether you saw object properties or not, you can now control whether you see object, asset, class, and interface properties.

diff --git a/Engine/Source/Editor/DataTableEditor/Private/SRowEditor.cpp b/Engine/Source/Editor/DataTableEditor/Private/SRowEditor.cpp
--- a/Engine/Source/Editor/DataTableEditor/Private/SRowEditor.cpp
+++ b/Engine/Source/Editor/DataTableEditor/Private/SRowEditor.cpp
@@ -338,7 +338,13 @@ void SRowEditor::Construct(const FArguments& InArgs, UDataTable* Changed)
 		ViewArgs.bShowActorLabel = false;
 		ViewArgs.NotifyHook = this;
 
-		StructureDetailsView = PropertyModule.CreateStructureDetailView(ViewArgs, CurrentRow, true/*bShowObjects*/, LOCTEXT("RowValue", "Row Value"));
+		FStructureDetailsViewArgs StructureViewArgs;
+		StructureViewArgs.bShowObjects = false;
+		StructureViewArgs.bShowAssets = true;
+		StructureViewArgs.bShowClasses = true;
+		StructureViewArgs.bShowInterfaces = false;
+
+		StructureDetailsView = PropertyModule.CreateStructureDetailView(ViewArgs, StructureViewArgs, CurrentRow, LOCTEXT("RowValue", "Row Value"));
 	}
 
 	RefreshNameList();
diff --git a/Engine/Source/Editor/Kismet/Private/UserDefinedStructureEditor.cpp b/Engine/Source/Editor/Kismet/Private/UserDefinedStructureEditor.cpp
--- a/Engine/Source/Editor/Kismet/Private/UserDefinedStructureEditor.cpp
+++ b/Engine/Source/Editor/Kismet/Private/UserDefinedStructureEditor.cpp
@@ -34,7 +34,13 @@ public:
 		ViewArgs.bHideSelectionTip = false;
 		ViewArgs.bShowActorLabel = false;
 
-		StructureDetailsView = PropertyModule.CreateStructureDetailView(ViewArgs, StructData, false, LOCTEXT("DefaultValues", "Default Values"));
+		FStructureDetailsViewArgs StructureViewArgs;
+		StructureViewArgs.bShowObjects = false;
+		StructureViewArgs.bShowAssets = false;
+		StructureViewArgs.bShowClasses = true;
+		StructureViewArgs.bShowInterfaces = false;
+
+		StructureDetailsView = PropertyModule.CreateStructureDetailView(ViewArgs, StructureViewArgs, StructData, LOCTEXT("DefaultValues", "Default Values"));
 		StructureDetailsView->GetOnFinishedChangingPropertiesDelegate().AddSP(this, &FStructureDefaultValueView::OnFinishedChangingProperties);
 	}
 
diff --git a/Engine/Source/Editor/PropertyEditor/Private/PropertyEditorModule.cpp b/Engine/Source/Editor/PropertyEditor/Private/PropertyEditorModule.cpp
--- a/Engine/Source/Editor/PropertyEditor/Private/PropertyEditorModule.cpp
+++ b/Engine/Source/Editor/PropertyEditor/Private/PropertyEditorModule.cpp
@@ -3,6 +3,7 @@
 
 #include "PropertyEditorPrivatePCH.h"
 #include "AssetSelection.h"
+#include "AssetToolsModule.h"
 #include "SPropertyTreeViewImpl.h"
 #include "SlateBasics.h"
 #include "PropertyNode.h"
@@ -651,29 +652,60 @@ FPropertyTypeLayoutCallback FPropertyEditorModule::GetPropertyTypeCustomization(
 	return FPropertyTypeLayoutCallback();
 }
 
-TSharedRef<class IStructureDetailsView> FPropertyEditorModule::CreateStructureDetailView(const FDetailsViewArgs& DetailsViewArgs, TSharedPtr<FStructOnScope> StructData, bool bShowObjects, const FText& CustomName)
+TSharedRef<class IStructureDetailsView> FPropertyEditorModule::CreateStructureDetailView(const FDetailsViewArgs& DetailsViewArgs, const FStructureDetailsViewArgs& StructureDetailsViewArgs, TSharedPtr<FStructOnScope> StructData, const FText& CustomName)
 {
 	TSharedRef<SStructureDetailsView> DetailView =
 		SNew(SStructureDetailsView)
 		.DetailsViewArgs(DetailsViewArgs)
 		.CustomName(CustomName);
 
-	if (!bShowObjects)
+	struct FStructureDetailsViewFilter
 	{
-		struct FDontShowObjects
+		static bool HasFilter( const FStructureDetailsViewArgs InStructureDetailsViewArgs )
 		{
-			static bool CanShow( const FPropertyAndParent& PropertyAndParent )
+			const bool bShowEverything = InStructureDetailsViewArgs.bShowObjects 
+				&& InStructureDetailsViewArgs.bShowAssets 
+				&& InStructureDetailsViewArgs.bShowClasses 
+				&& InStructureDetailsViewArgs.bShowInterfaces;
+			return !bShowEverything;
+		}
+
+		static bool PassesFilter( const FPropertyAndParent& PropertyAndParent, const FStructureDetailsViewArgs InStructureDetailsViewArgs )
+		{
+			const auto ArrayProperty = Cast<UArrayProperty>(&PropertyAndParent.Property);
+			const auto PropertyToTest = ArrayProperty ? ArrayProperty->Inner : &PropertyAndParent.Property;
+
+			if( InStructureDetailsViewArgs.bShowClasses && (PropertyToTest->IsA<UClassProperty>() || PropertyToTest->IsA<UAssetClassProperty>()) )
+			{
+				return true;
+			}
+
+			if( InStructureDetailsViewArgs.bShowInterfaces && PropertyToTest->IsA<UInterfaceProperty>() )
+			{
+				return true;
+			}
+
+			const auto ObjectProperty = Cast<UObjectPropertyBase>(PropertyToTest);
+			if( ObjectProperty )
 			{
-				const auto ArrayProperty = Cast<UArrayProperty>(&PropertyAndParent.Property);
-				const auto TrueProperty = ArrayProperty ? ArrayProperty->Inner : &PropertyAndParent.Property;
-				const bool bCanShow = TrueProperty 
-					&& (!TrueProperty->IsA<UObjectPropertyBase>() || TrueProperty->IsA<UClassProperty>())
-					&& !TrueProperty->IsA<UInterfaceProperty>();
-				return bCanShow;
+				if( InStructureDetailsViewArgs.bShowAssets && ObjectProperty->PropertyClass )
+				{
+					// We can use the asset tools module to see whether this type has asset actions (which likely means it's an asset class type)
+					FAssetToolsModule& AssetToolsModule = FAssetToolsModule::GetModule();
+					return AssetToolsModule.Get().GetAssetTypeActionsForClass(ObjectProperty->PropertyClass).IsValid();
+				}
+
+				return InStructureDetailsViewArgs.bShowObjects;
 			}
-		};
 
-		DetailView->SetIsPropertyVisibleDelegate(FIsPropertyVisible::CreateStatic(&FDontShowObjects::CanShow));
+			return true;
+		}
+	};
+
+	// Only add the filter if we need to exclude things
+	if (FStructureDetailsViewFilter::HasFilter(StructureDetailsViewArgs))
+	{
+		DetailView->SetIsPropertyVisibleDelegate(FIsPropertyVisible::CreateStatic(&FStructureDetailsViewFilter::PassesFilter, StructureDetailsViewArgs));
 	}
 	DetailView->SetStructureData(StructData);
 
diff --git a/Engine/Source/Editor/PropertyEditor/Public/PropertyEditorModule.h b/Engine/Source/Editor/PropertyEditor/Public/PropertyEditorModule.h
--- a/Engine/Source/Editor/PropertyEditor/Public/PropertyEditorModule.h
+++ b/Engine/Source/Editor/PropertyEditor/Public/PropertyEditorModule.h
@@ -114,6 +114,31 @@ typedef TMap< FName, FDetailLayoutCallback > FCustomDetailLayoutNameMap;
 typedef TMap< FName, FPropertyTypeLayoutCallbackList > FCustomPropertyTypeLayoutMap;
 
 
+/** Struct used to control the visibility of properties in a Structure Detail View */
+struct FStructureDetailsViewArgs
+{
+	FStructureDetailsViewArgs()
+		: bShowObjects(false)
+		, bShowAssets(true)
+		, bShowClasses(true)
+		, bShowInterfaces(false)
+	{
+	}
+
+	/** True if we should show general object properties in the details view */
+	bool bShowObjects : 1;
+
+	/** True if we should show asset properties in the details view */
+	bool bShowAssets : 1;
+
+	/** True if we should show class properties in the details view */
+	bool bShowClasses : 1;
+
+	/** True if we should show interface properties in the details view */
+	bool bShowInterfaces : 1;
+};
+
+
 class FPropertyEditorModule : public IModuleInterface
 {
 	friend class SPropertyTreeView;
@@ -232,7 +257,7 @@ public:
 	 */
 	virtual TSharedPtr<class ISinglePropertyView> CreateSingleProperty( UObject* InObject, FName InPropertyName, const struct FSinglePropertyParams& InitParams );
 
-	virtual TSharedRef<class IStructureDetailsView> CreateStructureDetailView(const struct FDetailsViewArgs& DetailsViewArgs, TSharedPtr<class FStructOnScope> StructData, bool bShowObjects, const FText& CustomName = FText::GetEmpty());
+	virtual TSharedRef<class IStructureDetailsView> CreateStructureDetailView(const struct FDetailsViewArgs& DetailsViewArgs, const FStructureDetailsViewArgs& StructureDetailsViewArgs, TSharedPtr<class FStructOnScope> StructData, const FText& CustomName = FText::GetEmpty());
 
 	/**
 	 * Creates a property change listener that notifies users via a  delegate when a property on an object changes
