- Fix inconsistency with property buttons displaying when edit const is changed on a property.
- Fixed memory leaks with property buttons keeping property nodes alive
- Fixed the browse button not being used when an property is edit const

diff --git a/Engine/Source/Editor/PropertyEditor/Private/PropertyEditorHelpers.cpp b/Engine/Source/Editor/PropertyEditor/Private/PropertyEditorHelpers.cpp
--- a/Engine/Source/Editor/PropertyEditor/Private/PropertyEditorHelpers.cpp
+++ b/Engine/Source/Editor/PropertyEditor/Private/PropertyEditorHelpers.cpp
@@ -483,18 +483,13 @@ namespace PropertyEditorHelpers
 		UProperty* NodeProperty = PropertyNode->GetProperty();
 
 		// If no property is bound, don't create any buttons.
-		// If property is const, don't create any buttons.
-		if ( !NodeProperty || PropertyNode->IsEditConst() )
+		if ( !NodeProperty )
 		{
 			return;
 		}
 
 		// If the property is an item of a const array, don't create any buttons.
 		const UArrayProperty* OuterArrayProp = Cast<UArrayProperty>( NodeProperty->GetOuter() );
-		if (OuterArrayProp != NULL && (OuterArrayProp->PropertyFlags & CPF_EditConst))
-		{
-			return;
-		}
 
 		if (NodeProperty->HasAnyPropertyFlags(CPF_GlobalConfig|CPF_Config))
 		{
@@ -717,6 +712,12 @@ namespace PropertyEditorHelpers
 		return SelectionPathName;
 	}
 
+	
+	static bool IsPropertyButtonEnabled( TWeakPtr<FPropertyNode> PropertyNode )
+	{
+		return PropertyNode.IsValid() ? !PropertyNode.Pin()->IsEditConst() : false;
+	}
+
 	/**
 	 * A helper method that checks to see if the editor's current selection is 
 	 * compatible with the specified property.
@@ -724,9 +725,10 @@ namespace PropertyEditorHelpers
 	 * @param  PropertyNode		The property you desire to set from the "use selected" button.
 	 * @return False if the currently selected object is restricted for the specified property, true otherwise.
 	 */
-	static bool IsUseSelectedUnrestricted(TSharedRef<FPropertyNode> PropertyNode)
+	static bool IsUseSelectedUnrestricted(TWeakPtr<FPropertyNode> PropertyNode)
 	{
-		return !PropertyNode->IsRestricted(GetSelectionPathNameForProperty(PropertyNode));
+		TSharedPtr<FPropertyNode> PropertyNodePin = PropertyNode.Pin();
+		return ( PropertyNodePin.IsValid() && IsPropertyButtonEnabled(PropertyNode) ) ? !PropertyNodePin->IsRestricted( GetSelectionPathNameForProperty( PropertyNodePin.ToSharedRef() ) ) : false;
 	}
 
 	/**
@@ -737,10 +739,11 @@ namespace PropertyEditorHelpers
 	 * @param  PropertyNode		The property that would be set from the "use selected" button.
 	 * @return A tooltip for the "use selected" button.
 	 */
-	static FText GetUseSelectedTooltip(TSharedRef<FPropertyNode> PropertyNode)
+	static FText GetUseSelectedTooltip(TWeakPtr<FPropertyNode> PropertyNode)
 	{
+		TSharedPtr<FPropertyNode> PropertyNodePin = PropertyNode.Pin();
 		FText ToolTip;
-		if (!PropertyNode->GenerateRestrictionToolTip(GetSelectionPathNameForProperty(PropertyNode), ToolTip))
+		if (PropertyNodePin.IsValid() && !PropertyNodePin->GenerateRestrictionToolTip(GetSelectionPathNameForProperty(PropertyNodePin.ToSharedRef()), ToolTip))
 		{
 			ToolTip = LOCTEXT("UseButtonToolTipText", "Use Selected Asset from Content Browser");
 		}
@@ -752,14 +755,19 @@ namespace PropertyEditorHelpers
 	{
 		TSharedPtr<SWidget> NewButton;
 
+		TWeakPtr<FPropertyNode> WeakPropertyEditor = PropertyEditor->GetPropertyNode();
+
+		TAttribute<bool>::FGetter EnabledDelegate = TAttribute<bool>::FGetter::CreateStatic(&IsPropertyButtonEnabled, WeakPropertyEditor);
+		TAttribute<bool> IsEnabledAttribute = TAttribute<bool>::Create( EnabledDelegate );
+
 		switch( ButtonType )
 		{
 		case EPropertyButton::Add:
-			NewButton = PropertyCustomizationHelpers::MakeAddButton( FSimpleDelegate::CreateSP( PropertyEditor, &FPropertyEditor::AddItem ) );
+			NewButton = PropertyCustomizationHelpers::MakeAddButton( FSimpleDelegate::CreateSP( PropertyEditor, &FPropertyEditor::AddItem ), FText(), IsEnabledAttribute );
 			break;
 
 		case EPropertyButton::Empty:
-			NewButton = PropertyCustomizationHelpers::MakeEmptyButton( FSimpleDelegate::CreateSP( PropertyEditor, &FPropertyEditor::EmptyArray ) );
+			NewButton = PropertyCustomizationHelpers::MakeEmptyButton( FSimpleDelegate::CreateSP( PropertyEditor, &FPropertyEditor::EmptyArray ), FText(), IsEnabledAttribute );
 			break;
 
 		case EPropertyButton::Insert_Delete:
@@ -774,7 +782,7 @@ namespace PropertyEditorHelpers
 				}
 
 				NewButton = PropertyCustomizationHelpers::MakeInsertDeleteDuplicateButton( InsertAction, DeleteAction, DuplicateAction );
-
+				NewButton->SetEnabled( IsEnabledAttribute );
 				break;
 			}
 
@@ -783,14 +791,14 @@ namespace PropertyEditorHelpers
 			break;
 
 		case EPropertyButton::Clear:
-			NewButton = PropertyCustomizationHelpers::MakeClearButton( FSimpleDelegate::CreateSP( PropertyEditor, &FPropertyEditor::ClearItem ) );
+			NewButton = PropertyCustomizationHelpers::MakeClearButton( FSimpleDelegate::CreateSP( PropertyEditor, &FPropertyEditor::ClearItem ), FText(), IsEnabledAttribute );
 			break;
 
 		case EPropertyButton::Use:
 			{
 				FSimpleDelegate OnClickDelegate = FSimpleDelegate::CreateSP(PropertyEditor, &FPropertyEditor::UseSelected);
-				TAttribute<bool>::FGetter EnabledDelegate = TAttribute<bool>::FGetter::CreateStatic(&IsUseSelectedUnrestricted, PropertyEditor->GetPropertyNode());
-				TAttribute<FText>::FGetter TooltipDelegate = TAttribute<FText>::FGetter::CreateStatic(&GetUseSelectedTooltip, PropertyEditor->GetPropertyNode());
+				TAttribute<bool>::FGetter EnabledDelegate = TAttribute<bool>::FGetter::CreateStatic(&IsUseSelectedUnrestricted, WeakPropertyEditor);
+				TAttribute<FText>::FGetter TooltipDelegate = TAttribute<FText>::FGetter::CreateStatic(&GetUseSelectedTooltip, WeakPropertyEditor);
 
 				NewButton = PropertyCustomizationHelpers::MakeUseSelectedButton(OnClickDelegate, TAttribute<FText>::Create(TooltipDelegate), TAttribute<bool>::Create(EnabledDelegate));
 				break;
