Ensure the proper property change flag is being set when resetting to default

diff --git a/Engine/Source/Editor/PropertyEditor/Private/PropertyNode.cpp b/Engine/Source/Editor/PropertyEditor/Private/PropertyNode.cpp
--- a/Engine/Source/Editor/PropertyEditor/Private/PropertyNode.cpp
+++ b/Engine/Source/Editor/PropertyEditor/Private/PropertyNode.cpp
@@ -1648,7 +1648,7 @@ void FPropertyNode::ResetToDefault( FNotifyHook* InNotifyHook )
 		{
 			// Call PostEditchange on all the objects
 			// Assume reset to default, can change topology
-			FPropertyChangedEvent ChangeEvent( TheProperty );
+			FPropertyChangedEvent ChangeEvent( TheProperty, EPropertyChangeType::ValueSet );
 			NotifyPostChange( ChangeEvent, InNotifyHook );
 		}
 
