Attempting to enter greater than the maximum allowed number of FName characters into a name property will now show an error message in the text box in a details panel

https://jira.ol.epicgames.net/browse/UE-14636

diff --git a/Engine/Source/Editor/PropertyEditor/Private/UserInterface/PropertyEditor/SPropertyEditorText.cpp b/Engine/Source/Editor/PropertyEditor/Private/UserInterface/PropertyEditor/SPropertyEditorText.cpp
--- a/Engine/Source/Editor/PropertyEditor/Private/UserInterface/PropertyEditor/SPropertyEditorText.cpp
+++ b/Engine/Source/Editor/PropertyEditor/Private/UserInterface/PropertyEditor/SPropertyEditorText.cpp
@@ -7,6 +7,7 @@
 #include "PropertyEditor.h"
 #include "PropertyEditorHelpers.h"
 
+#define LOCTEXT_NAMESPACE "PropertyEditor"
 
 void SPropertyEditorText::Construct( const FArguments& InArgs, const TSharedRef< class FPropertyEditor >& InPropertyEditor )
 {
@@ -16,6 +17,8 @@ void SPropertyEditorText::Construct( const FArguments& InArgs, const TSharedRef<
 
 	TSharedPtr<SHorizontalBox> HorizontalBox;
 
+	bIsFNameProperty = Property->IsA<UNameProperty>();
+
 	bIsMultiLine = Property->GetBoolMetaData("MultiLine");
 	if(bIsMultiLine)
 	{
@@ -25,18 +28,21 @@ void SPropertyEditorText::Construct( const FArguments& InArgs, const TSharedRef<
 			+SHorizontalBox::Slot()
 			.FillWidth(1.0f)
 			[
-				SAssignNew(PrimaryWidget, SMultiLineEditableTextBox)
+				SAssignNew(MultiLineWidget, SMultiLineEditableTextBox)
 				.Text(InPropertyEditor, &FPropertyEditor::GetValueAsText)
 				.Font(InArgs._Font)
 				.SelectAllTextWhenFocused(false)
 				.ClearKeyboardFocusOnCommit(false)
 				.OnTextCommitted(this, &SPropertyEditorText::OnTextCommitted)
+				.OnTextChanged(this, &SPropertyEditorText::OnMultiLineTextChanged)
 				.SelectAllTextOnCommit(false)
 				.IsReadOnly(this, &SPropertyEditorText::IsReadOnly)
 				.AutoWrapText(true)
 				.ModiferKeyForNewLine(EModifierKey::Shift)
 			]
 		];
+
+		PrimaryWidget = MultiLineWidget;
 	}
 	else
 	{
@@ -46,16 +52,20 @@ void SPropertyEditorText::Construct( const FArguments& InArgs, const TSharedRef<
 			+SHorizontalBox::Slot()
 			.FillWidth(1.0f)
 			[
-				SAssignNew( PrimaryWidget, SEditableTextBox )
+				SAssignNew( SingleLineWidget, SEditableTextBox )
 				.Text( InPropertyEditor, &FPropertyEditor::GetValueAsText )
 				.Font( InArgs._Font )
 				.SelectAllTextWhenFocused( true )
 				.ClearKeyboardFocusOnCommit(false)
 				.OnTextCommitted( this, &SPropertyEditorText::OnTextCommitted )
+				.OnTextChanged( this, &SPropertyEditorText::OnSingleLineTextChanged )
 				.SelectAllTextOnCommit( true )
 				.IsReadOnly(this, &SPropertyEditorText::IsReadOnly)
+			
 			]
 		];
+
+		PrimaryWidget = SingleLineWidget;
 	}
 
 	if( InPropertyEditor->PropertyIsA( UObjectPropertyBase::StaticClass() ) )
@@ -109,6 +119,35 @@ void SPropertyEditorText::OnTextCommitted( const FText& NewText, ETextCommit::Ty
 	PropertyHandle->SetValueFromFormattedString( NewText.ToString() );
 }
 
+static FText ValidateNameLength( const FText& Text )
+{
+	if( Text.ToString().Len() > NAME_SIZE )
+	{
+		static FText ErrorString = FText::Format( LOCTEXT("NamePropertySizeTooLongError", "Name properties may only be a maximum of {0} characters"), FText::AsNumber( NAME_SIZE ) );
+		return ErrorString;
+	}
+
+	return FText::GetEmpty();
+}
+
+void SPropertyEditorText::OnMultiLineTextChanged( const FText& NewText )
+{
+	if( bIsFNameProperty )
+	{
+		FText ErrorMessage = ValidateNameLength( NewText );
+		MultiLineWidget->SetError( ErrorMessage );
+	}
+}
+
+void SPropertyEditorText::OnSingleLineTextChanged( const FText& NewText )
+{
+	if( bIsFNameProperty )
+	{
+		FText ErrorMessage = ValidateNameLength( NewText );
+		SingleLineWidget->SetError( ErrorMessage );
+	}
+}
+
 bool SPropertyEditorText::SupportsKeyboardFocus() const
 {
 	return PrimaryWidget.IsValid() && PrimaryWidget->SupportsKeyboardFocus();
diff --git a/Engine/Source/Editor/PropertyEditor/Private/UserInterface/PropertyEditor/SPropertyEditorText.h b/Engine/Source/Editor/PropertyEditor/Private/UserInterface/PropertyEditor/SPropertyEditorText.h
--- a/Engine/Source/Editor/PropertyEditor/Private/UserInterface/PropertyEditor/SPropertyEditorText.h
+++ b/Engine/Source/Editor/PropertyEditor/Private/UserInterface/PropertyEditor/SPropertyEditorText.h
@@ -27,7 +27,13 @@ public:
 
 private:
 
-	virtual void OnTextCommitted( const FText& NewText, ETextCommit::Type CommitInfo );
+	void OnTextCommitted( const FText& NewText, ETextCommit::Type CommitInfo );
+
+	/** Called if the single line widget text changes */
+	void OnSingleLineTextChanged( const FText& NewText );
+
+	/** Called if the multi line widget text changes */
+	void OnMultiLineTextChanged( const FText& NewText );
 
 	/** @return True if the property can be edited */
 	bool CanEdit() const;
@@ -41,8 +47,17 @@ private:
 
 	TSharedPtr< class SWidget > PrimaryWidget;
 
+	/** Widget used for the multiline version of the text property */
+	TSharedPtr<SMultiLineEditableTextBox> MultiLineWidget;
+
+	/** Widget used for the single line version of the text property */
+	TSharedPtr<SEditableTextBox> SingleLineWidget;
+
 	TOptional< float > PreviousHeight;
 
 	/** Cached flag as we would like multi-line text widgets to be slightly larger */
 	bool bIsMultiLine;
+	
+	/** True if property is an FName property which causes us to run extra size validation checks */
+	bool bIsFNameProperty;
 };
\ No newline at end of file
