Fix details panel buttons not being clickable in object properties with very large names.  Slightly reworked were the asset buttons are.

https://jira.ol.epicgames.net/browse/UE-14391?

diff --git a/Engine/Source/Editor/PropertyEditor/Private/DetailPropertyRow.cpp b/Engine/Source/Editor/PropertyEditor/Private/DetailPropertyRow.cpp
--- a/Engine/Source/Editor/PropertyEditor/Private/DetailPropertyRow.cpp
+++ b/Engine/Source/Editor/PropertyEditor/Private/DetailPropertyRow.cpp
@@ -425,6 +425,7 @@ void FDetailPropertyRow::MakeValueWidget( FDetailWidgetRow& Row, const TSharedPt
 		TSharedPtr<SPropertyValueWidget> PropertyValue;
 
 		ValueWidget->AddSlot()
+		.Padding( 0.0f, 0.0f, 4.0f, 0.0f )
 		[
 			SAssignNew( PropertyValue, SPropertyValueWidget, PropertyEditor, ParentCategory.Pin()->GetParentLayoutImpl().GetPropertyUtilities() )
 			.ShowPropertyButtons( false ) // We handle this ourselves
diff --git a/Engine/Source/Editor/PropertyEditor/Private/PropertyCustomizationHelpers.cpp b/Engine/Source/Editor/PropertyEditor/Private/PropertyCustomizationHelpers.cpp
--- a/Engine/Source/Editor/PropertyEditor/Private/PropertyCustomizationHelpers.cpp
+++ b/Engine/Source/Editor/PropertyEditor/Private/PropertyCustomizationHelpers.cpp
@@ -667,65 +667,73 @@ public:
 	TSharedRef<SWidget> CreateValueContent( const TSharedPtr<FAssetThumbnailPool>& ThumbnailPool )
 	{
 		return
-			SNew(SVerticalBox)
-			+SVerticalBox::Slot()
-			.AutoHeight()
-			.Padding( 0.0f )
-			.VAlign(VAlign_Center)
-			.HAlign(HAlign_Fill)
+			SNew(SHorizontalBox)
+			+ SHorizontalBox::Slot()
 			[
-				SNew(SHorizontalBox)
-				+SHorizontalBox::Slot()
-				.FillWidth(1.0f)
+				SNew(SVerticalBox)
+				+SVerticalBox::Slot()
+				.AutoHeight()
+				.Padding( 0.0f )
+				.VAlign(VAlign_Center)
+				.HAlign(HAlign_Fill)
 				[
-					SNew( SPropertyEditorAsset )
-					.ObjectPath(MaterialItem.Material->GetPathName())
-					.Class(UMaterialInterface::StaticClass())
-					.OnSetObject(this, &FMaterialItemView::OnSetObject)
-					.DisplayThumbnail(true)
-					.ThumbnailPool(ThumbnailPool)
-					.CustomContentSlot()
+					SNew(SHorizontalBox)
+					+SHorizontalBox::Slot()
+					.FillWidth(1.0f)
 					[
-						SNew( SBox )
-						.HAlign(HAlign_Left)
+						SNew( SPropertyEditorAsset )
+						.ObjectPath(MaterialItem.Material->GetPathName())
+						.Class(UMaterialInterface::StaticClass())
+						.OnSetObject(this, &FMaterialItemView::OnSetObject)
+						.DisplayThumbnail(true)
+						.ThumbnailPool(ThumbnailPool)
+						.CustomContentSlot()
 						[
-							// Add a menu for displaying all textures 
-							SNew( SComboButton )
-							.OnGetMenuContent( this, &FMaterialItemView::OnGetTexturesMenuForMaterial )
-							.VAlign( VAlign_Center )
-							.ContentPadding(2)
-							.IsEnabled( this, &FMaterialItemView::IsTexturesMenuEnabled )
-							.ButtonContent()
+							SNew( SBox )
+							.HAlign(HAlign_Left)
 							[
-								SNew( STextBlock )
-								.Font( IDetailLayoutBuilder::GetDetailFont() )
-								.ToolTipText( LOCTEXT("ViewTexturesToolTip", "View the textures used by this material" ) )
-								.Text( LOCTEXT("ViewTextures","Textures") )
+								// Add a menu for displaying all textures 
+								SNew( SComboButton )
+								.OnGetMenuContent( this, &FMaterialItemView::OnGetTexturesMenuForMaterial )
+								.VAlign( VAlign_Center )
+								.ContentPadding(2)
+								.IsEnabled( this, &FMaterialItemView::IsTexturesMenuEnabled )
+								.ButtonContent()
+								[
+									SNew( STextBlock )
+									.Font( IDetailLayoutBuilder::GetDetailFont() )
+									.ToolTipText( LOCTEXT("ViewTexturesToolTip", "View the textures used by this material" ) )
+									.Text( LOCTEXT("ViewTextures","Textures") )
+								]
 							]
 						]
+					
 					]
-					.ResetToDefaultSlot()
-					[
-							// Add a button to reset the material to the base material
-						SNew(SButton)
-						.ToolTipText(LOCTEXT( "ResetToBase", "Reset to base material" ))
-						.ButtonStyle( FEditorStyle::Get(), "NoBorder" )
-						.ContentPadding(0) 
-						.Visibility( this, &FMaterialItemView::GetReplaceVisibility )
-						.OnClicked( this, &FMaterialItemView::OnResetToBaseClicked )
-						[
-							SNew(SImage)
-							.Image( FEditorStyle::GetBrush("PropertyWindow.DiffersFromDefault") )
-						]
-					]
+				]
+				+SVerticalBox::Slot()
+				.AutoHeight()
+				.Padding(2)
+				.VAlign( VAlign_Center )
+				[
+					OnGenerateCustomMaterialWidgets.IsBound() ? OnGenerateCustomMaterialWidgets.Execute( MaterialItem.Material.Get(), MaterialItem.SlotIndex ) : StaticCastSharedRef<SWidget>( SNullWidget::NullWidget )
 				]
 			]
-			+SVerticalBox::Slot()
-			.AutoHeight()
-			.Padding(2)
+			+SHorizontalBox::Slot()
+			.AutoWidth()
 			.VAlign( VAlign_Center )
+			.Padding( 4.0f, 0.0f )
 			[
-				OnGenerateCustomMaterialWidgets.IsBound() ? OnGenerateCustomMaterialWidgets.Execute( MaterialItem.Material.Get(), MaterialItem.SlotIndex ) : StaticCastSharedRef<SWidget>( SNullWidget::NullWidget )
+				// Add a button to reset the material to the base material
+				SNew(SButton)
+				.ToolTipText(LOCTEXT( "ResetToBase", "Reset to base material" ))
+				.ButtonStyle( FEditorStyle::Get(), "NoBorder" )
+				.ContentPadding(0) 
+				.Visibility( this, &FMaterialItemView::GetReplaceVisibility )
+				.OnClicked( this, &FMaterialItemView::OnResetToBaseClicked )
+				[
+					SNew(SImage)
+					.Image( FEditorStyle::GetBrush("PropertyWindow.DiffersFromDefault") )
+				]
 			];
 	}
 
