Adding missing files for CL 2533978

diff --git a/Engine/Source/Editor/PropertyEditor/Private/UserInterface/PropertyEditor/SPropertyEditorAsset.cpp b/Engine/Source/Editor/PropertyEditor/Private/UserInterface/PropertyEditor/SPropertyEditorAsset.cpp
--- a/Engine/Source/Editor/PropertyEditor/Private/UserInterface/PropertyEditor/SPropertyEditorAsset.cpp
+++ b/Engine/Source/Editor/PropertyEditor/Private/UserInterface/PropertyEditor/SPropertyEditorAsset.cpp
@@ -167,46 +167,18 @@ void SPropertyEditorAsset::Construct( const FArguments& InArgs, const TSharedPtr
 		NewAssetFactories = PropertyCustomizationHelpers::GetNewAssetFactoriesForClasses(CustomClassFilters);
 	}
 	
-	TSharedPtr<SHorizontalBox> HorizontalBox = NULL;
+	TSharedPtr<SHorizontalBox> ValueContentBox = NULL;
 	ChildSlot
 	[
 		SNew( SAssetDropTarget )
 		.OnIsAssetAcceptableForDrop( this, &SPropertyEditorAsset::OnAssetDraggedOver )
 		.OnAssetDropped( this, &SPropertyEditorAsset::OnAssetDropped )
 		[
-			SAssignNew( HorizontalBox, SHorizontalBox )			
+			SAssignNew( ValueContentBox, SHorizontalBox )	
 		]
 	];
 
-	if(ShouldDisplayThumbnail(InArgs))
-	{
-		FObjectOrAssetData Value; 
-		GetValue( Value );
-
-		AssetThumbnail = MakeShareable( new FAssetThumbnail( Value.AssetData, InArgs._ThumbnailSize.X, InArgs._ThumbnailSize.Y, InArgs._ThumbnailPool ) );
-		
-		if( AssetThumbnail.IsValid()  )
-		{
-			HorizontalBox->AddSlot()
-				.Padding( 0.0f, 0.0f, 2.0f, 0.0f )
-				.AutoWidth()
-				[
-					SAssignNew( ThumbnailBorder, SBorder )
-					.Padding( 5.0f )
-					.BorderImage( this, &SPropertyEditorAsset::GetThumbnailBorder )
-					.OnMouseDoubleClick( this, &SPropertyEditorAsset::OnAssetThumbnailDoubleClick )
-					[
-						SNew( SBox )
-						.ToolTipText( this, &SPropertyEditorAsset::OnGetToolTip )
-						.WidthOverride( InArgs._ThumbnailSize.X ) 
-						.HeightOverride( InArgs._ThumbnailSize.Y )
-						[
-							AssetThumbnail->MakeThumbnailWidget()
-						]
-					]
-				];
-		}
-	}
+	TAttribute<bool> IsEnabledAttribute( this, &SPropertyEditorAsset::CanEdit );
 
 	AssetComboButton = SNew(SComboButton)
 		.ToolTipText( this, &SPropertyEditorAsset::OnGetToolTip )
@@ -214,6 +186,7 @@ void SPropertyEditorAsset::Construct( const FArguments& InArgs, const TSharedPtr
 		.ForegroundColor(FEditorStyle::GetColor("PropertyEditor.AssetName.ColorAndOpacity"))
 		.OnGetMenuContent( this, &SPropertyEditorAsset::OnGetMenuContent )
 		.OnMenuOpenChanged( this, &SPropertyEditorAsset::OnMenuOpenChanged )
+		.IsEnabled( IsEnabledAttribute )
 		.ContentPadding(2.0f)
 		.ButtonContent()
 		[
@@ -224,51 +197,106 @@ void SPropertyEditorAsset::Construct( const FArguments& InArgs, const TSharedPtr
 			.Text(this,&SPropertyEditorAsset::OnGetAssetName)
 		];
 
+	
+	TSharedRef<SHorizontalBox> ButtonBox = SNew( SHorizontalBox );
+	
+	TSharedPtr<SVerticalBox> CustomContentBox;
 
-	TSharedRef<SHorizontalBox> ButtonBox = SNew( SHorizontalBox )
-		.IsEnabled( this, &SPropertyEditorAsset::CanEdit )
-		+ SHorizontalBox::Slot()
+	if( ShouldDisplayThumbnail(InArgs) )
+	{
+		FObjectOrAssetData Value; 
+		GetValue( Value );
+
+		AssetThumbnail = MakeShareable( new FAssetThumbnail( Value.AssetData, InArgs._ThumbnailSize.X, InArgs._ThumbnailSize.Y, InArgs._ThumbnailPool ) );
+
+		ValueContentBox->AddSlot()
+		.Padding( 0.0f, 0.0f, 2.0f, 0.0f )
 		.AutoWidth()
 		[
-			AssetComboButton.ToSharedRef()
+			SAssignNew( ThumbnailBorder, SBorder )
+			.Padding( 5.0f )
+			.BorderImage( this, &SPropertyEditorAsset::GetThumbnailBorder )
+			.OnMouseDoubleClick( this, &SPropertyEditorAsset::OnAssetThumbnailDoubleClick )
+			[
+				SNew( SBox )
+				.ToolTipText( this, &SPropertyEditorAsset::OnGetToolTip )
+				.WidthOverride( InArgs._ThumbnailSize.X ) 
+				.HeightOverride( InArgs._ThumbnailSize.Y )
+				[
+					AssetThumbnail->MakeThumbnailWidget()
+				]
+			]
 		];
 
-	
-	HorizontalBox->AddSlot()
-	.FillWidth(1.0f)
-	.Padding( 0.0f, 4.0f, 4.0f, 4.0f )
-	.VAlign(VAlign_Center)
-	[
-		SNew(SVerticalBox)
-		+SVerticalBox::Slot()
-		.VAlign( VAlign_Center )
+
+		ValueContentBox->AddSlot()
 		[
-			ButtonBox
-		]
-		+SVerticalBox::Slot()
-		.AutoHeight()
-		.Padding( InArgs._CustomContentSlot.Widget == SNullWidget::NullWidget ? FMargin( 0.0f, 0.0f ) : FMargin( 0.0f, 2.0f ) )
+			SNew( SBox )
+			.VAlign( VAlign_Center )
+			[
+				SAssignNew( CustomContentBox, SVerticalBox )
+				+ SVerticalBox::Slot()
+				[
+					AssetComboButton.ToSharedRef()
+				]
+				+ SVerticalBox::Slot()
+				.AutoHeight()
+				.Padding( 0.0f, 2.0f, 4.0f, 2.0f )
+				[
+					ButtonBox
+				]
+			]
+		];
+	}
+	else
+	{
+		ValueContentBox->AddSlot()
+		[
+			SAssignNew( CustomContentBox, SVerticalBox )
+			+SVerticalBox::Slot()
+			.VAlign( VAlign_Center )
+			[
+				SNew( SHorizontalBox )
+				+ SHorizontalBox::Slot()
+				[
+					AssetComboButton.ToSharedRef()
+				]
+				+ SHorizontalBox::Slot()
+				.AutoWidth()
+				.Padding( 4.f, 0.f )
+				[
+					ButtonBox
+				]
+			]
+		];
+	}
+
+	if( InArgs._CustomContentSlot.Widget != SNullWidget::NullWidget )
+	{
+		CustomContentBox->AddSlot()
+		.VAlign( VAlign_Center )
+		.Padding( FMargin( 0.0f, 2.0f ) )
 		[
 			InArgs._CustomContentSlot.Widget
-		]
-	];
+		];
+	}
 
-	if(!bIsActor && InArgs._DisplayUseSelected)
+	if( !bIsActor && InArgs._DisplayUseSelected )
 	{
 		ButtonBox->AddSlot()
 		.VAlign(VAlign_Center)
-		.Padding( 2.0f, 0.0f )
 		.AutoWidth()
+		.Padding( 2.0f, 0.0f )
 		[
-			PropertyCustomizationHelpers::MakeUseSelectedButton( FSimpleDelegate::CreateSP( this, &SPropertyEditorAsset::OnUse ) )
+			PropertyCustomizationHelpers::MakeUseSelectedButton( FSimpleDelegate::CreateSP( this, &SPropertyEditorAsset::OnUse ), FText(), IsEnabledAttribute )
 		];
 	}
 
-	if(InArgs._DisplayBrowse)
+	if( InArgs._DisplayBrowse )
 	{
 		ButtonBox->AddSlot()
-		.AutoWidth()
 		.Padding( 2.0f, 0.0f )
+		.AutoWidth()
 		.VAlign(VAlign_Center)
 		[
 			PropertyCustomizationHelpers::MakeBrowseButton(
@@ -278,34 +306,40 @@ void SPropertyEditorAsset::Construct( const FArguments& InArgs, const TSharedPtr
 		];
 	}
 
-	if(bIsActor)
+	if( bIsActor )
 	{
+		TSharedRef<SWidget> ActorPicker = PropertyCustomizationHelpers::MakeInteractiveActorPicker( FOnGetAllowedClasses::CreateSP(this, &SPropertyEditorAsset::OnGetAllowedClasses), FOnShouldFilterActor(), FOnActorSelected::CreateSP( this, &SPropertyEditorAsset::OnActorSelected ) );
+		ActorPicker->SetEnabled( IsEnabledAttribute );
+
 		ButtonBox->AddSlot()
-		.AutoWidth()
 		.Padding( 2.0f, 0.0f )
+		.AutoWidth()
 		.VAlign(VAlign_Center)
 		[
-			PropertyCustomizationHelpers::MakeInteractiveActorPicker( FOnGetAllowedClasses::CreateSP(this, &SPropertyEditorAsset::OnGetAllowedClasses), FOnShouldFilterActor(), FOnActorSelected::CreateSP( this, &SPropertyEditorAsset::OnActorSelected ) )
+			ActorPicker
 		];
 	}
 
 	if( InArgs._ResetToDefaultSlot.Widget != SNullWidget::NullWidget )
 	{
+		TSharedRef<SWidget> ResetToDefaultWidget  = InArgs._ResetToDefaultSlot.Widget;
+		ResetToDefaultWidget->SetEnabled( IsEnabledAttribute );
+
 		ButtonBox->AddSlot()
+		.Padding( 4.0f, 0.0f )
 		.AutoWidth()
-		.Padding( 2.0f, 0.0f )
 		.VAlign(VAlign_Center)
 		[
-			InArgs._ResetToDefaultSlot.Widget
+			ResetToDefaultWidget
 		];
 	}
 }
 
 void SPropertyEditorAsset::GetDesiredWidth( float& OutMinDesiredWidth, float &OutMaxDesiredWidth )
 {
-	OutMinDesiredWidth = 250.0f;
+	OutMinDesiredWidth = 250.f;
 	// No max width
-	OutMaxDesiredWidth = 0.0f;
+	OutMaxDesiredWidth = 350.f;
 }
 
 const FSlateBrush* SPropertyEditorAsset::GetThumbnailBorder() const
