Fix crash if re-run construction scripts is called while routing post-edit change on a multi-object selection

https://jira.ol.epicgames.net/browse/UE-15047

diff --git a/Engine/Source/Editor/PropertyEditor/Private/PropertyNode.cpp b/Engine/Source/Editor/PropertyEditor/Private/PropertyNode.cpp
--- a/Engine/Source/Editor/PropertyEditor/Private/PropertyNode.cpp
+++ b/Engine/Source/Editor/PropertyEditor/Private/PropertyNode.cpp
@@ -1956,7 +1956,10 @@ void FPropertyNode::NotifyPostChange( FPropertyChangedEvent& InPropertyChangedEv
 						ChangedEvent = FPropertyChangedEvent(CurProperty, InPropertyChangedEvent.ChangeType);
 					}
 					ChangedEvent.ObjectIteratorIndex = CurrentObjectIndex;
-					Object->PostEditChangeProperty( ChangedEvent );
+					if( Object )
+					{
+						Object->PostEditChangeProperty( ChangedEvent );
+					}
 				}
 				else
 				{
@@ -1968,7 +1971,10 @@ void FPropertyNode::NotifyPostChange( FPropertyChangedEvent& InPropertyChangedEv
 					}
 					FPropertyChangedChainEvent ChainEvent(*PropertyChain, ChangedEvent);
 					ChainEvent.ObjectIteratorIndex = CurrentObjectIndex;
-					Object->PostEditChangeChainProperty(ChainEvent);
+					if( Object )
+					{
+						Object->PostEditChangeChainProperty(ChainEvent);
+					}
 				}
 				LevelDirtyCallback.Request();
 				++CurrentObjectIndex;
