Fix property propagation from CDO to instances.

diff --git a/Engine/Source/Editor/PropertyEditor/Private/PropertyHandleImpl.cpp b/Engine/Source/Editor/PropertyEditor/Private/PropertyHandleImpl.cpp
--- a/Engine/Source/Editor/PropertyEditor/Private/PropertyHandleImpl.cpp
+++ b/Engine/Source/Editor/PropertyEditor/Private/PropertyHandleImpl.cpp
@@ -346,13 +346,6 @@ FPropertyAccess::Result FPropertyValueImpl::ImportText( const TArray<FObjectBase
 		const bool bTransactable = (Flags & EPropertyValueSetFlags::NotTransactable) == 0;
 		const bool bFinished = ( Flags & EPropertyValueSetFlags::InteractiveChange) == 0;
 		
-		// Cache previous values to compare against them after all potential modifications to property.
-		TArray<FString> PreviousValues;
-		PreviousValues.Reserve(InObjects.Num());
-
-		TArray<FString> PreviousArrayValues;
-		PreviousArrayValues.Reserve(InObjects.Num());
-
 		for ( int32 ObjectIndex = 0 ; ObjectIndex < InObjects.Num() ; ++ObjectIndex )
 		{	
 			const FObjectBaseAddress& Cur = InObjects[ ObjectIndex ];
@@ -366,7 +359,6 @@ FPropertyAccess::Result FPropertyValueImpl::ImportText( const TArray<FObjectBase
 			// Cache the value of the property before modifying it.
 			FString PreviousValue;
 			NodeProperty->ExportText_Direct(PreviousValue, Cur.BaseAddress, Cur.BaseAddress, NULL, 0 );
-			PreviousValues.Add(PreviousValue);
 
 			// If this property is the inner-property of an array, cache the current value as well
 			FString PreviousArrayValue;
@@ -384,7 +376,6 @@ FPropertyAccess::Result FPropertyValueImpl::ImportText( const TArray<FObjectBase
 					}
 				}
 			}
-			PreviousArrayValues.Add(PreviousArrayValue);
 
 			// Check if we need to call PreEditChange on all objects.
 			// Remove quotes from the original value because FName properties  
@@ -413,6 +404,26 @@ FPropertyAccess::Result FPropertyValueImpl::ImportText( const TArray<FObjectBase
 			const TCHAR* NewValue = *InValues[ObjectIndex];
 			NodeProperty->ImportText( NewValue, Cur.BaseAddress, 0, Cur.Object );
 
+			if (Cur.Object)
+			{
+				// Cache the value of the property after having modified it.
+				FString ValueAfterImport;
+				NodeProperty->ExportText_Direct(ValueAfterImport, Cur.BaseAddress, Cur.BaseAddress, NULL, 0);
+
+				if ((Cur.Object->HasAnyFlags(RF_ClassDefaultObject | RF_ArchetypeObject) ||
+					(Cur.Object->HasAnyFlags(RF_DefaultSubObject) && Cur.Object->GetOuter()->HasAnyFlags(RF_ClassDefaultObject | RF_ArchetypeObject))) &&
+					!bIsGameWorld)
+				{
+					InPropertyNode->PropagatePropertyChange(Cur.Object, NewValue, PreviousArrayValue.IsEmpty() ? PreviousValue : PreviousArrayValue);
+				}
+
+				// If the values before and after setting the property differ, mark the object dirty.
+				if (FCString::Strcmp(*PreviousValue, *ValueAfterImport) != 0)
+				{
+					Cur.Object->MarkPackageDirty();
+				}
+			}
+
 			//add on array index so we can tell which entry just changed
 			ArrayIndicesPerObject.Add(TMap<FString,int32>());
 			FPropertyValueImpl::GenerateArrayIndexMapToObjectNode( ArrayIndicesPerObject[ObjectIndex], InPropertyNode );
@@ -439,41 +450,6 @@ FPropertyAccess::Result FPropertyValueImpl::ImportText( const TArray<FObjectBase
 			}
 		}
 
-		for (int32 ObjectIndex = 0; ObjectIndex < InObjects.Num(); ++ObjectIndex)
-		{
-			const FObjectBaseAddress& Cur = InObjects[ObjectIndex];
-
-			if (Cur.Object == nullptr)
-			{
-				continue;
-			}
-
-			// Cache the value of the property after having modified it (either by entering value or via PostEditChangeProperty).
-			FString ValueAfterImport;
-			NodeProperty->ExportText_Direct(ValueAfterImport, Cur.BaseAddress, Cur.BaseAddress, NULL, 0);
-
-			const FString& PreviousValue = PreviousValues[ObjectIndex];
-
-			// If the values before and after setting the property are the same, we don't need to do anything.
-			if (FCString::Strcmp(*PreviousValue, *ValueAfterImport) == 0)
-			{
-				continue;
-			}
-
-			const FString& PreviousArrayValue = PreviousArrayValues[ObjectIndex];
-
-			// Propagate property change to all objects using this archetype/CDO (or parent's archetype/CDE if this is default subobject)
-			if ((Cur.Object->HasAnyFlags(RF_ClassDefaultObject | RF_ArchetypeObject) ||
-				(Cur.Object->HasAnyFlags(RF_DefaultSubObject) && Cur.Object->GetOuter()->HasAnyFlags(RF_ClassDefaultObject | RF_ArchetypeObject))) &&
-				!bIsGameWorld)
-			{
-				InPropertyNode->PropagatePropertyChange(Cur.Object, *ValueAfterImport, PreviousArrayValue.IsEmpty() ? PreviousValue : PreviousArrayValue);
-			}
-
-			// Property changed, mark package as dirty.
-			Cur.Object->MarkPackageDirty();
-		}
-
 		if (OldGWorld)
 		{
 			// restore the original (editor) GWorld
diff --git a/Engine/Source/Runtime/CoreUObject/Private/UObject/Obj.cpp b/Engine/Source/Runtime/CoreUObject/Private/UObject/Obj.cpp
--- a/Engine/Source/Runtime/CoreUObject/Private/UObject/Obj.cpp
+++ b/Engine/Source/Runtime/CoreUObject/Private/UObject/Obj.cpp
@@ -331,14 +331,47 @@ void UObject::PostEditChangeChainProperty(FPropertyChangedChainEvent& PropertyCh
 		PropertyEvent.SetActiveMemberProperty( PropertyChangedEvent.PropertyChain.GetActiveMemberNode()->GetValue() );
 	}
 
-	if ( HasAnyFlags(RF_ClassDefaultObject|RF_ArchetypeObject) && PropertyChangedEvent.PropertyChain.GetActiveMemberNode() == PropertyChangedEvent.PropertyChain.GetHead() && !FApp::IsGame() )
+	// Propagate change to archetype instances first if necessary.
+	if (!FApp::IsGame())
 	{
-		// Get a list of all objects which will be affected by this change; 
-		TArray<UObject*> Objects;
-		GetArchetypeInstances(Objects);
+		if (HasAnyFlags(RF_ClassDefaultObject | RF_ArchetypeObject) && PropertyChangedEvent.PropertyChain.GetActiveMemberNode() == PropertyChangedEvent.PropertyChain.GetHead())
+	{
+			// Get a list of all archetype instances
+			TArray<UObject*> ArchetypeInstances;
+			GetArchetypeInstances(ArchetypeInstances);
 
 		// Propagate the editchange call to archetype instances
-		PropagatePostEditChange(Objects, PropertyChangedEvent);
+			PropagatePostEditChange(ArchetypeInstances, PropertyChangedEvent);
+		}
+		else if (GetOuter()->HasAnyFlags(RF_ClassDefaultObject | RF_ArchetypeObject))
+		{
+			// Get a list of all outer's archetype instances
+			TArray<UObject*> ArchetypeInstances;
+			GetOuter()->GetArchetypeInstances(ArchetypeInstances);
+
+			// Find UProperty describing this in Outer.
+			for (UProperty* Property = GetOuter()->GetClass()->RefLink; Property != nullptr; Property = Property->NextRef)
+			{
+				if (this != *Property->ContainerPtrToValuePtr<UObject*>(GetOuter()))
+				{
+					continue;
+				}
+
+				// Since we found property, propagate PostEditChange to all relevant components of archetype instances.
+				TArray<UObject*> ArchetypeComponentInstances;
+				for (auto ArchetypeInstance : ArchetypeInstances)
+				{
+					if (UObject* ComponentInstance = *Property->ContainerPtrToValuePtr<UObject*>(ArchetypeInstance))
+					{
+						ArchetypeComponentInstances.Add(ComponentInstance);
+					}
+				}
+
+				GetOuter()->PropagatePostEditChange(ArchetypeComponentInstances, PropertyChangedEvent);
+
+				break;
+			}
+		}
 	}
 
 	PostEditChangeProperty(PropertyEvent);
@@ -395,7 +428,7 @@ void UObject::PropagatePostEditChange( TArray<UObject*>& AffectedObjects, FPrope
 
 		// in order to ensure that all objects are re-initialized properly, only process the objects which have this object as their
 		// ObjectArchetype
-		if ( Obj->GetArchetype() == this )
+		if ( Obj->GetArchetype() == this || Obj->GetOuter()->GetArchetype() == this )
 		{
 			// add this object to the list that we're going to process
 			Instances.Add(Obj);
