Property Editor: Added a check for no instances when validating an external property [UE-14270]

diff --git a/Engine/Source/Editor/PropertyEditor/Private/ItemPropertyNode.cpp b/Engine/Source/Editor/PropertyEditor/Private/ItemPropertyNode.cpp
--- a/Engine/Source/Editor/PropertyEditor/Private/ItemPropertyNode.cpp
+++ b/Engine/Source/Editor/PropertyEditor/Private/ItemPropertyNode.cpp
@@ -230,8 +230,7 @@ void FItemPropertyNode::InitChildNodes()
 		{
 			// We've got some addresses, and we know they're all NULL or non-NULL.
 			// Have a peek at the first one, and only build an objects node if we've got addresses.
-			UObject* obj = ObjectProperty->GetObjectPropertyValue(ReadAddresses.GetAddress(0));
-			if( obj )
+			if( UObject* Obj = (ReadAddresses.Num() > 0) ? ObjectProperty->GetObjectPropertyValue(ReadAddresses.GetAddress(0)) : nullptr )
 			{
 				//verify it's not above in the hierarchy somewhere
 				FObjectPropertyNode* ParentObjectNode = FindObjectItemParent();
@@ -239,7 +238,7 @@ void FItemPropertyNode::InitChildNodes()
 				{
 					for ( TPropObjectIterator Itor( ParentObjectNode->ObjectIterator() ) ; Itor ; ++Itor )
 					{
-						if (*Itor == obj)
+						if (*Itor == Obj)
 						{
 							SetNodeFlags(EPropertyNodeFlags::NoChildrenDueToCircularReference, true);
 							//stop the circular loop!!!
