Config Heirarchy button is now enable even when ini is read-only. Fix UE-14448.

diff --git a/Engine/Source/Editor/EditorStyle/Private/SlateEditorStyle.cpp b/Engine/Source/Editor/EditorStyle/Private/SlateEditorStyle.cpp
--- a/Engine/Source/Editor/EditorStyle/Private/SlateEditorStyle.cpp
+++ b/Engine/Source/Editor/EditorStyle/Private/SlateEditorStyle.cpp
@@ -3207,7 +3207,8 @@ void FSlateEditorStyle::FStyle::SetupPropertyEditorStyles()
 		Set( "DetailsView.PulldownArrow.Up", new IMAGE_BRUSH( "PropertyView/AdvancedButton_Up", FVector2D(10,8) ) );
 		Set( "DetailsView.PulldownArrow.Up.Hovered", new IMAGE_BRUSH( "PropertyView/AdvancedButton_Up_Hovered", FVector2D(10,8) ) );
 
-		Set( "DetailsView.EditRawProperties", new IMAGE_BRUSH( "Icons/icon_PropertyMatrix_16px", Icon16x16, FLinearColor( 0, 0, 0, 255 ) ) );
+		Set( "DetailsView.EditRawProperties", new IMAGE_BRUSH( "Icons/icon_PropertyMatrix_16px", Icon16x16, FLinearColor::Black ) );
+		Set( "DetailsView.EditConfigProperties", new IMAGE_BRUSH( "Icons/icon_PropertyMatrix_16px", Icon16x16, FLinearColor::White ) );
 
 		Set( "DetailsView.CollapsedCategory", new BOX_BRUSH( "Common/GroupBorder", FMargin(4.0f/16.0f) ) );
 		Set( "DetailsView.CategoryTop", new BOX_BRUSH( "PropertyView/DetailCategoryTop", FMargin( 4/16.0f, 8.0f/16.0f, 4/16.0f, 4/16.0f ) ) );
diff --git a/Engine/Source/Editor/PropertyEditor/Private/DetailPropertyRow.h b/Engine/Source/Editor/PropertyEditor/Private/DetailPropertyRow.h
--- a/Engine/Source/Editor/PropertyEditor/Private/DetailPropertyRow.h
+++ b/Engine/Source/Editor/PropertyEditor/Private/DetailPropertyRow.h
@@ -49,6 +49,11 @@ public:
 	 */
 	TSharedPtr<FPropertyNode> GetPropertyNode() { return PropertyNode; }
 
+	/**
+	 * @return The property node for this row
+	 */
+	TSharedPtr<FPropertyEditor> GetPropertyEditor() { return PropertyEditor; }
+
 	/**
 	 * Called when children of this row should be generated
 	 *
diff --git a/Engine/Source/Editor/PropertyEditor/Private/PropertyCustomizationHelpers.cpp b/Engine/Source/Editor/PropertyEditor/Private/PropertyCustomizationHelpers.cpp
--- a/Engine/Source/Editor/PropertyEditor/Private/PropertyCustomizationHelpers.cpp
+++ b/Engine/Source/Editor/PropertyEditor/Private/PropertyCustomizationHelpers.cpp
@@ -254,7 +254,7 @@ namespace PropertyCustomizationHelpers
 			SNew(SPropertyEditorButton)
 			.Text(LOCTEXT("EditConfigHierarchyButtonLabel", "Edit Config Hierarchy"))
 			.ToolTipText(OptionalToolTipText.Get().IsEmpty() ? LOCTEXT("EditConfigHierarchyButtonToolTipText", "Edit the config values of this property") : OptionalToolTipText)
-			.Image(FEditorStyle::GetBrush("DetailsView.EditRawProperties"))
+			.Image(FEditorStyle::GetBrush("DetailsView.EditConfigProperties"))
 			.OnClickAction(OnEditConfigClicked)
 			.IsEnabled(IsEnabled)
 			.IsFocusable(false);
diff --git a/Engine/Source/Editor/PropertyEditor/Private/PropertyEditorHelpers.cpp b/Engine/Source/Editor/PropertyEditor/Private/PropertyEditorHelpers.cpp
--- a/Engine/Source/Editor/PropertyEditor/Private/PropertyEditorHelpers.cpp
+++ b/Engine/Source/Editor/PropertyEditor/Private/PropertyEditorHelpers.cpp
@@ -491,14 +491,6 @@ namespace PropertyEditorHelpers
 		// If the property is an item of a const array, don't create any buttons.
 		const UArrayProperty* OuterArrayProp = Cast<UArrayProperty>( NodeProperty->GetOuter() );
 
-		if (NodeProperty->HasAnyPropertyFlags(CPF_GlobalConfig|CPF_Config))
-		{
-			if (NodeProperty->HasMetaData(TEXT("ConfigHierarchyEditable")))
-			{
-				OutRequiredButtons.Add( EPropertyButton::EditConfigHierarchy );
-			}
-		}
-
 		//////////////////////////////
 		// Handle an array property.
 		if( NodeProperty->IsA(UArrayProperty::StaticClass() ) )
diff --git a/Engine/Source/Editor/PropertyEditor/Private/SDetailSingleItemRow.cpp b/Engine/Source/Editor/PropertyEditor/Private/SDetailSingleItemRow.cpp
--- a/Engine/Source/Editor/PropertyEditor/Private/SDetailSingleItemRow.cpp
+++ b/Engine/Source/Editor/PropertyEditor/Private/SDetailSingleItemRow.cpp
@@ -6,6 +6,7 @@
 #include "PropertyEditorHelpers.h"
 #include "IDetailKeyframeHandler.h"
 #include "IDetailPropertyExtensionHandler.h"
+#include "DetailPropertyRow.h"
 
 namespace DetailWidgetConstants
 {
@@ -79,7 +80,22 @@ void SDetailSingleItemRow::Construct( const FArguments& InArgs, FDetailLayoutCus
 	if( InCustomization->IsValidCustomization() )
 	{
 		FDetailWidgetRow Row = InCustomization->GetWidgetRow();
-			
+		TSharedPtr<FPropertyEditor> PropertyEditor;
+		auto bRequiresEditConfigHierarchy = false;
+		auto PropertyNode = InCustomization->GetPropertyNode();
+		if (PropertyNode.IsValid() && PropertyNode->GetProperty())
+		{
+			auto Prop = PropertyNode->GetProperty();
+			if (Prop->HasAnyPropertyFlags(CPF_GlobalConfig | CPF_Config))
+			{
+				if (Prop->HasMetaData(TEXT("ConfigHierarchyEditable")))
+				{
+					PropertyEditor = InCustomization->PropertyRow->GetPropertyEditor();
+					bRequiresEditConfigHierarchy = true;
+				}
+			}
+		}
+
 		TSharedPtr<SWidget> NameWidget;
 		TSharedPtr<SWidget> ValueWidget;
 	
@@ -110,7 +126,8 @@ void SDetailSingleItemRow::Construct( const FArguments& InArgs, FDetailLayoutCus
 		if( bHasMultipleColumns )
 		{
 			NameWidget->SetEnabled(IsPropertyEditingEnabled);
-			Widget = 
+			TSharedPtr<SHorizontalBox> HBox;
+			TSharedRef<SSplitter> Splitter = 
 				SNew( SSplitter )
 				.Style( FEditorStyle::Get(), "DetailsView.Splitter" )
 				.PhysicalSplitterHandleSize( 1.0f )
@@ -150,15 +167,34 @@ void SDetailSingleItemRow::Construct( const FArguments& InArgs, FDetailLayoutCus
 				.OnSlotResized( ColumnSizeData.OnWidthChanged )
 				[
 					SNew( SHorizontalBox )
-					.IsEnabled(IsPropertyEditingEnabled)
 					+ SHorizontalBox::Slot()
-					.Padding( DetailWidgetConstants::RightRowPadding )
-					.HAlign( Row.ValueWidget.HorizontalAlignment )
-					.VAlign( Row.ValueWidget.VerticalAlignment )
 					[
-						ValueWidget.ToSharedRef()
+						SAssignNew(HBox, SHorizontalBox)
+						+ SHorizontalBox::Slot()
+						.Padding(DetailWidgetConstants::RightRowPadding)
+						.HAlign(Row.ValueWidget.HorizontalAlignment)
+						.VAlign(Row.ValueWidget.VerticalAlignment)
+						[
+							SNew(SBox)
+							.IsEnabled(IsPropertyEditingEnabled)
+							[
+								ValueWidget.ToSharedRef()
+							]
+						]
 					]
 				];
+				if (bRequiresEditConfigHierarchy && PropertyEditor.IsValid())
+				{
+					HBox->AddSlot()
+					.AutoWidth()
+					.Padding(0)
+					//.HAlign(HAlign_Right)
+					.VAlign(VAlign_Center)
+					[
+						PropertyCustomizationHelpers::MakeEditConfigHierarchyButton(FSimpleDelegate::CreateSP(PropertyEditor.ToSharedRef(), &FPropertyEditor::EditConfigHierarchy))
+					];
+				}
+				Widget = Splitter;
 		}
 		else
 		{
