Fix crash hitting f7 to compile a blueprint while focused on a default value entry node in a blueprint

https://jira.ol.epicgames.net/browse/UE-15622

diff --git a/Engine/Source/Editor/PropertyEditor/Private/PropertyHandleImpl.cpp b/Engine/Source/Editor/PropertyEditor/Private/PropertyHandleImpl.cpp
--- a/Engine/Source/Editor/PropertyEditor/Private/PropertyHandleImpl.cpp
+++ b/Engine/Source/Editor/PropertyEditor/Private/PropertyHandleImpl.cpp
@@ -307,7 +307,12 @@ FPropertyAccess::Result FPropertyValueImpl::ImportText( const TArray<FObjectBase
 
 	FPropertyAccess::Result Result = FPropertyAccess::Success;
 
-	if( NodeProperty->IsA<UObjectProperty>() || NodeProperty->IsA<UNameProperty>() )
+	if( !NodeProperty )
+	{
+		// The property has been deleted out from under this
+		Result = FPropertyAccess::Fail;
+	}
+	else if( NodeProperty->IsA<UObjectProperty>() || NodeProperty->IsA<UNameProperty>() )
 	{
 		// certain properties have requirements on the size of string values that can be imported.  Search for strings that are too large.
 		for( int32 ValueIndex = 0; ValueIndex < InValues.Num(); ++ValueIndex )
@@ -884,58 +889,60 @@ template< typename Type >
 static Type ClampValueFromMetaData(Type InValue, FPropertyNode& InPropertyNode )
 {
 	UProperty* Property = InPropertyNode.GetProperty();
-	checkSlow(Property);
 
 	Type RetVal = InValue;
-	//enforce min
-	const FString& MinString = Property->GetMetaData(TEXT("ClampMin"));
-	if (MinString.Len())
-	{
-		checkSlow(MinString.IsNumeric());
-		Type MinValue;
-		TTypeFromString<Type>::FromString(MinValue, *MinString);
-		RetVal = FMath::Max<Type>(MinValue, RetVal);
-	}
-	//Enforce max 
-	const FString& MaxString = Property->GetMetaData(TEXT("ClampMax"));
-	if (MaxString.Len())
-	{
-		checkSlow(MaxString.IsNumeric());
-		Type MaxValue;
-		TTypeFromString<Type>::FromString(MaxValue, *MaxString);
-		RetVal = FMath::Min<Type>(MaxValue, RetVal);
-	}
-
-	const bool bIsInteger = Property->IsA(UIntProperty::StaticClass());
-	const bool bIsNonEnumByte = ( Property->IsA(UByteProperty::StaticClass()) && Cast<const UByteProperty>(Property)->Enum == NULL);
-
-	if( bIsInteger || bIsNonEnumByte )
+	if( Property )
 	{
-		//if there is "Multiple" meta data, the selected number is a multiple
-		const FString& MultipleString = Property->GetMetaData(TEXT("Multiple"));
-		if (MultipleString.Len())
+		//enforce min
+		const FString& MinString = Property->GetMetaData(TEXT("ClampMin"));
+		if(MinString.Len())
 		{
-			check(MultipleString.IsNumeric());
-			int32 MultipleValue = FCString::Atoi(*MultipleString);
-			if (MultipleValue!=0)
-			{
-				RetVal -= int32(RetVal)%MultipleValue;
-			}
+			checkSlow(MinString.IsNumeric());
+			Type MinValue;
+			TTypeFromString<Type>::FromString(MinValue, *MinString);
+			RetVal = FMath::Max<Type>(MinValue, RetVal);
 		}
+		//Enforce max 
+		const FString& MaxString = Property->GetMetaData(TEXT("ClampMax"));
+		if(MaxString.Len())
+		{
+			checkSlow(MaxString.IsNumeric());
+			Type MaxValue;
+			TTypeFromString<Type>::FromString(MaxValue, *MaxString);
+			RetVal = FMath::Min<Type>(MaxValue, RetVal);
+		}
+
+		const bool bIsInteger = Property->IsA(UIntProperty::StaticClass());
+		const bool bIsNonEnumByte = (Property->IsA(UByteProperty::StaticClass()) && Cast<const UByteProperty>(Property)->Enum == NULL);
 
-		//enforce array bounds
-		const FString& ArrayClampString = Property->GetMetaData(TEXT("ArrayClamp"));
-		if (ArrayClampString.Len())
+		if(bIsInteger || bIsNonEnumByte)
 		{
-			FObjectPropertyNode* ObjectPropertyNode = InPropertyNode.FindObjectItemParent();
-			if (ObjectPropertyNode && ObjectPropertyNode->GetNumObjects() == 1)
+			//if there is "Multiple" meta data, the selected number is a multiple
+			const FString& MultipleString = Property->GetMetaData(TEXT("Multiple"));
+			if(MultipleString.Len())
 			{
-				int32 LastValidIndex = GetArrayPropertyLastValidIndex(ObjectPropertyNode, ArrayClampString);
-				RetVal = FMath::Clamp<int32>(RetVal, 0, LastValidIndex);
+				check(MultipleString.IsNumeric());
+				int32 MultipleValue = FCString::Atoi(*MultipleString);
+				if(MultipleValue!=0)
+				{
+					RetVal -= int32(RetVal)%MultipleValue;
+				}
 			}
-			else
+
+			//enforce array bounds
+			const FString& ArrayClampString = Property->GetMetaData(TEXT("ArrayClamp"));
+			if(ArrayClampString.Len())
 			{
-				UE_LOG(LogPropertyNode, Warning, TEXT("Array Clamping isn't supported in multi-select (Param Name: %s)"), *Property->GetName());
+				FObjectPropertyNode* ObjectPropertyNode = InPropertyNode.FindObjectItemParent();
+				if(ObjectPropertyNode && ObjectPropertyNode->GetNumObjects() == 1)
+				{
+					int32 LastValidIndex = GetArrayPropertyLastValidIndex(ObjectPropertyNode, ArrayClampString);
+					RetVal = FMath::Clamp<int32>(RetVal, 0, LastValidIndex);
+				}
+				else
+				{
+					UE_LOG(LogPropertyNode, Warning, TEXT("Array Clamping isn't supported in multi-select (Param Name: %s)"), *Property->GetName());
+				}
 			}
 		}
 	}
