Added support for all signed and unsigned integer types in property editors.

diff --git a/Engine/Source/Editor/PropertyEditor/Private/PropertyEditorHelpers.cpp b/Engine/Source/Editor/PropertyEditor/Private/PropertyEditorHelpers.cpp
--- a/Engine/Source/Editor/PropertyEditor/Private/PropertyEditorHelpers.cpp
+++ b/Engine/Source/Editor/PropertyEditor/Private/PropertyEditorHelpers.cpp
@@ -172,6 +172,22 @@ TSharedRef<SWidget> SPropertyValueWidget::ConstructPropertyEditorWidget( TShared
 
 			NumericWidget->GetDesiredWidth( MinDesiredWidth, MaxDesiredWidth );
 		}
+		else if (SPropertyEditorNumeric<int8>::Supports(PropertyEditorRef))
+		{
+			auto NumericWidget =
+				SAssignNew(PropertyWidget, SPropertyEditorNumeric<int8>, PropertyEditorRef)
+				.Font(FontStyle);
+
+			NumericWidget->GetDesiredWidth(MinDesiredWidth, MaxDesiredWidth);
+		}
+		else if (SPropertyEditorNumeric<int16>::Supports(PropertyEditorRef))
+		{
+			auto NumericWidget =
+				SAssignNew(PropertyWidget, SPropertyEditorNumeric<int16>, PropertyEditorRef)
+				.Font(FontStyle);
+
+			NumericWidget->GetDesiredWidth(MinDesiredWidth, MaxDesiredWidth);
+		}
 		else if ( SPropertyEditorNumeric<int32>::Supports( PropertyEditorRef ) )
 		{
 			auto NumericWidget = 
@@ -180,6 +196,14 @@ TSharedRef<SWidget> SPropertyValueWidget::ConstructPropertyEditorWidget( TShared
 
 			NumericWidget->GetDesiredWidth( MinDesiredWidth, MaxDesiredWidth );
 		}
+		else if (SPropertyEditorNumeric<int64>::Supports(PropertyEditorRef))
+		{
+			auto NumericWidget =
+				SAssignNew(PropertyWidget, SPropertyEditorNumeric<int64>, PropertyEditorRef)
+				.Font(FontStyle);
+
+			NumericWidget->GetDesiredWidth(MinDesiredWidth, MaxDesiredWidth);
+		}
 		else if ( SPropertyEditorNumeric<uint8>::Supports( PropertyEditorRef ) )
 		{
 			auto NumericWidget = 
@@ -188,6 +212,30 @@ TSharedRef<SWidget> SPropertyValueWidget::ConstructPropertyEditorWidget( TShared
 
 			NumericWidget->GetDesiredWidth( MinDesiredWidth, MaxDesiredWidth );
 		}
+		else if (SPropertyEditorNumeric<uint16>::Supports(PropertyEditorRef))
+		{
+			auto NumericWidget =
+				SAssignNew(PropertyWidget, SPropertyEditorNumeric<uint16>, PropertyEditorRef)
+				.Font(FontStyle);
+
+			NumericWidget->GetDesiredWidth(MinDesiredWidth, MaxDesiredWidth);
+		}
+		else if (SPropertyEditorNumeric<uint32>::Supports(PropertyEditorRef))
+		{
+			auto NumericWidget =
+				SAssignNew(PropertyWidget, SPropertyEditorNumeric<uint32>, PropertyEditorRef)
+				.Font(FontStyle);
+
+			NumericWidget->GetDesiredWidth(MinDesiredWidth, MaxDesiredWidth);
+		}
+		else if (SPropertyEditorNumeric<uint64>::Supports(PropertyEditorRef))
+		{
+			auto NumericWidget =
+				SAssignNew(PropertyWidget, SPropertyEditorNumeric<uint64>, PropertyEditorRef)
+				.Font(FontStyle);
+
+			NumericWidget->GetDesiredWidth(MinDesiredWidth, MaxDesiredWidth);
+		}
 		else if ( SPropertyEditorCombo::Supports( PropertyEditorRef ) )
 		{
 			TSharedRef<SPropertyEditorCombo> ComboWidget = 
diff --git a/Engine/Source/Editor/PropertyEditor/Private/PropertyHandleImpl.cpp b/Engine/Source/Editor/PropertyEditor/Private/PropertyHandleImpl.cpp
--- a/Engine/Source/Editor/PropertyEditor/Private/PropertyHandleImpl.cpp
+++ b/Engine/Source/Editor/PropertyEditor/Private/PropertyHandleImpl.cpp
@@ -885,7 +885,7 @@ bool FPropertyValueImpl::IsPropertyTypeOf( UClass* ClassType ) const
 	return false;
 }
 
-template< typename Type >
+template< typename Type>
 static Type ClampValueFromMetaData(Type InValue, FPropertyNode& InPropertyNode )
 {
 	UProperty* Property = InPropertyNode.GetProperty();
@@ -911,65 +911,52 @@ static Type ClampValueFromMetaData(Type InValue, FPropertyNode& InPropertyNode )
 			TTypeFromString<Type>::FromString(MaxValue, *MaxString);
 			RetVal = FMath::Min<Type>(MaxValue, RetVal);
 		}
-
-		const bool bIsInteger = Property->IsA(UIntProperty::StaticClass());
-		const bool bIsNonEnumByte = (Property->IsA(UByteProperty::StaticClass()) && Cast<const UByteProperty>(Property)->Enum == NULL);
-
-		if(bIsInteger || bIsNonEnumByte)
-		{
-			//if there is "Multiple" meta data, the selected number is a multiple
-			const FString& MultipleString = Property->GetMetaData(TEXT("Multiple"));
-			if(MultipleString.Len())
-			{
-				check(MultipleString.IsNumeric());
-				int32 MultipleValue = FCString::Atoi(*MultipleString);
-				if(MultipleValue!=0)
-				{
-					RetVal -= int32(RetVal)%MultipleValue;
-				}
-			}
-
-			//enforce array bounds
-			const FString& ArrayClampString = Property->GetMetaData(TEXT("ArrayClamp"));
-			if(ArrayClampString.Len())
-			{
-				FObjectPropertyNode* ObjectPropertyNode = InPropertyNode.FindObjectItemParent();
-				if(ObjectPropertyNode && ObjectPropertyNode->GetNumObjects() == 1)
-				{
-					int32 LastValidIndex = GetArrayPropertyLastValidIndex(ObjectPropertyNode, ArrayClampString);
-					RetVal = FMath::Clamp<int32>(RetVal, 0, LastValidIndex);
-				}
-				else
-				{
-					UE_LOG(LogPropertyNode, Warning, TEXT("Array Clamping isn't supported in multi-select (Param Name: %s)"), *Property->GetName());
-				}
-			}
-		}
 	}
 
 	return RetVal;
 }
 
-float FPropertyValueImpl::ClampFloatValueFromMetaData( float InValue ) const
+template <typename Type>
+static Type ClampIntegerValueFromMetaData(Type InValue, FPropertyNode& InPropertyNode )
 {
-	if( PropertyNode.IsValid() )
+	Type RetVal = ClampValueFromMetaData<Type>( InValue, InPropertyNode );
+
+	UProperty* Property = InPropertyNode.GetProperty();
+
+
+	//if there is "Multiple" meta data, the selected number is a multiple
+	const FString& MultipleString = Property->GetMetaData(TEXT("Multiple"));
+	if (MultipleString.Len())
 	{
-		return ClampValueFromMetaData<float>( InValue, *PropertyNode.Pin() );
+		check(MultipleString.IsNumeric());
+		Type MultipleValue;
+		TTypeFromString<Type>::FromString(MultipleValue, *MultipleString);
+		if (MultipleValue != 0)
+		{
+			RetVal -= Type(RetVal) % MultipleValue;
+		}
 	}
 
-	return InValue;
-}
-
-int32 FPropertyValueImpl::ClampIntValueFromMetaData( int32 InValue ) const
-{
-	if( PropertyNode.IsValid() )
+	//enforce array bounds
+	const FString& ArrayClampString = Property->GetMetaData(TEXT("ArrayClamp"));
+	if (ArrayClampString.Len())
 	{
-		return ClampValueFromMetaData<int32>( InValue, *PropertyNode.Pin() );
+		FObjectPropertyNode* ObjectPropertyNode = InPropertyNode.FindObjectItemParent();
+		if (ObjectPropertyNode && ObjectPropertyNode->GetNumObjects() == 1)
+		{
+			Type LastValidIndex = GetArrayPropertyLastValidIndex(ObjectPropertyNode, ArrayClampString);
+			RetVal = FMath::Clamp<Type>(RetVal, 0, LastValidIndex);
+		}
+		else
+		{
+			UE_LOG(LogPropertyNode, Warning, TEXT("Array Clamping isn't supported in multi-select (Param Name: %s)"), *Property->GetName());
+		}
 	}
 
-	return InValue;
+	return RetVal;
 }
 
+
 int32 FPropertyValueImpl::GetNumChildren() const
 {
 	TSharedPtr<FPropertyNode> PropertyNodePin = PropertyNode.Pin();
@@ -1462,11 +1449,17 @@ FText FPropertyValueImpl::GetDisplayName() const
 	}
 
 IMPLEMENT_PROPERTY_ACCESSOR( bool )
+IMPLEMENT_PROPERTY_ACCESSOR( int8 )
+IMPLEMENT_PROPERTY_ACCESSOR( int16 )
 IMPLEMENT_PROPERTY_ACCESSOR( int32 )
+IMPLEMENT_PROPERTY_ACCESSOR( int64 )
+IMPLEMENT_PROPERTY_ACCESSOR( uint8 )
+IMPLEMENT_PROPERTY_ACCESSOR( uint16 )
+IMPLEMENT_PROPERTY_ACCESSOR( uint32 )
+IMPLEMENT_PROPERTY_ACCESSOR( uint64 )
 IMPLEMENT_PROPERTY_ACCESSOR( float )
 IMPLEMENT_PROPERTY_ACCESSOR( FString )
 IMPLEMENT_PROPERTY_ACCESSOR( FName )
-IMPLEMENT_PROPERTY_ACCESSOR( uint8 )
 IMPLEMENT_PROPERTY_ACCESSOR( FVector )
 IMPLEMENT_PROPERTY_ACCESSOR( FVector2D )
 IMPLEMENT_PROPERTY_ACCESSOR( FVector4 )
@@ -2075,54 +2068,201 @@ bool FPropertyHandleInt::Supports( TSharedRef<FPropertyNode> PropertyNode )
 		return false;
 	}
 
-	const bool bIsInteger = Property->IsA(UIntProperty::StaticClass());
+	const bool bIsInteger = 
+			Property->IsA<UInt8Property>()
+		||	Property->IsA<UInt16Property>()
+		||  Property->IsA<UIntProperty>()
+		||	Property->IsA<UInt64Property>()
+		||  Property->IsA<UUInt16Property>()
+		||  Property->IsA<UUInt32Property>()
+		||  Property->IsA<UUInt64Property>();
+
 	// The value is an integer
 	return bIsInteger;
 }
 
-FPropertyAccess::Result FPropertyHandleInt::GetValue( int32& OutValue ) const
+template <typename PropertyClass, typename ValueType>
+ValueType GetIntegerValue(void* PropValue, FPropertyValueImpl& Implementation )
+{
+	check( Implementation.IsPropertyTypeOf( PropertyClass::StaticClass() ) );
+	return Implementation.GetPropertyValue<PropertyClass>(PropValue);
+}
+
+FPropertyAccess::Result FPropertyHandleInt::GetValue(int8& OutValue) const
+{
+	void* PropValue = NULL;
+	FPropertyAccess::Result Res = Implementation->GetValueData(PropValue);
+
+	if (Res == FPropertyAccess::Success)
+	{
+		OutValue = GetIntegerValue<UInt8Property, int8>(PropValue, *Implementation);
+	}
+
+	return Res;
+}
+
+FPropertyAccess::Result FPropertyHandleInt::GetValue(int16& OutValue) const
+{
+	void* PropValue = NULL;
+	FPropertyAccess::Result Res = Implementation->GetValueData(PropValue);
+
+	if (Res == FPropertyAccess::Success)
+	{
+		OutValue = GetIntegerValue<UInt16Property, int16>(PropValue, *Implementation);
+	}
+
+	return Res;
+}
+
+FPropertyAccess::Result FPropertyHandleInt::GetValue(int32& OutValue) const
+{
+	void* PropValue = NULL;
+	FPropertyAccess::Result Res = Implementation->GetValueData(PropValue);
+
+	if (Res == FPropertyAccess::Success)
+	{
+		OutValue = GetIntegerValue<UIntProperty, int32>(PropValue, *Implementation);
+	}
+
+	return Res;
+}
+FPropertyAccess::Result FPropertyHandleInt::GetValue( int64& OutValue ) const
 {
 	void* PropValue = NULL;
 	FPropertyAccess::Result Res = Implementation->GetValueData( PropValue );
 
 	if( Res == FPropertyAccess::Success )
 	{
-		if( Implementation->IsPropertyTypeOf( UInt8Property::StaticClass() ) )
-		{
-			OutValue = Implementation->GetPropertyValue<UInt8Property>(PropValue);
-		}
-		else if( Implementation->IsPropertyTypeOf( UInt16Property::StaticClass() ) )
-		{
-			OutValue = Implementation->GetPropertyValue<UInt16Property>(PropValue);
-		}
-		else if( Implementation->IsPropertyTypeOf( UUInt16Property::StaticClass() ) )
-		{
-			OutValue = Implementation->GetPropertyValue<UUInt16Property>(PropValue);
-		}
-		else
-		{
-			//@todo need to handle all integer types. Also, this should be a template so we don't need a bunch of if statements
-			check(Implementation->IsPropertyTypeOf( UIntProperty::StaticClass() ));
-			OutValue = Implementation->GetPropertyValue<UIntProperty>(PropValue);
-		}
+		OutValue = GetIntegerValue<UInt64Property,int64>(PropValue, *Implementation);
 	}
 
 	return Res;
 }
 
+FPropertyAccess::Result FPropertyHandleInt::GetValue(uint16& OutValue) const
+{
+	void* PropValue = NULL;
+	FPropertyAccess::Result Res = Implementation->GetValueData(PropValue);
+
+	if (Res == FPropertyAccess::Success)
+	{
+		OutValue = GetIntegerValue<UUInt16Property, uint16>(PropValue, *Implementation);
+	}
+
+	return Res;
+}
+
+FPropertyAccess::Result FPropertyHandleInt::GetValue(uint32& OutValue) const
+{
+	void* PropValue = NULL;
+	FPropertyAccess::Result Res = Implementation->GetValueData(PropValue);
+
+	if (Res == FPropertyAccess::Success)
+	{
+		OutValue = GetIntegerValue<UUInt32Property, uint32>(PropValue, *Implementation);
+	}
+
+	return Res;
+}
+
+FPropertyAccess::Result FPropertyHandleInt::GetValue(uint64& OutValue) const
+{
+	void* PropValue = NULL;
+	FPropertyAccess::Result Res = Implementation->GetValueData(PropValue);
+
+	if (Res == FPropertyAccess::Success)
+	{
+		OutValue = GetIntegerValue<UUInt64Property, uint64>(PropValue, *Implementation);
+	}
+
+	return Res;
+}
+
+FPropertyAccess::Result FPropertyHandleInt::SetValue(const int8& NewValue, EPropertyValueSetFlags::Type Flags)
+{
+	FPropertyAccess::Result Res;
+	// Clamp the value from any meta data ranges stored on the property value
+	int8 FinalValue = ClampIntegerValueFromMetaData<int8>( NewValue, *Implementation->GetPropertyNode() );
+
+	const FString ValueStr = LexicalConversion::ToString(FinalValue);
+	Res = Implementation->ImportText(ValueStr, Flags);
+
+	return Res;
+}
+
+
+FPropertyAccess::Result FPropertyHandleInt::SetValue(const int16& NewValue, EPropertyValueSetFlags::Type Flags)
+{
+	FPropertyAccess::Result Res;
+	// Clamp the value from any meta data ranges stored on the property value
+	int16 FinalValue = ClampIntegerValueFromMetaData<int16>(NewValue, *Implementation->GetPropertyNode());
+
+	const FString ValueStr = LexicalConversion::ToString(FinalValue);
+	Res = Implementation->ImportText(ValueStr, Flags);
+
+	return Res;
+}
+
+
 FPropertyAccess::Result FPropertyHandleInt::SetValue( const int32& NewValue, EPropertyValueSetFlags::Type Flags )
 {
 	FPropertyAccess::Result Res;
 	// Clamp the value from any meta data ranges stored on the property value
-	int32 FinalValue = Implementation->ClampIntValueFromMetaData( NewValue );
+	int32 FinalValue = ClampIntegerValueFromMetaData<int32>( NewValue, *Implementation->GetPropertyNode() );
 
-	//@todo this doesn't handle all int types
-	const FString ValueStr = FString::Printf( TEXT("%i"), FinalValue );
+	const FString ValueStr = LexicalConversion::ToString(FinalValue);
 	Res = Implementation->ImportText( ValueStr, Flags );
 
 	return Res;
 }
 
+FPropertyAccess::Result FPropertyHandleInt::SetValue(const int64& NewValue, EPropertyValueSetFlags::Type Flags)
+{
+	FPropertyAccess::Result Res;
+
+	// Clamp the value from any meta data ranges stored on the property value
+	int64 FinalValue = ClampIntegerValueFromMetaData<int64>(NewValue, *Implementation->GetPropertyNode());
+
+	const FString ValueStr = LexicalConversion::ToString(FinalValue);
+	Res = Implementation->ImportText(ValueStr, Flags);
+	return Res;
+}
+
+FPropertyAccess::Result FPropertyHandleInt::SetValue(const uint16& NewValue, EPropertyValueSetFlags::Type Flags)
+{
+	FPropertyAccess::Result Res;
+	// Clamp the value from any meta data ranges stored on the property value
+	uint16 FinalValue = ClampIntegerValueFromMetaData<uint16>(NewValue, *Implementation->GetPropertyNode());
+
+	const FString ValueStr = LexicalConversion::ToString(FinalValue);
+	Res = Implementation->ImportText(ValueStr, Flags);
+
+	return Res;
+}
+
+
+FPropertyAccess::Result FPropertyHandleInt::SetValue(const uint32& NewValue, EPropertyValueSetFlags::Type Flags)
+{
+	FPropertyAccess::Result Res;
+	// Clamp the value from any meta data ranges stored on the property value
+	uint32 FinalValue = ClampIntegerValueFromMetaData<uint32>(NewValue, *Implementation->GetPropertyNode());
+
+	const FString ValueStr = LexicalConversion::ToString(FinalValue);
+	Res = Implementation->ImportText(ValueStr, Flags);
+
+	return Res;
+}
+
+FPropertyAccess::Result FPropertyHandleInt::SetValue(const uint64& NewValue, EPropertyValueSetFlags::Type Flags)
+{
+	FPropertyAccess::Result Res;
+	// Clamp the value from any meta data ranges stored on the property value
+	uint64 FinalValue = ClampIntegerValueFromMetaData<uint64>(NewValue, *Implementation->GetPropertyNode());
+
+	const FString ValueStr = LexicalConversion::ToString(FinalValue);
+	Res = Implementation->ImportText(ValueStr, Flags);
+	return Res;
+}
 
 // float
 bool FPropertyHandleFloat::Supports( TSharedRef<FPropertyNode> PropertyNode )
@@ -2154,7 +2294,7 @@ FPropertyAccess::Result FPropertyHandleFloat::SetValue( const float& NewValue, E
 {
 	FPropertyAccess::Result Res;
 	// Clamp the value from any meta data ranges stored on the property value
-	float FinalValue = Implementation->ClampFloatValueFromMetaData( NewValue );
+	float FinalValue = ClampValueFromMetaData<float>( NewValue, *Implementation->GetPropertyNode() );
 
 	const FString ValueStr = FString::Printf( TEXT("%f"), FinalValue );
 	Res = Implementation->ImportText( ValueStr, Flags );
diff --git a/Engine/Source/Editor/PropertyEditor/Private/PropertyHandleImpl.h b/Engine/Source/Editor/PropertyEditor/Private/PropertyHandleImpl.h
--- a/Engine/Source/Editor/PropertyEditor/Private/PropertyHandleImpl.h
+++ b/Engine/Source/Editor/PropertyEditor/Private/PropertyHandleImpl.h
@@ -192,19 +192,7 @@ public:
 	 * @param true if the property is a ClassType
 	 */
 	bool IsPropertyTypeOf( UClass* ClassType ) const ;
-
-
-	/**
-	 * @return Clamps the value as a float from meta data stored on the property
-	 */
-	float ClampFloatValueFromMetaData( float InValue ) const ;
-
-	/**
-	 * @return Clamps the value as an integer from meta data stored on the property
-	 */
-	int32 ClampIntValueFromMetaData( int32 InValue ) const ;
 	
-
 	/**
 	 * @return The property node used by this value
 	 */
@@ -352,11 +340,17 @@ public:
 
 	/** IPropertyHandle interface */
 	DECLARE_PROPERTY_ACCESSOR( bool )
+	DECLARE_PROPERTY_ACCESSOR( int8 )
+	DECLARE_PROPERTY_ACCESSOR( int16)
 	DECLARE_PROPERTY_ACCESSOR( int32 )
+	DECLARE_PROPERTY_ACCESSOR( int64 )
+	DECLARE_PROPERTY_ACCESSOR( uint8 )
+	DECLARE_PROPERTY_ACCESSOR( uint16 )
+	DECLARE_PROPERTY_ACCESSOR( uint32 )
+	DECLARE_PROPERTY_ACCESSOR( uint64 )
 	DECLARE_PROPERTY_ACCESSOR( float )
 	DECLARE_PROPERTY_ACCESSOR( FString )
 	DECLARE_PROPERTY_ACCESSOR( FName )
-	DECLARE_PROPERTY_ACCESSOR( uint8 )
 	DECLARE_PROPERTY_ACCESSOR( FVector )
 	DECLARE_PROPERTY_ACCESSOR( FVector2D )
 	DECLARE_PROPERTY_ACCESSOR( FVector4 )
@@ -426,8 +420,27 @@ class FPropertyHandleInt : public FPropertyHandleBase
 public:
 	FPropertyHandleInt( TSharedRef<FPropertyNode> PropertyNode, FNotifyHook* NotifyHook, TSharedPtr<IPropertyUtilities> PropertyUtilities );
 	static bool Supports( TSharedRef<FPropertyNode> PropertyNode );
-	virtual FPropertyAccess::Result GetValue( int32& OutValue ) const override;
-	virtual FPropertyAccess::Result SetValue( const int32& InValue, EPropertyValueSetFlags::Type Flags = EPropertyValueSetFlags::DefaultFlags ) override;
+
+	virtual FPropertyAccess::Result GetValue(int8& OutValue ) const override;
+	virtual FPropertyAccess::Result SetValue(const int8& InValue, EPropertyValueSetFlags::Type Flags = EPropertyValueSetFlags::DefaultFlags ) override;
+
+	virtual FPropertyAccess::Result GetValue(int16& OutValue) const override;
+	virtual FPropertyAccess::Result SetValue(const int16& InValue, EPropertyValueSetFlags::Type Flags = EPropertyValueSetFlags::DefaultFlags) override;
+
+	virtual FPropertyAccess::Result GetValue(int32& OutValue) const override;
+	virtual FPropertyAccess::Result SetValue(const int32& InValue, EPropertyValueSetFlags::Type Flags = EPropertyValueSetFlags::DefaultFlags) override;
+
+	virtual FPropertyAccess::Result GetValue(int64& OutValue) const override;
+	virtual FPropertyAccess::Result SetValue(const int64& InValue, EPropertyValueSetFlags::Type Flags = EPropertyValueSetFlags::DefaultFlags) override;
+
+	virtual FPropertyAccess::Result GetValue(uint16& OutValue) const override;
+	virtual FPropertyAccess::Result SetValue(const uint16& InValue, EPropertyValueSetFlags::Type Flags = EPropertyValueSetFlags::DefaultFlags) override;
+
+	virtual FPropertyAccess::Result GetValue(uint32& OutValue) const override;
+	virtual FPropertyAccess::Result SetValue(const uint32& InValue, EPropertyValueSetFlags::Type Flags = EPropertyValueSetFlags::DefaultFlags) override;
+
+	virtual FPropertyAccess::Result GetValue(uint64& OutValue) const override;
+	virtual FPropertyAccess::Result SetValue(const uint64& InValue, EPropertyValueSetFlags::Type Flags = EPropertyValueSetFlags::DefaultFlags) override;
 };
 
 class FPropertyHandleFloat : public FPropertyHandleBase
diff --git a/Engine/Source/Editor/PropertyEditor/Private/PropertyNode.cpp b/Engine/Source/Editor/PropertyEditor/Private/PropertyNode.cpp
--- a/Engine/Source/Editor/PropertyEditor/Private/PropertyNode.cpp
+++ b/Engine/Source/Editor/PropertyEditor/Private/PropertyNode.cpp
@@ -380,7 +380,7 @@ FPropertyNode::DataValidationResult FPropertyNode::EnsureDataIsValid()
 			//make sure we got the addresses correctly
 			if (!bSuccess)
 			{
-				UE_LOG( LogPropertyNode, Log, TEXT("Object is invalid %s"), *Property->GetName() );
+				UE_LOG( LogPropertyNode, Verbose, TEXT("Object is invalid %s"), *Property->GetName() );
 				return ObjectInvalid;
 			}
 
@@ -392,7 +392,7 @@ FPropertyNode::DataValidationResult FPropertyNode::EnsureDataIsValid()
 				//make sure the data still exists
 				if (Addr==NULL)
 				{
-					UE_LOG( LogPropertyNode, Log, TEXT("Object is invalid %s"), *Property->GetName() );
+					UE_LOG( LogPropertyNode, Verbose, TEXT("Object is invalid %s"), *Property->GetName() );
 					return ObjectInvalid;
 				}
 
diff --git a/Engine/Source/Editor/PropertyEditor/Private/UserInterface/PropertyEditor/SPropertyEditorNumeric.h b/Engine/Source/Editor/PropertyEditor/Private/UserInterface/PropertyEditor/SPropertyEditorNumeric.h
--- a/Engine/Source/Editor/PropertyEditor/Private/UserInterface/PropertyEditor/SPropertyEditorNumeric.h
+++ b/Engine/Source/Editor/PropertyEditor/Private/UserInterface/PropertyEditor/SPropertyEditorNumeric.h
@@ -191,18 +191,54 @@ private:
 		static bool Match(const UProperty* InProperty) { return InProperty->IsA(UFloatProperty::StaticClass()); } 
 	};
 
-	template<typename U> 
-	struct TTypeToProperty<int32, U> 
-	{	
-		static bool Match(const UProperty* InProperty) { return InProperty->IsA(UIntProperty::StaticClass()); } 
+	template<typename U>
+	struct TTypeToProperty<int8, U>
+	{
+		static bool Match(const UProperty* InProperty) { return InProperty->IsA(UInt8Property::StaticClass()); }
 	};
 
-	template<typename U> 
-	struct TTypeToProperty<uint8, U> 
-	{	
-		static bool Match(const UProperty* InProperty) { return (InProperty->IsA(UByteProperty::StaticClass()) && Cast<const UByteProperty>(InProperty)->Enum == NULL); } 
+	template<typename U>
+	struct TTypeToProperty<int16, U>
+	{
+		static bool Match(const UProperty* InProperty) { return InProperty->IsA(UInt16Property::StaticClass()); }
 	};
-				
+
+	template<typename U>
+	struct TTypeToProperty<int32, U>
+	{
+		static bool Match(const UProperty* InProperty) { return InProperty->IsA(UIntProperty::StaticClass()); }
+	};
+
+	template<typename U>
+	struct TTypeToProperty<int64, U>
+	{
+		static bool Match(const UProperty* InProperty) { return InProperty->IsA(UInt64Property::StaticClass()); }
+	};
+
+	template<typename U>
+	struct TTypeToProperty<uint8, U>
+	{
+		static bool Match(const UProperty* InProperty) { return (InProperty->IsA(UByteProperty::StaticClass()) && Cast<const UByteProperty>(InProperty)->Enum == NULL); }
+	};
+
+	template<typename U>
+	struct TTypeToProperty<uint16, U>
+	{
+		static bool Match(const UProperty* InProperty) { return InProperty->IsA(UUInt16Property::StaticClass()); }
+	};
+
+	template<typename U>
+	struct TTypeToProperty<uint32, U>
+	{
+		static bool Match(const UProperty* InProperty) { return InProperty->IsA(UUInt32Property::StaticClass()); }
+	};
+
+	template<typename U>
+	struct TTypeToProperty<uint64, U>
+	{
+		static bool Match(const UProperty* InProperty) { return InProperty->IsA(UUInt64Property::StaticClass()); }
+	};
+		
 	
 	/** @return The value or unset if properties with multiple values are viewed */
 	TOptional<NumericType> OnGetValue() const
diff --git a/Engine/Source/Editor/PropertyEditor/Private/UserInterface/PropertyEditor/SPropertyEditorTableRow.cpp b/Engine/Source/Editor/PropertyEditor/Private/UserInterface/PropertyEditor/SPropertyEditorTableRow.cpp
--- a/Engine/Source/Editor/PropertyEditor/Private/UserInterface/PropertyEditor/SPropertyEditorTableRow.cpp
+++ b/Engine/Source/Editor/PropertyEditor/Private/UserInterface/PropertyEditor/SPropertyEditorTableRow.cpp
@@ -242,14 +242,38 @@ TSharedRef<SWidget> SPropertyEditorTableRow::ConstructPropertyEditorWidget()
 		{
 			PropertyWidget = SNew( SPropertyEditorNumeric<float>, PropertyEditorRef );
 		}
-		if ( SPropertyEditorNumeric<int32>::Supports(PropertyEditorRef) )
+		else if (SPropertyEditorNumeric<int8>::Supports(PropertyEditorRef))
 		{
-			PropertyWidget = SNew( SPropertyEditorNumeric<int32>, PropertyEditorRef );
+			PropertyWidget = SNew(SPropertyEditorNumeric<int8>, PropertyEditorRef);
 		}
-		if ( SPropertyEditorNumeric<uint8>::Supports(PropertyEditorRef) )
+		else if (SPropertyEditorNumeric<int16>::Supports(PropertyEditorRef))
+		{
+			PropertyWidget = SNew(SPropertyEditorNumeric<int16>, PropertyEditorRef);
+		}
+		else if (SPropertyEditorNumeric<int32>::Supports(PropertyEditorRef))
+		{
+			PropertyWidget = SNew(SPropertyEditorNumeric<int32>, PropertyEditorRef);
+		}
+		else if ( SPropertyEditorNumeric<int64>::Supports(PropertyEditorRef) )
+		{
+			PropertyWidget = SNew( SPropertyEditorNumeric<int64>, PropertyEditorRef );
+		}
+		else if ( SPropertyEditorNumeric<uint8>::Supports(PropertyEditorRef) )
 		{
 			PropertyWidget = SNew( SPropertyEditorNumeric<uint8>, PropertyEditorRef );
 		}
+		else if (SPropertyEditorNumeric<uint16>::Supports(PropertyEditorRef))
+		{
+			PropertyWidget = SNew(SPropertyEditorNumeric<uint16>, PropertyEditorRef);
+		}
+		else if (SPropertyEditorNumeric<uint32>::Supports(PropertyEditorRef))
+		{
+			PropertyWidget = SNew(SPropertyEditorNumeric<uint32>, PropertyEditorRef);
+		}
+		else if (SPropertyEditorNumeric<uint64>::Supports(PropertyEditorRef))
+		{
+			PropertyWidget = SNew(SPropertyEditorNumeric<uint64>, PropertyEditorRef);
+		}
 		else if ( SPropertyEditorArray::Supports(PropertyEditorRef) )
 		{
 			PropertyWidget = SNew( SPropertyEditorArray, PropertyEditorRef );
diff --git a/Engine/Source/Editor/PropertyEditor/Private/UserInterface/PropertyTable/TextPropertyTableCellPresenter.cpp b/Engine/Source/Editor/PropertyEditor/Private/UserInterface/PropertyTable/TextPropertyTableCellPresenter.cpp
--- a/Engine/Source/Editor/PropertyEditor/Private/UserInterface/PropertyTable/TextPropertyTableCellPresenter.cpp
+++ b/Engine/Source/Editor/PropertyEditor/Private/UserInterface/PropertyTable/TextPropertyTableCellPresenter.cpp
@@ -109,16 +109,46 @@ TSharedRef<SWidget> FTextPropertyTableCellPresenter::ConstructEditModeCellWidget
 			PropertyWidget = SNew( SPropertyEditorNumeric<float>, PropertyEditor )
 				.Font( Font );
 		}
-		if ( SPropertyEditorNumeric<int32>::Supports(PropertyEditor) )
+		else if (SPropertyEditorNumeric<int8>::Supports(PropertyEditor))
+		{
+			PropertyWidget = SNew(SPropertyEditorNumeric<int8>, PropertyEditor)
+				.Font(Font);
+		}
+		else if (SPropertyEditorNumeric<int16>::Supports(PropertyEditor))
+		{
+			PropertyWidget = SNew(SPropertyEditorNumeric<int16>, PropertyEditor)
+				.Font(Font);
+		}
+		else if ( SPropertyEditorNumeric<int32>::Supports(PropertyEditor) )
 		{
 			PropertyWidget = SNew( SPropertyEditorNumeric<int32>, PropertyEditor )
 				.Font( Font );
 		}
-		if ( SPropertyEditorNumeric<uint8>::Supports(PropertyEditor) )
+		else if (SPropertyEditorNumeric<int64>::Supports(PropertyEditor))
+		{
+			PropertyWidget = SNew(SPropertyEditorNumeric<int64>, PropertyEditor)
+				.Font(Font);
+		}
+		else if ( SPropertyEditorNumeric<uint8>::Supports(PropertyEditor) )
 		{
 			PropertyWidget = SNew( SPropertyEditorNumeric<uint8>, PropertyEditor )
 				.Font( Font );
 		}
+		else if (SPropertyEditorNumeric<uint16>::Supports(PropertyEditor))
+		{
+			PropertyWidget = SNew(SPropertyEditorNumeric<uint16>, PropertyEditor)
+				.Font(Font);
+		}
+		else if (SPropertyEditorNumeric<uint32>::Supports(PropertyEditor))
+		{
+			PropertyWidget = SNew(SPropertyEditorNumeric<uint32>, PropertyEditor)
+				.Font(Font);
+		}
+		else if (SPropertyEditorNumeric<uint64>::Supports(PropertyEditor))
+		{
+			PropertyWidget = SNew(SPropertyEditorNumeric<uint64>, PropertyEditor)
+				.Font(Font);
+		}
 		else if ( SPropertyEditorCombo::Supports(PropertyEditor) )
 		{
 			PropertyWidget = SNew( SPropertyEditorCombo, PropertyEditor )
diff --git a/Engine/Source/Editor/PropertyEditor/Public/PropertyHandle.h b/Engine/Source/Editor/PropertyEditor/Public/PropertyHandle.h
--- a/Engine/Source/Editor/PropertyEditor/Public/PropertyHandle.h
+++ b/Engine/Source/Editor/PropertyEditor/Public/PropertyHandle.h
@@ -167,10 +167,16 @@ public:
 	 * @param OutValue	The value that will be set if successful
 	 * @return The result of attempting to get the value
 	 */
-	virtual FPropertyAccess::Result GetValue( int32& OutValue ) const = 0;
 	virtual FPropertyAccess::Result GetValue( float& OutValue ) const = 0;
 	virtual FPropertyAccess::Result GetValue( bool& OutValue ) const = 0;
+	virtual FPropertyAccess::Result GetValue( int8& OutValue ) const = 0;
+	virtual FPropertyAccess::Result GetValue( int16& OutValue ) const = 0;
+	virtual FPropertyAccess::Result GetValue( int32& OutValue ) const = 0;
+	virtual FPropertyAccess::Result GetValue( int64& OutValue ) const = 0;
 	virtual FPropertyAccess::Result GetValue( uint8& OutValue ) const = 0;
+	virtual FPropertyAccess::Result GetValue( uint16& OutValue ) const = 0;
+	virtual FPropertyAccess::Result GetValue( uint32& OutValue ) const = 0;
+	virtual FPropertyAccess::Result GetValue( uint64& OutValue ) const = 0;
 	virtual FPropertyAccess::Result GetValue( FString& OutValue ) const = 0;
 	virtual FPropertyAccess::Result GetValue( FName& OutValue ) const = 0;
 	virtual FPropertyAccess::Result GetValue( FVector& OutValue ) const = 0;
@@ -188,10 +194,16 @@ public:
 	 * @param InValue	The value to set
 	 * @return The result of attempting to set the value
 	 */
-	virtual FPropertyAccess::Result SetValue( const int32& InValue, EPropertyValueSetFlags::Type Flags = EPropertyValueSetFlags::DefaultFlags ) = 0;
 	virtual FPropertyAccess::Result SetValue( const float& InValue, EPropertyValueSetFlags::Type Flags = EPropertyValueSetFlags::DefaultFlags ) = 0;
 	virtual FPropertyAccess::Result SetValue( const bool& InValue,  EPropertyValueSetFlags::Type Flags = EPropertyValueSetFlags::DefaultFlags ) = 0;
+	virtual FPropertyAccess::Result SetValue( const int8& InValue, EPropertyValueSetFlags::Type Flags = EPropertyValueSetFlags::DefaultFlags ) = 0;
+	virtual FPropertyAccess::Result SetValue( const int16& InValue, EPropertyValueSetFlags::Type Flags = EPropertyValueSetFlags::DefaultFlags ) = 0;
+	virtual FPropertyAccess::Result SetValue( const int32& InValue, EPropertyValueSetFlags::Type Flags = EPropertyValueSetFlags::DefaultFlags ) = 0;
+	virtual FPropertyAccess::Result SetValue( const int64& InValue, EPropertyValueSetFlags::Type Flags = EPropertyValueSetFlags::DefaultFlags ) = 0;
 	virtual FPropertyAccess::Result SetValue( const uint8& InValue,  EPropertyValueSetFlags::Type Flags = EPropertyValueSetFlags::DefaultFlags ) = 0;
+	virtual FPropertyAccess::Result SetValue( const uint16& InValue,  EPropertyValueSetFlags::Type Flags = EPropertyValueSetFlags::DefaultFlags ) = 0;
+	virtual FPropertyAccess::Result SetValue( const uint32& InValue,  EPropertyValueSetFlags::Type Flags = EPropertyValueSetFlags::DefaultFlags ) = 0;
+	virtual FPropertyAccess::Result SetValue( const uint64& InValue,  EPropertyValueSetFlags::Type Flags = EPropertyValueSetFlags::DefaultFlags ) = 0;
 	virtual FPropertyAccess::Result SetValue( const FString& InValue,  EPropertyValueSetFlags::Type Flags = EPropertyValueSetFlags::DefaultFlags ) = 0;
 	virtual FPropertyAccess::Result SetValue( const FName& InValue, EPropertyValueSetFlags::Type Flags = EPropertyValueSetFlags::DefaultFlags ) = 0;
 	virtual FPropertyAccess::Result SetValue( const FVector& InValue,  EPropertyValueSetFlags::Type Flags = EPropertyValueSetFlags::DefaultFlags ) = 0;
diff --git a/Engine/Source/Editor/UnrealEd/Classes/Editor/PropertyEditorTestObject.h b/Engine/Source/Editor/UnrealEd/Classes/Editor/PropertyEditorTestObject.h
--- a/Engine/Source/Editor/UnrealEd/Classes/Editor/PropertyEditorTestObject.h
+++ b/Engine/Source/Editor/UnrealEd/Classes/Editor/PropertyEditorTestObject.h
@@ -92,14 +92,32 @@ class UPropertyEditorTestObject : public UObject
 {
     GENERATED_UCLASS_BODY()
 
-	// Integer
-	UPROPERTY(EditAnywhere, Category=BasicProperties)
-	int32 IntProperty32;
 
-	// Byte
-	UPROPERTY(EditAnywhere, Category=BasicProperties)
+
+	UPROPERTY(EditAnywhere, Category = BasicProperties)
+	int8 Int8Property;
+
+	UPROPERTY(EditAnywhere, Category = BasicProperties)
+	int16 Int16roperty;
+
+	UPROPERTY(EditAnywhere, Category = BasicProperties)
+	int32 Int32Property;
+
+	UPROPERTY(EditAnywhere, Category = BasicProperties)
+	int64 Int64Property;
+
+	UPROPERTY(EditAnywhere, Category = BasicProperties)
 	uint8 ByteProperty;
 
+	UPROPERTY(EditAnywhere, Category = BasicProperties)
+	uint16 UnsignedInt16Property;
+
+	UPROPERTY(EditAnywhere, Category = BasicProperties)
+	uint32 UnsignedInt32Property;
+
+	UPROPERTY(EditAnywhere, Category = BasicProperties)
+	uint64 UnsignedInt64Property;
+
 	UPROPERTY(EditAnywhere, Category=BasicProperties)
 	float FloatProperty;
 
