Fixed "Multiple Values" being committed to property text blocks when they lose focus or when the user presses enter

diff --git a/Engine/Source/Editor/PropertyEditor/Private/Presentation/PropertyEditor/PropertyEditor.cpp b/Engine/Source/Editor/PropertyEditor/Private/Presentation/PropertyEditor/PropertyEditor.cpp
--- a/Engine/Source/Editor/PropertyEditor/Private/Presentation/PropertyEditor/PropertyEditor.cpp
+++ b/Engine/Source/Editor/PropertyEditor/Private/Presentation/PropertyEditor/PropertyEditor.cpp
@@ -17,6 +17,8 @@
 
 #define LOCTEXT_NAMESPACE "PropertyEditor"
 
+const FString FPropertyEditor::MultipleValuesDisplayName = NSLOCTEXT("PropertyEditor", "MultipleValues", "Multiple Values").ToString();
+
 TSharedRef< FPropertyEditor > FPropertyEditor::Create( const TSharedRef< class FPropertyNode >& InPropertyNode, const TSharedRef<class IPropertyUtilities >& InPropertyUtilities )
 {
 	return MakeShareable( new FPropertyEditor( InPropertyNode, InPropertyUtilities ) );
@@ -113,7 +115,7 @@ FString FPropertyEditor::GetValueAsString() const
 
 	if( PropertyHandle->GetValueAsFormattedString( Str ) == FPropertyAccess::MultipleValues )
 	{
-		Str = NSLOCTEXT("PropertyEditor", "MultipleValues", "Multiple Values").ToString();
+		Str = MultipleValuesDisplayName;
 	}
 
 	return Str;
@@ -125,7 +127,7 @@ FString FPropertyEditor::GetValueAsDisplayString() const
 
 	if( PropertyHandle->GetValueAsDisplayString( Str ) == FPropertyAccess::MultipleValues )
 	{
-		Str = NSLOCTEXT("PropertyEditor", "MultipleValues", "Multiple Values").ToString();
+		Str = MultipleValuesDisplayName;
 	}
 
 	return Str;
diff --git a/Engine/Source/Editor/PropertyEditor/Private/Presentation/PropertyEditor/PropertyEditor.h b/Engine/Source/Editor/PropertyEditor/Private/Presentation/PropertyEditor/PropertyEditor.h
--- a/Engine/Source/Editor/PropertyEditor/Private/Presentation/PropertyEditor/PropertyEditor.h
+++ b/Engine/Source/Editor/PropertyEditor/Private/Presentation/PropertyEditor/PropertyEditor.h
@@ -101,6 +101,7 @@ public:
 
 	static void SyncToObjectsInNode( const TWeakPtr< FPropertyNode >& WeakPropertyNode );
 
+	static const FString MultipleValuesDisplayName;
 private:
 	FPropertyEditor( const TSharedRef< class FPropertyNode >& InPropertyNode, const TSharedRef<class IPropertyUtilities >& InPropertyUtilities );
 
diff --git a/Engine/Source/Editor/PropertyEditor/Private/UserInterface/PropertyEditor/SPropertyEditorText.cpp b/Engine/Source/Editor/PropertyEditor/Private/UserInterface/PropertyEditor/SPropertyEditorText.cpp
--- a/Engine/Source/Editor/PropertyEditor/Private/UserInterface/PropertyEditor/SPropertyEditorText.cpp
+++ b/Engine/Source/Editor/PropertyEditor/Private/UserInterface/PropertyEditor/SPropertyEditorText.cpp
@@ -116,7 +116,11 @@ void SPropertyEditorText::OnTextCommitted( const FText& NewText, ETextCommit::Ty
 	const TSharedRef< FPropertyNode > PropertyNode = PropertyEditor->GetPropertyNode();
 	const TSharedRef< IPropertyHandle > PropertyHandle = PropertyEditor->GetPropertyHandle();
 
-	PropertyHandle->SetValueFromFormattedString( NewText.ToString() );
+	FText CurrentText;
+	if( PropertyHandle->GetValueAsFormattedText( CurrentText ) != FPropertyAccess::MultipleValues || NewText.ToString() != FPropertyEditor::MultipleValuesDisplayName )
+	{
+		PropertyHandle->SetValueFromFormattedString( NewText.ToString() );
+	}
 }
 
 static FText ValidateNameLength( const FText& Text )
