Made SPropertyEditorEditInline respect the "noclear" propery flag.

diff --git a/Engine/Source/Editor/PropertyEditor/Private/UserInterface/PropertyEditor/SPropertyEditorEditInline.cpp b/Engine/Source/Editor/PropertyEditor/Private/UserInterface/PropertyEditor/SPropertyEditorEditInline.cpp
--- a/Engine/Source/Editor/PropertyEditor/Private/UserInterface/PropertyEditor/SPropertyEditorEditInline.cpp
+++ b/Engine/Source/Editor/PropertyEditor/Private/UserInterface/PropertyEditor/SPropertyEditorEditInline.cpp
@@ -153,7 +153,6 @@ TSharedRef<SWidget> SPropertyEditorEditInline::GenerateClassPicker()
 	FClassViewerInitializationOptions Options;
 	Options.bShowUnloadedBlueprints = true;
 	Options.bShowDisplayNames = true;
-	Options.bShowNoneOption = true;
 
 	TSharedPtr<FPropertyEditorInlineClassFilter> ClassFilter = MakeShareable( new FPropertyEditorInlineClassFilter );
 	Options.ClassFilter = ClassFilter;
@@ -163,6 +162,7 @@ TSharedRef<SWidget> SPropertyEditorEditInline::GenerateClassPicker()
 	UProperty* Property = PropertyNode->GetProperty();
 	ClassFilter->ObjProperty = Cast<UObjectPropertyBase>( Property );
 	ClassFilter->IntProperty = Cast<UInterfaceProperty>( Property );
+	Options.bShowNoneOption = !(Property->PropertyFlags & CPF_NoClear);
 
 	FObjectPropertyNode* ObjectPropertyNode = PropertyNode->FindObjectItemParent();
 	if( ObjectPropertyNode )
