From NickA - Added ability to set per-property instance metadata.

diff --git a/Engine/Source/Editor/DetailCustomizations/Private/DirectionalLightComponentDetails.cpp b/Engine/Source/Editor/DetailCustomizations/Private/DirectionalLightComponentDetails.cpp
--- a/Engine/Source/Editor/DetailCustomizations/Private/DirectionalLightComponentDetails.cpp
+++ b/Engine/Source/Editor/DetailCustomizations/Private/DirectionalLightComponentDetails.cpp
@@ -32,8 +32,8 @@ void FDirectionalLightComponentDetails::CustomizeDetails( IDetailLayoutBuilder&
 
 	TSharedPtr<IPropertyHandle> LightIntensityProperty = DetailBuilder.GetProperty("Intensity", ULightComponentBase::StaticClass());
 	// Point lights need to override the ui min and max for units of lumens, so we have to undo that
-	LightIntensityProperty->GetProperty()->SetMetaData("UIMin",TEXT("0.0f"));
-	LightIntensityProperty->GetProperty()->SetMetaData("UIMax",TEXT("20.0f"));
+	LightIntensityProperty->SetInstanceMetaData("UIMin",TEXT("0.0f"));
+	LightIntensityProperty->SetInstanceMetaData("UIMax",TEXT("20.0f"));
 }
 
 bool FDirectionalLightComponentDetails::IsLightMovable() const
diff --git a/Engine/Source/Editor/DetailCustomizations/Private/PointLightComponentDetails.cpp b/Engine/Source/Editor/DetailCustomizations/Private/PointLightComponentDetails.cpp
--- a/Engine/Source/Editor/DetailCustomizations/Private/PointLightComponentDetails.cpp
+++ b/Engine/Source/Editor/DetailCustomizations/Private/PointLightComponentDetails.cpp
@@ -16,9 +16,9 @@ void FPointLightComponentDetails::CustomizeDetails( IDetailLayoutBuilder& Detail
 	TSharedPtr<IPropertyHandle> LightIntensityProperty = DetailBuilder.GetProperty(GET_MEMBER_NAME_CHECKED(ULightComponentBase, Intensity), ULightComponentBase::StaticClass());
 
 	// Inverse squared falloff point lights (the default) are in units of lumens, instead of just being a brightness scale
-	LightIntensityProperty->GetProperty()->SetMetaData("UIMin",TEXT("0.0f"));
-	LightIntensityProperty->GetProperty()->SetMetaData("UIMax",TEXT("100000.0f"));
-	LightIntensityProperty->GetProperty()->SetMetaData("SliderExponent",TEXT("2.0f"));
+	LightIntensityProperty->SetInstanceMetaData("UIMin",TEXT("0.0f"));
+	LightIntensityProperty->SetInstanceMetaData("UIMax", TEXT("100000.0f"));
+	LightIntensityProperty->SetInstanceMetaData("SliderExponent", TEXT("2.0f"));
 }
 
 #undef LOCTEXT_NAMESPACE
diff --git a/Engine/Source/Editor/DetailCustomizations/Private/SkyLightComponentDetails.cpp b/Engine/Source/Editor/DetailCustomizations/Private/SkyLightComponentDetails.cpp
--- a/Engine/Source/Editor/DetailCustomizations/Private/SkyLightComponentDetails.cpp
+++ b/Engine/Source/Editor/DetailCustomizations/Private/SkyLightComponentDetails.cpp
@@ -27,8 +27,8 @@ void FSkyLightComponentDetails::CustomizeDetails( IDetailLayoutBuilder& DetailLa
 	if( LightIntensityProperty->IsValidHandle() )
 	{
 		// Point lights need to override the ui min and max for units of lumens, so we have to undo that
-		LightIntensityProperty->GetProperty()->SetMetaData("UIMin", TEXT("0.0f"));
-		LightIntensityProperty->GetProperty()->SetMetaData("UIMax", TEXT("20.0f"));
+		LightIntensityProperty->SetInstanceMetaData("UIMin", TEXT("0.0f"));
+		LightIntensityProperty->SetInstanceMetaData("UIMax", TEXT("20.0f"));
 	}
 
 	const TArray< TWeakObjectPtr<UObject> >& SelectedObjects = DetailLayout.GetDetailsView().GetSelectedObjects();
diff --git a/Engine/Source/Editor/PropertyEditor/Private/PropertyHandleImpl.cpp b/Engine/Source/Editor/PropertyEditor/Private/PropertyHandleImpl.cpp
--- a/Engine/Source/Editor/PropertyEditor/Private/PropertyHandleImpl.cpp
+++ b/Engine/Source/Editor/PropertyEditor/Private/PropertyHandleImpl.cpp
@@ -1678,7 +1678,7 @@ void FPropertyHandleBase::SetOnPropertyValueChanged( const FSimpleDelegate& InOn
 	return Implementation->SetOnPropertyValueChanged(InOnPropertyValueChanged);
 }
 
-TSharedPtr<FPropertyNode> FPropertyHandleBase::GetPropertyNode()
+TSharedPtr<FPropertyNode> FPropertyHandleBase::GetPropertyNode() const
 {
 	return Implementation->GetPropertyNode();
 }
@@ -1775,6 +1775,17 @@ UClass* FPropertyHandleBase::GetClassMetaData(const FName& Key) const
 	return (MetaDataProperty) ? MetaDataProperty->GetClassMetaData(Key) : nullptr;
 }
 
+
+void FPropertyHandleBase::SetInstanceMetaData(const FName& Key, const FString& Value)
+{
+	GetPropertyNode()->SetInstanceMetaData(Key, Value);
+}
+
+const FString* FPropertyHandleBase::GetInstanceMetaData(const FName& Key) const
+{
+	return GetPropertyNode()->GetInstanceMetaData(Key);
+}
+
 FText FPropertyHandleBase::GetToolTipText() const
 {
 	TSharedPtr<FPropertyNode> PropertyNode = Implementation->GetPropertyNode();
diff --git a/Engine/Source/Editor/PropertyEditor/Private/PropertyHandleImpl.h b/Engine/Source/Editor/PropertyEditor/Private/PropertyHandleImpl.h
--- a/Engine/Source/Editor/PropertyEditor/Private/PropertyHandleImpl.h
+++ b/Engine/Source/Editor/PropertyEditor/Private/PropertyHandleImpl.h
@@ -396,6 +396,8 @@ public:
 	virtual int32 GetINTMetaData(const FName& Key) const override;
 	virtual float GetFLOATMetaData(const FName& Key) const override;
 	virtual UClass* GetClassMetaData(const FName& Key) const override;
+	virtual void SetInstanceMetaData(const FName& Key, const FString& Value) override;
+	virtual const FString* GetInstanceMetaData(const FName& Key) const override;
 	virtual FText GetToolTipText() const override;
 	virtual void SetToolTipText(const FText& ToolTip) override;
 	virtual FPropertyAccess::Result SetPerObjectValues( const TArray<FString>& InPerObjectValues,  EPropertyValueSetFlags::Type Flags = EPropertyValueSetFlags::DefaultFlags ) override;
@@ -410,7 +412,7 @@ public:
 	virtual bool GenerateRestrictionToolTip(const FString& Value, FText& OutTooltip) const override;
 	virtual void SetIgnoreValidation(bool bInIgnore) override;
 
-	TSharedPtr<FPropertyNode> GetPropertyNode();
+	TSharedPtr<FPropertyNode> GetPropertyNode() const;
 protected:
 	TSharedPtr<FPropertyValueImpl> Implementation;
 };
diff --git a/Engine/Source/Editor/PropertyEditor/Private/PropertyNode.cpp b/Engine/Source/Editor/PropertyEditor/Private/PropertyNode.cpp
--- a/Engine/Source/Editor/PropertyEditor/Private/PropertyNode.cpp
+++ b/Engine/Source/Editor/PropertyEditor/Private/PropertyNode.cpp
@@ -2161,6 +2161,16 @@ FPropertyChangedEvent& FPropertyNode::FixPropertiesInEvent(FPropertyChangedEvent
 	return Event;
 }
 
+void FPropertyNode::SetInstanceMetaData(const FName& Key, const FString& Value)
+{
+	InstanceMetaData.Add(Key, Value);
+}
+
+const FString* FPropertyNode::GetInstanceMetaData(const FName& Key) const
+{
+	return InstanceMetaData.Find(Key);
+}
+
 /**
  * Does the string compares to ensure this Name is acceptable to the filter that is passed in
  * @return		Return True if this property should be displayed.  False if it should be culled
diff --git a/Engine/Source/Editor/PropertyEditor/Private/PropertyNode.h b/Engine/Source/Editor/PropertyEditor/Private/PropertyNode.h
--- a/Engine/Source/Editor/PropertyEditor/Private/PropertyNode.h
+++ b/Engine/Source/Editor/PropertyEditor/Private/PropertyNode.h
@@ -662,6 +662,16 @@ public:
 
 	FPropertyChangedEvent& FixPropertiesInEvent(FPropertyChangedEvent& Event);
 
+	/** Set metadata value for 'Key' to 'Value' on this property instance (as opposed to the class) */
+	void SetInstanceMetaData(const FName& Key, const FString& Value);
+
+	/**
+	 * Get metadata value for 'Key' for this property instance (as opposed to the class)
+	 * 
+	 * @return Pointer to metadata value; nullptr if Key not found
+	 */
+	const FString* GetInstanceMetaData(const FName& Key) const;
+
 protected:
 
 	TSharedRef<FEditPropertyChain> BuildPropertyChain( UProperty* PropertyAboutToChange );
@@ -781,6 +791,12 @@ protected:
 
 	/** Optional reference to a tree node that is displaying this property */
 	TWeakPtr< class IDetailTreeNode > TreeNode;
+
+	/**
+	 * Stores metadata for this instasnce of the property (in contrast
+	 * to regular metadata, which is stored per-class)
+	 */
+	TMap<FName, FString> InstanceMetaData;
 };
 
 class FComplexPropertyNode : public FPropertyNode
diff --git a/Engine/Source/Editor/PropertyEditor/Private/UserInterface/PropertyEditor/SPropertyEditorNumeric.h b/Engine/Source/Editor/PropertyEditor/Private/UserInterface/PropertyEditor/SPropertyEditorNumeric.h
--- a/Engine/Source/Editor/PropertyEditor/Private/UserInterface/PropertyEditor/SPropertyEditorNumeric.h
+++ b/Engine/Source/Editor/PropertyEditor/Private/UserInterface/PropertyEditor/SPropertyEditorNumeric.h
@@ -29,12 +29,19 @@ public:
 		const TSharedRef< FPropertyNode > PropertyNode = InPropertyEditor->GetPropertyNode();
 		const UProperty* Property = InPropertyEditor->GetProperty();
 
-		const FString& MetaUIMinString = Property->GetMetaData(TEXT("UIMin"));
-		const FString& MetaUIMaxString = Property->GetMetaData(TEXT("UIMax"));
-		const FString& SliderExponentString = Property->GetMetaData(TEXT("SliderExponent"));
-		const FString& DeltaString = Property->GetMetaData(TEXT("Delta"));
-		const FString& ClampMinString = Property->GetMetaData(TEXT("ClampMin"));
-		const FString& ClampMaxString = Property->GetMetaData(TEXT("ClampMax"));
+		// Instance metadata overrides per-class metadata.
+		auto GetMetaDataFromKey = [&PropertyNode, &Property](const FName& Key) -> const FString&
+		{
+			const FString* InstanceValue = PropertyNode->GetInstanceMetaData(Key);
+			return (InstanceValue != nullptr) ? *InstanceValue : Property->GetMetaData(Key);
+		};
+
+		const FString& MetaUIMinString = GetMetaDataFromKey("UIMin");
+		const FString& MetaUIMaxString = GetMetaDataFromKey("UIMax");
+		const FString& SliderExponentString = GetMetaDataFromKey("SliderExponent");
+		const FString& DeltaString = GetMetaDataFromKey("Delta");
+		const FString& ClampMinString = GetMetaDataFromKey("ClampMin");
+		const FString& ClampMaxString = GetMetaDataFromKey("ClampMax");
 
 		// If no UIMin/Max was specified then use the clamp string
 		const FString& UIMinString = MetaUIMinString.Len() ? MetaUIMinString : ClampMinString;
diff --git a/Engine/Source/Editor/PropertyEditor/Public/PropertyHandle.h b/Engine/Source/Editor/PropertyEditor/Public/PropertyHandle.h
--- a/Engine/Source/Editor/PropertyEditor/Public/PropertyHandle.h
+++ b/Engine/Source/Editor/PropertyEditor/Public/PropertyHandle.h
@@ -101,6 +101,16 @@ public:
 	 */
 	virtual UClass* GetClassMetaData(const FName& Key) const = 0;
 
+	/** Set metadata value for 'Key' to 'Value' on this property instance (as opposed to the class) */
+	virtual void SetInstanceMetaData(const FName& Key, const FString& Value) = 0;
+
+	/**
+	 * Get metadata value for 'Key' for this property instance (as opposed to the class)
+	 * 
+	 * @return Pointer to metadata value; nullptr if Key not found
+	 */
+	virtual const FString* GetInstanceMetaData(const FName& Key) const = 0;
+
 	/**
 	 * Gets the property tool tip text.
 	 */
