Fix crashes assigning invalid object types to Slate brushes.  SPropertyEditorAsset now properly checks the custom classes filter before allowing an object to be set

diff --git a/Engine/Source/Editor/PropertyEditor/Private/UserInterface/PropertyEditor/SPropertyEditorAsset.cpp b/Engine/Source/Editor/PropertyEditor/Private/UserInterface/PropertyEditor/SPropertyEditorAsset.cpp
--- a/Engine/Source/Editor/PropertyEditor/Private/UserInterface/PropertyEditor/SPropertyEditorAsset.cpp
+++ b/Engine/Source/Editor/PropertyEditor/Private/UserInterface/PropertyEditor/SPropertyEditorAsset.cpp
@@ -516,12 +516,18 @@ FText SPropertyEditorAsset::OnGetToolTip() const
 void SPropertyEditorAsset::SetValue( const FAssetData& AssetData )
 {
 	AssetComboButton->SetIsOpen(false);
-	if(PropertyEditor.IsValid())
+
+	bool bAllowedToSetBasedOnFilter = CanSetBasedOnCustomClasses( AssetData );
+
+	if( bAllowedToSetBasedOnFilter )
 	{
-		PropertyEditor->GetPropertyHandle()->SetValue(AssetData);
-	}
+		if(PropertyEditor.IsValid())
+		{
+			PropertyEditor->GetPropertyHandle()->SetValue(AssetData);
+		}
 
-	OnSetObject.ExecuteIfBound(AssetData);
+		OnSetObject.ExecuteIfBound(AssetData);
+	}
 }
 
 FPropertyAccess::Result SPropertyEditorAsset::GetValue( FObjectOrAssetData& OutValue ) const
@@ -763,7 +769,7 @@ bool SPropertyEditorAsset::OnAssetDraggedOver( const UObject* InObject ) const
 		if (!OnShouldFilterAsset.IsBound()
 			|| !OnShouldFilterAsset.Execute(FAssetData(InObject)))
 		{
-			return true;
+			return CanSetBasedOnCustomClasses( FAssetData(InObject) );
 		}
 	}
 
@@ -850,4 +856,24 @@ bool SPropertyEditorAsset::CanEdit() const
 	return PropertyEditor.IsValid() ? !PropertyEditor->IsEditConst() : true;
 }
 
+bool SPropertyEditorAsset::CanSetBasedOnCustomClasses( const FAssetData& InAssetData ) const
+{
+	bool bAllowedToSetBasedOnFilter = true;
+	if( CustomClassFilters.Num() > 0 )
+	{
+		bAllowedToSetBasedOnFilter = false;
+		UClass* AssetClass = InAssetData.GetClass();
+		for( const UClass* AllowedClass : CustomClassFilters )
+		{
+			if( AssetClass->IsChildOf( AllowedClass ) )
+			{
+				bAllowedToSetBasedOnFilter = true;
+				break;
+			}
+		}
+	}
+
+	return bAllowedToSetBasedOnFilter;
+}
+
 #undef LOCTEXT_NAMESPACE
diff --git a/Engine/Source/Editor/PropertyEditor/Private/UserInterface/PropertyEditor/SPropertyEditorAsset.h b/Engine/Source/Editor/PropertyEditor/Private/UserInterface/PropertyEditor/SPropertyEditorAsset.h
--- a/Engine/Source/Editor/PropertyEditor/Private/UserInterface/PropertyEditor/SPropertyEditorAsset.h
+++ b/Engine/Source/Editor/PropertyEditor/Private/UserInterface/PropertyEditor/SPropertyEditorAsset.h
@@ -242,6 +242,9 @@ private:
 
 	/** @return True if the property can be edited */
 	bool CanEdit() const;
+
+	/** @return true if the passed in AssetData can be used to set the property based on the list of custom classes */
+	bool CanSetBasedOnCustomClasses( const FAssetData& InAssetData ) const;
 private:
 
 	/** Main combobutton */
