Guard against re-entrancy when picking classes from details panels

diff --git a/Engine/Source/Editor/PropertyEditor/Private/UserInterface/PropertyEditor/SPropertyEditorClass.cpp b/Engine/Source/Editor/PropertyEditor/Private/UserInterface/PropertyEditor/SPropertyEditorClass.cpp
--- a/Engine/Source/Editor/PropertyEditor/Private/UserInterface/PropertyEditor/SPropertyEditorClass.cpp
+++ b/Engine/Source/Editor/PropertyEditor/Private/UserInterface/PropertyEditor/SPropertyEditorClass.cpp
@@ -160,20 +160,32 @@ static FString GetClassDisplayName(const UObject* Object)
 
 FText SPropertyEditorClass::GetDisplayValueAsString() const
 {
-	if(PropertyEditor.IsValid())
-	{
-		UObject* ObjectValue = NULL;
-		FPropertyAccess::Result Result = PropertyEditor->GetPropertyHandle()->GetValue(ObjectValue);
+	static bool bIsReentrant = false;
 
-		if(Result == FPropertyAccess::Success && ObjectValue != NULL)
+	// Guard against re-entrancy which can happen if the delegate executed below (SelectedClass.Get()) forces a slow task dialog to open, thus causing this to lose context and regain focus later starting the loop over again
+	if( !bIsReentrant )
+	{
+		TGuardValue<bool>( bIsReentrant, true );
+		if(PropertyEditor.IsValid())
 		{
-			return FText::FromString(GetClassDisplayName(ObjectValue));
+			UObject* ObjectValue = NULL;
+			FPropertyAccess::Result Result = PropertyEditor->GetPropertyHandle()->GetValue(ObjectValue);
+
+			if(Result == FPropertyAccess::Success && ObjectValue != NULL)
+			{
+				return FText::FromString(GetClassDisplayName(ObjectValue));
+			}
+
+			return FText::FromString(FPaths::GetBaseFilename(PropertyEditor->GetValueAsString()));
 		}
 
-		return FText::FromString(FPaths::GetBaseFilename(PropertyEditor->GetValueAsString()));
+		return FText::FromString(GetClassDisplayName(SelectedClass.Get()));
 	}
-
-	return FText::FromString(GetClassDisplayName(SelectedClass.Get()));
+	else
+	{
+		return FText::GetEmpty();
+	}
+	
 }
 
 TSharedRef<SWidget> SPropertyEditorClass::GenerateClassPicker()
