Fixed "Clear" button in asset pickers having no effect.

https://jira.ol.epicgames.net/browse/UE-18387

diff --git a/Engine/Source/Editor/PropertyEditor/Private/UserInterface/PropertyEditor/SPropertyEditorAsset.cpp b/Engine/Source/Editor/PropertyEditor/Private/UserInterface/PropertyEditor/SPropertyEditorAsset.cpp
--- a/Engine/Source/Editor/PropertyEditor/Private/UserInterface/PropertyEditor/SPropertyEditorAsset.cpp
+++ b/Engine/Source/Editor/PropertyEditor/Private/UserInterface/PropertyEditor/SPropertyEditorAsset.cpp
@@ -859,7 +859,7 @@ bool SPropertyEditorAsset::CanEdit() const
 bool SPropertyEditorAsset::CanSetBasedOnCustomClasses( const FAssetData& InAssetData ) const
 {
 	bool bAllowedToSetBasedOnFilter = true;
-	if( CustomClassFilters.Num() > 0 )
+	if( InAssetData.IsValid() && CustomClassFilters.Num() > 0 )
 	{
 		bAllowedToSetBasedOnFilter = false;
 		UClass* AssetClass = InAssetData.GetClass();
