Fixed non-unity

diff --git a/Engine/Source/Editor/PropertyEditor/Public/DetailCategoryBuilder.h b/Engine/Source/Editor/PropertyEditor/Public/DetailCategoryBuilder.h
--- a/Engine/Source/Editor/PropertyEditor/Public/DetailCategoryBuilder.h
+++ b/Engine/Source/Editor/PropertyEditor/Public/DetailCategoryBuilder.h
@@ -2,11 +2,14 @@
 
 #pragma once
 
+
 class IDetailPropertyRow;
 class IDetailCategoryBuilder;
 class FDetailWidgetRow;
 class IDetailGroup;
 class IDetailCustomNodeBuilder;
+class IPropertyHandle;
+
 
 /** The location of a property within a category */
 namespace EPropertyLocation
diff --git a/Engine/Source/Runtime/Windows/D3D12RHI/Private/D3D12DirectCommandListManager.cpp b/Engine/Source/Runtime/Windows/D3D12RHI/Private/D3D12DirectCommandListManager.cpp
--- a/Engine/Source/Runtime/Windows/D3D12RHI/Private/D3D12DirectCommandListManager.cpp
+++ b/Engine/Source/Runtime/Windows/D3D12RHI/Private/D3D12DirectCommandListManager.cpp
@@ -1,6 +1,10 @@
 #include "D3D12RHIPrivate.h"
 #include "Windows.h"
 
+
+extern bool D3D12RHI_ShouldCreateWithD3DDebug();
+
+
 FD3D12Fence::FD3D12Fence()
 	: CurrentFence(-1)
 	, LastCompletedFence(-1)
