Expose Recursive setting to IPropertyHandle::GetChildHandle to allow customizers to prevent recursing into children of children and so on when searching for children.

diff --git a/Engine/Source/Editor/PropertyEditor/Private/PropertyHandleImpl.cpp b/Engine/Source/Editor/PropertyEditor/Private/PropertyHandleImpl.cpp
--- a/Engine/Source/Editor/PropertyEditor/Private/PropertyHandleImpl.cpp
+++ b/Engine/Source/Editor/PropertyEditor/Private/PropertyHandleImpl.cpp
@@ -48,7 +48,7 @@ FPropertyAccess::Result FPropertyValueImpl::GetPropertyValueString( FString& Out
 {
 	FPropertyAccess::Result Result = FPropertyAccess::Success;
 
-	uint8* ValueAddress = NULL;
+	uint8* ValueAddress = nullptr;
 	FReadAddressList ReadAddresses;
 	bool bAllValuesTheSame = InPropertyNode->GetReadAddress( !!InPropertyNode->HasNodeFlags(EPropertyNodeFlags::SingleSelectOnly), ReadAddresses, false, true );
 
@@ -56,17 +56,17 @@ FPropertyAccess::Result FPropertyValueImpl::GetPropertyValueString( FString& Out
 	{
 		ValueAddress = ReadAddresses.GetAddress(0);
 
-		if( ValueAddress != NULL )
+		if( ValueAddress != nullptr )
 		{
 			UProperty* Property = InPropertyNode->GetProperty();
 
 			// Check for bogus data
-			if( Property != NULL && InPropertyNode->GetParentNode() != NULL )
+			if( Property != nullptr && InPropertyNode->GetParentNode() != nullptr )
 			{
-				Property->ExportText_Direct(OutString, ValueAddress, ValueAddress, NULL, PPF_PropertyWindow );
+				Property->ExportText_Direct(OutString, ValueAddress, ValueAddress, nullptr, PPF_PropertyWindow );
 
 				UByteProperty* ByteProperty = Cast<UByteProperty>(Property);
-				if ( ByteProperty != NULL && ByteProperty->Enum != NULL )
+				if ( ByteProperty != nullptr && ByteProperty->Enum != nullptr )
 				{
 					const uint8 EnumValueIndex = ByteProperty->GetPropertyValue(ValueAddress);
 
@@ -104,7 +104,7 @@ FPropertyAccess::Result FPropertyValueImpl::GetPropertyValueText( FText& OutText
 {
 	FPropertyAccess::Result Result = FPropertyAccess::Success;
 
-	uint8* ValueAddress = NULL;
+	uint8* ValueAddress = nullptr;
 	FReadAddressList ReadAddresses;
 	bool bAllValuesTheSame = InPropertyNode->GetReadAddress( !!InPropertyNode->HasNodeFlags(EPropertyNodeFlags::SingleSelectOnly), ReadAddresses, false, true );
 
@@ -112,7 +112,7 @@ FPropertyAccess::Result FPropertyValueImpl::GetPropertyValueText( FText& OutText
 	{
 		ValueAddress = ReadAddresses.GetAddress(0);
 
-		if( ValueAddress != NULL )
+		if( ValueAddress != nullptr )
 		{
 			UProperty* Property = InPropertyNode->GetProperty();
 
@@ -123,10 +123,10 @@ FPropertyAccess::Result FPropertyValueImpl::GetPropertyValueText( FText& OutText
 			else
 			{
 				FString ExportedTextString;
-				Property->ExportText_Direct(ExportedTextString, ValueAddress, ValueAddress, NULL, PPF_PropertyWindow );
+				Property->ExportText_Direct(ExportedTextString, ValueAddress, ValueAddress, nullptr, PPF_PropertyWindow );
 
 				UByteProperty* ByteProperty = Cast<UByteProperty>(Property);
-				if ( ByteProperty != NULL && ByteProperty->Enum != NULL )
+				if ( ByteProperty != nullptr && ByteProperty->Enum != nullptr )
 				{
 					const uint8 EnumValueIndex = ByteProperty->GetPropertyValue(ValueAddress);
 					
@@ -163,11 +163,11 @@ FPropertyAccess::Result FPropertyValueImpl::GetPropertyValueText( FText& OutText
 FPropertyAccess::Result FPropertyValueImpl::GetValueData( void*& OutAddress ) const
 {
 	FPropertyAccess::Result Res = FPropertyAccess::Fail;
-	OutAddress = NULL;
+	OutAddress = nullptr;
 	TSharedPtr<FPropertyNode> PropertyNodePin = PropertyNode.Pin();
 	if( PropertyNodePin.IsValid() )
 	{
-		uint8* ValueAddress = NULL;
+		uint8* ValueAddress = nullptr;
 		FReadAddressList ReadAddresses;
 		bool bAllValuesTheSame = PropertyNodePin->GetReadAddress( !!PropertyNodePin->HasNodeFlags(EPropertyNodeFlags::SingleSelectOnly), ReadAddresses, false, true );
 
@@ -209,7 +209,7 @@ FString FPropertyValueImpl::GetPropertyValueArray() const
 	TSharedPtr<FPropertyNode> PropertyNodePin = PropertyNode.Pin();
 	if( PropertyNodePin.IsValid() )
 	{
-		uint8* Addr = NULL;
+		uint8* Addr = nullptr;
 		FReadAddressList ReadAddresses;
 		UProperty* NodeProperty = PropertyNodePin->GetProperty();
 
@@ -218,7 +218,7 @@ FString FPropertyValueImpl::GetPropertyValueArray() const
 		Addr = ReadAddresses.GetAddress(0);
 		if( bSingleValue && Addr )
 		{
-			if ( NodeProperty != NULL && Cast<UArrayProperty>(NodeProperty) != NULL )
+			if ( NodeProperty != nullptr && Cast<UArrayProperty>(NodeProperty) != nullptr )
 			{
 				String = FString::Printf( TEXT("%(%d)"), FScriptArrayHelper::Num(Addr) );
 			}
@@ -262,7 +262,7 @@ void FPropertyValueImpl::GenerateArrayIndexMapToObjectNode( TMap<FString,int32>&
 	if( PropertyNode )
 	{
 		OutArrayIndexMap.Empty();
-		for (FPropertyNode* IterationNode = PropertyNode; (IterationNode != NULL) && (IterationNode->AsObjectNode() == NULL); IterationNode = IterationNode->GetParentNode())
+		for (FPropertyNode* IterationNode = PropertyNode; (IterationNode != nullptr) && (IterationNode->AsObjectNode() == nullptr); IterationNode = IterationNode->GetParentNode())
 		{
 			UProperty* Property = IterationNode->GetProperty();
 			if (Property)
@@ -285,7 +285,7 @@ FPropertyAccess::Result FPropertyValueImpl::ImportText( const FString& InValue,
 	TArray<FString> Values;
 	for( int32 ObjectIndex = 0 ; ObjectIndex < ObjectsToModify.Num() ; ++ObjectIndex )
 	{
-		if (ObjectsToModify[ObjectIndex].Object != NULL || ObjectsToModify[ObjectIndex].bIsStruct)
+		if (ObjectsToModify[ObjectIndex].Object != nullptr || ObjectsToModify[ObjectIndex].bIsStruct)
 		{
 			Values.Add( InValue );
 		}
@@ -327,12 +327,12 @@ FPropertyAccess::Result FPropertyValueImpl::ImportText( const TArray<FObjectBase
 
 	if( Result != FPropertyAccess::Fail )
 	{
-		UWorld* OldGWorld = NULL;
+		UWorld* OldGWorld = nullptr;
 
 		bool bIsGameWorld = false;
 		// If the object we are modifying is in the PIE world, than make the PIE world the active
 		// GWorld.  Assumes all objects managed by this property window belong to the same world.
-		if (UPackage* ObjectPackage = (InObjects[0].Object ? InObjects[0].Object->GetOutermost() : NULL))
+		if (UPackage* ObjectPackage = (InObjects[0].Object ? InObjects[0].Object->GetOutermost() : nullptr))
 		{
 			bool bIsPIEPackage = !!(ObjectPackage->PackageFlags & PKG_PlayInEditor);
 			if (GUnrealEd && GUnrealEd->PlayWorld && bIsPIEPackage && !GIsPlayInEditorWorld)
@@ -345,7 +345,7 @@ FPropertyAccess::Result FPropertyValueImpl::ImportText( const TArray<FObjectBase
 
 		// Send the values and assemble a list of pre/posteditchange values.
 		bool bNotifiedPreChange = false;
-		UObject *NotifiedObj = NULL;
+		UObject *NotifiedObj = nullptr;
 		TArray< TMap<FString,int32> > ArrayIndicesPerObject;
 
 		const bool bTransactable = (Flags & EPropertyValueSetFlags::NotTransactable) == 0;
@@ -354,7 +354,7 @@ FPropertyAccess::Result FPropertyValueImpl::ImportText( const TArray<FObjectBase
 		for ( int32 ObjectIndex = 0 ; ObjectIndex < InObjects.Num() ; ++ObjectIndex )
 		{	
 			const FObjectBaseAddress& Cur = InObjects[ ObjectIndex ];
-			if (Cur.BaseAddress == NULL)
+			if (Cur.BaseAddress == nullptr)
 			{
 				//Fully abort this procedure.  The data has changed out from under the object
 				Result = FPropertyAccess::Fail;
@@ -363,21 +363,21 @@ FPropertyAccess::Result FPropertyValueImpl::ImportText( const TArray<FObjectBase
 
 			// Cache the value of the property before modifying it.
 			FString PreviousValue;
-			NodeProperty->ExportText_Direct(PreviousValue, Cur.BaseAddress, Cur.BaseAddress, NULL, 0 );
+			NodeProperty->ExportText_Direct(PreviousValue, Cur.BaseAddress, Cur.BaseAddress, nullptr, 0 );
 
 			// If this property is the inner-property of an array, cache the current value as well
 			FString PreviousArrayValue;
 			if (Cur.Object)
 			{
 				FPropertyNode* ParentNode = InPropertyNode->GetParentNode();
-				if (ParentNode != NULL && ParentNode->GetProperty() && ParentNode->GetProperty()->IsA(UArrayProperty::StaticClass()))
+				if (ParentNode != nullptr && ParentNode->GetProperty() && ParentNode->GetProperty()->IsA(UArrayProperty::StaticClass()))
 				{
 					UArrayProperty* ArrayProp = Cast<UArrayProperty>(ParentNode->GetProperty());
 					if (ArrayProp->Inner == NodeProperty)
 					{
 						uint8* Addr = ParentNode->GetValueBaseAddress((uint8*)Cur.Object);
 
-						ArrayProp->ExportText_Direct(PreviousArrayValue, Addr, Addr, NULL, 0);
+						ArrayProp->ExportText_Direct(PreviousArrayValue, Addr, Addr, nullptr, 0);
 					}
 				}
 			}
@@ -413,7 +413,7 @@ FPropertyAccess::Result FPropertyValueImpl::ImportText( const TArray<FObjectBase
 			{
 				// Cache the value of the property after having modified it.
 				FString ValueAfterImport;
-				NodeProperty->ExportText_Direct(ValueAfterImport, Cur.BaseAddress, Cur.BaseAddress, NULL, 0);
+				NodeProperty->ExportText_Direct(ValueAfterImport, Cur.BaseAddress, Cur.BaseAddress, nullptr, 0);
 
 				if ((Cur.Object->HasAnyFlags(RF_ClassDefaultObject | RF_ArchetypeObject) ||
 					(Cur.Object->HasAnyFlags(RF_DefaultSubObject) && Cur.Object->GetOuter()->HasAnyFlags(RF_ClassDefaultObject | RF_ArchetypeObject))) &&
@@ -487,7 +487,7 @@ void FPropertyValueImpl::AccessRawData( TArray<void*>& RawData )
 		for (int32 ObjectIndex = 0; ObjectIndex < ObjectAddresses.Num(); ++ObjectIndex)
 		{
 			const FObjectBaseAddress& Cur = ObjectAddresses[ObjectIndex];
-			if (Cur.BaseAddress != NULL)
+			if (Cur.BaseAddress != nullptr)
 			{
 				RawData[ObjectIndex] = Cur.BaseAddress;
 			}
@@ -507,7 +507,7 @@ void FPropertyValueImpl::AccessRawData( TArray<const void*>& RawData ) const
 		for (int32 ObjectIndex = 0; ObjectIndex < ObjectAddresses.Num(); ++ObjectIndex)
 		{
 			const FObjectBaseAddress& Cur = ObjectAddresses[ObjectIndex];
-			if (Cur.BaseAddress != NULL)
+			if (Cur.BaseAddress != nullptr)
 			{
 				RawData[ObjectIndex] = Cur.BaseAddress;
 			}
@@ -546,7 +546,7 @@ static int32 GetArrayPropertyLastValidIndex( FObjectPropertyNode* InObjectNode,
 	UObject* ParentObject = InObjectNode->GetUObject(0);
 
 	//find the associated property
-	UProperty* FoundProperty = NULL;
+	UProperty* FoundProperty = nullptr;
 	for( TFieldIterator<UProperty> It(ParentObject->GetClass()); It; ++It )
 	{
 		UProperty* CurProp = *It;
@@ -748,7 +748,7 @@ bool FPropertyValueImpl::SetObject( const UObject* NewObject, EPropertyValueSetF
 					uint8*		Addr = PropertyNodePin->GetValueBaseAddress( (uint8*) Object.Get() );
 
 					FString PropertyValue;
-					NodeProperty->ExportText_Direct(PropertyValue, Addr, Addr, NULL, 0 );
+					NodeProperty->ExportText_Direct(PropertyValue, Addr, Addr, nullptr, 0 );
 					if ( PropertyValue != CompareName )
 					{
 						bAllObjectPropertyValuesMatch = false;
@@ -786,7 +786,7 @@ FPropertyAccess::Result FPropertyValueImpl::OnUseSelected()
 		UInterfaceProperty* IntProp = Cast<UInterfaceProperty>( NodeProperty );
 		UClassProperty* ClassProp = Cast<UClassProperty>( NodeProperty );
 		UAssetClassProperty* AssetClassProperty = Cast<UAssetClassProperty>( NodeProperty );
-		UClass* const InterfaceThatMustBeImplemented = ObjProp ? ObjProp->GetOwnerProperty()->GetClassMetaData(TEXT("MustImplement")) : NULL;
+		UClass* const InterfaceThatMustBeImplemented = ObjProp ? ObjProp->GetOwnerProperty()->GetClassMetaData(TEXT("MustImplement")) : nullptr;
 
 		if(ClassProp || AssetClassProperty)
 		{
@@ -830,7 +830,7 @@ FPropertyAccess::Result FPropertyValueImpl::OnUseSelected()
 			bool const bMustBeLevelActor = ObjProp ? ObjProp->GetOwnerProperty()->GetBoolMetaData(TEXT("MustBeLevelActor")) : false;
 
 			// Find best appropriate selected object
-			UObject* SelectedObject = NULL;
+			UObject* SelectedObject = nullptr;
 
 			if (bMustBeLevelActor)
 			{
@@ -973,15 +973,15 @@ TSharedPtr<FPropertyNode> FPropertyValueImpl::GetPropertyNode() const
 	return PropertyNode.Pin();
 }
 
-TSharedPtr<FPropertyNode> FPropertyValueImpl::GetChildNode( FName ChildName ) const
+TSharedPtr<FPropertyNode> FPropertyValueImpl::GetChildNode( FName ChildName, bool bRecurse ) const
 {
 	TSharedPtr<FPropertyNode> PropertyNodePin = PropertyNode.Pin();
 	if( PropertyNodePin.IsValid() )
 	{
-		return PropertyNodePin->FindChildPropertyNode(ChildName,true);
+		return PropertyNodePin->FindChildPropertyNode(ChildName, bRecurse);
 	}
 
-	return NULL;
+	return nullptr;
 }
 
 
@@ -993,7 +993,7 @@ TSharedPtr<FPropertyNode> FPropertyValueImpl::GetChildNode( int32 ChildIndex ) c
 		return PropertyNodePin->GetChildNode( ChildIndex );
 	}
 
-	return NULL;
+	return nullptr;
 }
 
 
@@ -1078,7 +1078,7 @@ void FPropertyValueImpl::AddChild()
 					ArrayIndicesPerObject.Add(TMap<FString,int32>());
 					FPropertyValueImpl::GenerateArrayIndexMapToObjectNode(ArrayIndicesPerObject[i], PropertyNodePin.Get());
 					
-					UObject* Obj = ObjectNode ? ObjectNode->GetUObject(i) : NULL;
+					UObject* Obj = ObjectNode ? ObjectNode->GetUObject(i) : nullptr;
 					if (Obj)
 					{
 						if ((Obj->HasAnyFlags(RF_ClassDefaultObject | RF_ArchetypeObject) ||
@@ -1086,7 +1086,7 @@ void FPropertyValueImpl::AddChild()
 							!FApp::IsGame())
 						{
 							FString OrgArrayContent;
-							Array->ExportText_Direct(OrgArrayContent, Addr, Addr, NULL, 0);
+							Array->ExportText_Direct(OrgArrayContent, Addr, Addr, nullptr, 0);
 
 							PropertyNodePin->PropagateArrayPropertyChange(Obj, OrgArrayContent, EPropertyArrayChangeType::Add, -1);
 						}
@@ -1154,7 +1154,7 @@ void FPropertyValueImpl::ClearChildren()
 						PropertyNodePin->NotifyPreChange( NodeProperty, NotifyHook );
 					}
 
-					UObject* Obj = ObjectNode ? ObjectNode->GetUObject(i) : NULL;
+					UObject* Obj = ObjectNode ? ObjectNode->GetUObject(i) : nullptr;
 					if (Obj)
 					{
 						if ((Obj->HasAnyFlags(RF_ClassDefaultObject | RF_ArchetypeObject) ||
@@ -1162,7 +1162,7 @@ void FPropertyValueImpl::ClearChildren()
 							!FApp::IsGame())
 						{
 							FString OrgArrayContent;
-							Array->ExportText_Direct(OrgArrayContent, Addr, Addr, NULL, 0);
+							Array->ExportText_Direct(OrgArrayContent, Addr, Addr, nullptr, 0);
 
 							PropertyNodePin->PropagateArrayPropertyChange(Obj, OrgArrayContent, EPropertyArrayChangeType::Clear, -1);
 						}
@@ -1209,7 +1209,7 @@ void FPropertyValueImpl::InsertChild( TSharedPtr<FPropertyNode> ChildNodeToInser
 	UArrayProperty* ArrayProperty = CastChecked<UArrayProperty>(NodeProperty->GetOuter());
 
 	FReadAddressList ReadAddresses;
-	void* Addr = NULL;
+	void* Addr = nullptr;
 	ParentNode->GetReadAddress( !!ParentNode->HasNodeFlags(EPropertyNodeFlags::SingleSelectOnly), ReadAddresses );
 	if ( ReadAddresses.Num() )
 	{
@@ -1226,7 +1226,7 @@ void FPropertyValueImpl::InsertChild( TSharedPtr<FPropertyNode> ChildNodeToInser
 		FScriptArrayHelper	ArrayHelper(ArrayProperty,Addr);
 		int32 Index = ChildNodePtr->GetArrayIndex();
 
-		UObject* Obj = ObjectNode ? ObjectNode->GetUObject(0) : NULL;
+		UObject* Obj = ObjectNode ? ObjectNode->GetUObject(0) : nullptr;
 		if (Obj)
 		{
 			if ((Obj->HasAnyFlags(RF_ClassDefaultObject | RF_ArchetypeObject) ||
@@ -1234,7 +1234,7 @@ void FPropertyValueImpl::InsertChild( TSharedPtr<FPropertyNode> ChildNodeToInser
 				!FApp::IsGame())
 			{
 				FString OrgArrayContent;
-				ArrayProperty->ExportText_Direct(OrgArrayContent, Addr, Addr, NULL, 0);
+				ArrayProperty->ExportText_Direct(OrgArrayContent, Addr, Addr, nullptr, 0);
 
 				ChildNodePtr->PropagateArrayPropertyChange(Obj, OrgArrayContent, EPropertyArrayChangeType::Insert, Index);
 			}
@@ -1306,7 +1306,7 @@ void FPropertyValueImpl::DeleteChild( TSharedPtr<FPropertyNode> ChildNodeToDelet
 				FScriptArrayHelper ArrayHelper(ArrayProperty,Address);
 				int32 Index = ChildNodePtr->GetArrayIndex();
 
-				UObject* Obj = ObjectNode ? ObjectNode->GetUObject(i) : NULL;
+				UObject* Obj = ObjectNode ? ObjectNode->GetUObject(i) : nullptr;
 				if (Obj)
 				{
 					if ((Obj->HasAnyFlags(RF_ClassDefaultObject | RF_ArchetypeObject) ||
@@ -1314,7 +1314,7 @@ void FPropertyValueImpl::DeleteChild( TSharedPtr<FPropertyNode> ChildNodeToDelet
 						!FApp::IsGame())
 					{
 						FString OrgArrayContent;
-						ArrayProperty->ExportText_Direct(OrgArrayContent, Address, Address, NULL, 0);
+						ArrayProperty->ExportText_Direct(OrgArrayContent, Address, Address, nullptr, 0);
 
 						ChildNodePtr->PropagateArrayPropertyChange(Obj, OrgArrayContent, EPropertyArrayChangeType::Delete, Index);
 					}
@@ -1358,7 +1358,7 @@ void FPropertyValueImpl::DuplicateChild( TSharedPtr<FPropertyNode> ChildNodeToDu
 	UArrayProperty* ArrayProperty = CastChecked<UArrayProperty>(NodeProperty->GetOuter());
 
 	FReadAddressList ReadAddresses;
-	void* Addr = NULL;
+	void* Addr = nullptr;
 	ParentNode->GetReadAddress( !!ParentNode->HasNodeFlags(EPropertyNodeFlags::SingleSelectOnly), ReadAddresses );
 	if ( ReadAddresses.Num() )
 	{
@@ -1374,7 +1374,7 @@ void FPropertyValueImpl::DuplicateChild( TSharedPtr<FPropertyNode> ChildNodeToDu
 		FScriptArrayHelper	ArrayHelper(ArrayProperty,Addr);
 
 		int32 Index = ChildNodePtr->GetArrayIndex();
-		UObject* Obj = ObjectNode ? ObjectNode->GetUObject(0) : NULL;
+		UObject* Obj = ObjectNode ? ObjectNode->GetUObject(0) : nullptr;
 		if (Obj)
 		{
 			if ((Obj->HasAnyFlags(RF_ClassDefaultObject | RF_ArchetypeObject) ||
@@ -1382,7 +1382,7 @@ void FPropertyValueImpl::DuplicateChild( TSharedPtr<FPropertyNode> ChildNodeToDu
 				!FApp::IsGame())
 			{
 				FString OrgArrayContent;
-				ArrayProperty->ExportText_Direct(OrgArrayContent, Addr, Addr, NULL, 0);
+				ArrayProperty->ExportText_Direct(OrgArrayContent, Addr, Addr, nullptr, 0);
 
 				ChildNodePtr->PropagateArrayPropertyChange(Obj, OrgArrayContent, EPropertyArrayChangeType::Duplicate, Index);
 			}
@@ -1397,7 +1397,7 @@ void FPropertyValueImpl::DuplicateChild( TSharedPtr<FPropertyNode> ChildNodeToDu
 		// Find the object that owns the array and instance any subobjects
 		if (FObjectPropertyNode* ObjectPropertyNode = ChildNodePtr->FindObjectItemParent())
 		{
-			UObject* ArrayOwner = NULL;
+			UObject* ArrayOwner = nullptr;
 			for (TPropObjectIterator Itor(ObjectPropertyNode->ObjectIterator()); Itor && !ArrayOwner; ++Itor)
 			{
 				ArrayOwner = Itor->Get();
@@ -1592,12 +1592,12 @@ FPropertyAccess::Result FPropertyHandleBase::SetValueFromFormattedString( const
 	return Implementation->SetValueAsString( InValue, Flags );
 }
 
-TSharedPtr<IPropertyHandle> FPropertyHandleBase::GetChildHandle( FName ChildName ) const
+TSharedPtr<IPropertyHandle> FPropertyHandleBase::GetChildHandle( FName ChildName, bool bRecurse ) const
 {
 	// Array children cannot be accessed in this manner
 	if( !Implementation->IsPropertyTypeOf(UArrayProperty::StaticClass() ) )
 	{
-		TSharedPtr<FPropertyNode> PropertyNode = Implementation->GetChildNode( ChildName );
+		TSharedPtr<FPropertyNode> PropertyNode = Implementation->GetChildNode( ChildName, bRecurse );
 
 		if( PropertyNode.IsValid() )
 		{
@@ -1605,7 +1605,7 @@ TSharedPtr<IPropertyHandle> FPropertyHandleBase::GetChildHandle( FName ChildName
 		}
 	}
 	
-	return NULL;
+	return nullptr;
 }
 
 TSharedPtr<IPropertyHandle> FPropertyHandleBase::GetChildHandle( uint32 ChildIndex ) const
@@ -1616,7 +1616,7 @@ TSharedPtr<IPropertyHandle> FPropertyHandleBase::GetChildHandle( uint32 ChildInd
 	{
 		return PropertyEditorHelpers::GetPropertyHandle( PropertyNode.ToSharedRef(), Implementation->GetNotifyHook(), Implementation->GetPropertyUtilities() );
 	}
-	return NULL;
+	return nullptr;
 }
 
 TSharedPtr<IPropertyHandle> FPropertyHandleBase::GetParentHandle() const
@@ -1627,7 +1627,7 @@ TSharedPtr<IPropertyHandle> FPropertyHandleBase::GetParentHandle() const
 		return PropertyEditorHelpers::GetPropertyHandle( ParentNode.ToSharedRef(), Implementation->GetNotifyHook(), Implementation->GetPropertyUtilities() ).ToSharedRef();
 	}
 
-	return NULL;
+	return nullptr;
 }
 
 FPropertyAccess::Result FPropertyHandleBase::GetNumChildren( uint32& OutNumChildren ) const
@@ -1701,7 +1701,7 @@ const UClass* FPropertyHandleBase::GetPropertyClass() const
 		return PropertyNode->GetProperty()->GetClass();
 	}
 
-	return NULL;
+	return nullptr;
 }
 
 UProperty* FPropertyHandleBase::GetProperty() const
@@ -1712,7 +1712,7 @@ UProperty* FPropertyHandleBase::GetProperty() const
 		return PropertyNode->GetProperty();
 	}
 
-	return NULL;
+	return nullptr;
 }
 
 UProperty* FPropertyHandleBase::GetMetaDataProperty() const
@@ -1853,7 +1853,7 @@ FPropertyAccess::Result FPropertyHandleBase::GetPerObjectValues( TArray<FString>
 				new( OutPerObjectValues ) FString();
 				if( Address )
 				{
-					NodeProperty->ExportText_Direct(OutPerObjectValues[AddrIndex], Address, Address, NULL, 0 );
+					NodeProperty->ExportText_Direct(OutPerObjectValues[AddrIndex], Address, Address, nullptr, 0 );
 				}
 			}
 
@@ -1893,7 +1893,7 @@ bool FPropertyHandleBase::GeneratePossibleValues(TArray< TSharedPtr<FString> >&
 
 	if( ByteProperty || ( Property->IsA(UStrProperty::StaticClass()) && Property->HasMetaData( TEXT("Enum") ) ) )
 	{
-		UEnum* Enum = NULL;
+		UEnum* Enum = nullptr;
 		if( ByteProperty )
 		{
 			Enum = ByteProperty->Enum;
@@ -2074,7 +2074,7 @@ bool FPropertyHandleInt::Supports( TSharedRef<FPropertyNode> PropertyNode )
 {
 	UProperty* Property = PropertyNode->GetProperty();
 
-	if ( Property == NULL )
+	if ( Property == nullptr )
 	{
 		return false;
 	}
@@ -2101,7 +2101,7 @@ ValueType GetIntegerValue(void* PropValue, FPropertyValueImpl& Implementation )
 
 FPropertyAccess::Result FPropertyHandleInt::GetValue(int8& OutValue) const
 {
-	void* PropValue = NULL;
+	void* PropValue = nullptr;
 	FPropertyAccess::Result Res = Implementation->GetValueData(PropValue);
 
 	if (Res == FPropertyAccess::Success)
@@ -2114,7 +2114,7 @@ FPropertyAccess::Result FPropertyHandleInt::GetValue(int8& OutValue) const
 
 FPropertyAccess::Result FPropertyHandleInt::GetValue(int16& OutValue) const
 {
-	void* PropValue = NULL;
+	void* PropValue = nullptr;
 	FPropertyAccess::Result Res = Implementation->GetValueData(PropValue);
 
 	if (Res == FPropertyAccess::Success)
@@ -2127,7 +2127,7 @@ FPropertyAccess::Result FPropertyHandleInt::GetValue(int16& OutValue) const
 
 FPropertyAccess::Result FPropertyHandleInt::GetValue(int32& OutValue) const
 {
-	void* PropValue = NULL;
+	void* PropValue = nullptr;
 	FPropertyAccess::Result Res = Implementation->GetValueData(PropValue);
 
 	if (Res == FPropertyAccess::Success)
@@ -2139,7 +2139,7 @@ FPropertyAccess::Result FPropertyHandleInt::GetValue(int32& OutValue) const
 }
 FPropertyAccess::Result FPropertyHandleInt::GetValue( int64& OutValue ) const
 {
-	void* PropValue = NULL;
+	void* PropValue = nullptr;
 	FPropertyAccess::Result Res = Implementation->GetValueData( PropValue );
 
 	if( Res == FPropertyAccess::Success )
@@ -2152,7 +2152,7 @@ FPropertyAccess::Result FPropertyHandleInt::GetValue( int64& OutValue ) const
 
 FPropertyAccess::Result FPropertyHandleInt::GetValue(uint16& OutValue) const
 {
-	void* PropValue = NULL;
+	void* PropValue = nullptr;
 	FPropertyAccess::Result Res = Implementation->GetValueData(PropValue);
 
 	if (Res == FPropertyAccess::Success)
@@ -2165,7 +2165,7 @@ FPropertyAccess::Result FPropertyHandleInt::GetValue(uint16& OutValue) const
 
 FPropertyAccess::Result FPropertyHandleInt::GetValue(uint32& OutValue) const
 {
-	void* PropValue = NULL;
+	void* PropValue = nullptr;
 	FPropertyAccess::Result Res = Implementation->GetValueData(PropValue);
 
 	if (Res == FPropertyAccess::Success)
@@ -2178,7 +2178,7 @@ FPropertyAccess::Result FPropertyHandleInt::GetValue(uint32& OutValue) const
 
 FPropertyAccess::Result FPropertyHandleInt::GetValue(uint64& OutValue) const
 {
-	void* PropValue = NULL;
+	void* PropValue = nullptr;
 	FPropertyAccess::Result Res = Implementation->GetValueData(PropValue);
 
 	if (Res == FPropertyAccess::Success)
@@ -2280,7 +2280,7 @@ bool FPropertyHandleFloat::Supports( TSharedRef<FPropertyNode> PropertyNode )
 {
 	UProperty* Property = PropertyNode->GetProperty();
 
-	if ( Property == NULL )
+	if ( Property == nullptr )
 	{
 		return false;
 	}
@@ -2290,7 +2290,7 @@ bool FPropertyHandleFloat::Supports( TSharedRef<FPropertyNode> PropertyNode )
 
 FPropertyAccess::Result FPropertyHandleFloat::GetValue( float& OutValue ) const
 {
-	void* PropValue = NULL;
+	void* PropValue = nullptr;
 	FPropertyAccess::Result Res = Implementation->GetValueData( PropValue );
 
 	if( Res == FPropertyAccess::Success )
@@ -2318,7 +2318,7 @@ bool FPropertyHandleBool::Supports( TSharedRef<FPropertyNode> PropertyNode )
 {
 	UProperty* Property = PropertyNode->GetProperty();
 
-	if ( Property == NULL )
+	if ( Property == nullptr )
 	{
 		return false;
 	}
@@ -2328,7 +2328,7 @@ bool FPropertyHandleBool::Supports( TSharedRef<FPropertyNode> PropertyNode )
 
 FPropertyAccess::Result FPropertyHandleBool::GetValue( bool& OutValue ) const
 {
-	void* PropValue = NULL;
+	void* PropValue = nullptr;
 	FPropertyAccess::Result Res = Implementation->GetValueData( PropValue );
 
 	if( Res == FPropertyAccess::Success )
@@ -2363,7 +2363,7 @@ bool FPropertyHandleByte::Supports( TSharedRef<FPropertyNode> PropertyNode )
 {
 	UProperty* Property = PropertyNode->GetProperty();
 
-	if ( Property == NULL )
+	if ( Property == nullptr )
 	{
 		return false;
 	}
@@ -2373,7 +2373,7 @@ bool FPropertyHandleByte::Supports( TSharedRef<FPropertyNode> PropertyNode )
 
 FPropertyAccess::Result FPropertyHandleByte::GetValue( uint8& OutValue ) const
 {
-	void* PropValue = NULL;
+	void* PropValue = nullptr;
 	FPropertyAccess::Result Res = Implementation->GetValueData( PropValue );
 
 	if( Res == FPropertyAccess::Success )
@@ -2411,7 +2411,7 @@ bool FPropertyHandleString::Supports( TSharedRef<FPropertyNode> PropertyNode )
 {
 	UProperty* Property = PropertyNode->GetProperty();
 
-	if ( Property == NULL )
+	if ( Property == nullptr )
 	{
 		return false;
 	}
@@ -2435,7 +2435,7 @@ FPropertyAccess::Result FPropertyHandleString::SetValue( const FString& NewValue
 
 FPropertyAccess::Result FPropertyHandleString::GetValue( FName& OutValue ) const
 {
-	void* PropValue = NULL;
+	void* PropValue = nullptr;
 	FPropertyAccess::Result Res = Implementation->GetValueData( PropValue );
 
 	if( Res == FPropertyAccess::Success )
@@ -2457,7 +2457,7 @@ bool FPropertyHandleObject::Supports( TSharedRef<FPropertyNode> PropertyNode )
 {
 	UProperty* Property = PropertyNode->GetProperty();
 
-	if ( Property == NULL )
+	if ( Property == nullptr )
 	{
 		return false;
 	}
@@ -2467,7 +2467,7 @@ bool FPropertyHandleObject::Supports( TSharedRef<FPropertyNode> PropertyNode )
 
 FPropertyAccess::Result FPropertyHandleObject::GetValue( UObject*& OutValue ) const
 {
-	void* PropValue = NULL;
+	void* PropValue = nullptr;
 	FPropertyAccess::Result Res = Implementation->GetValueData( PropValue );
 
 	if( Res == FPropertyAccess::Success )
@@ -2532,7 +2532,7 @@ bool FPropertyHandleVector::Supports( TSharedRef<FPropertyNode> PropertyNode )
 {
 	UProperty* Property = PropertyNode->GetProperty();
 
-	if ( Property == NULL )
+	if ( Property == nullptr )
 	{
 		return false;
 	}
@@ -2556,20 +2556,21 @@ bool FPropertyHandleVector::Supports( TSharedRef<FPropertyNode> PropertyNode )
 FPropertyHandleVector::FPropertyHandleVector( TSharedRef<class FPropertyNode> PropertyNode, class FNotifyHook* NotifyHook, TSharedPtr<IPropertyUtilities> PropertyUtilities )
 	: FPropertyHandleBase( PropertyNode, NotifyHook, PropertyUtilities ) 
 {
+	const bool bRecurse = false;
 	// A vector is a struct property that has 3 children.  We get/set the values from the children
-	VectorComponents.Add( MakeShareable( new FPropertyHandleFloat( Implementation->GetChildNode("X").ToSharedRef(), NotifyHook, PropertyUtilities ) ) );
+	VectorComponents.Add( MakeShareable( new FPropertyHandleFloat( Implementation->GetChildNode("X", bRecurse).ToSharedRef(), NotifyHook, PropertyUtilities ) ) );
 
-	VectorComponents.Add( MakeShareable( new FPropertyHandleFloat( Implementation->GetChildNode("Y").ToSharedRef(), NotifyHook, PropertyUtilities ) ) );
+	VectorComponents.Add( MakeShareable( new FPropertyHandleFloat( Implementation->GetChildNode("Y", bRecurse).ToSharedRef(), NotifyHook, PropertyUtilities ) ) );
 
 	if( Implementation->GetNumChildren() > 2 )
 	{
 		// at least a 3 component vector
-		VectorComponents.Add( MakeShareable( new FPropertyHandleFloat( Implementation->GetChildNode("Z").ToSharedRef(), NotifyHook, PropertyUtilities ) ) );
+		VectorComponents.Add( MakeShareable( new FPropertyHandleFloat( Implementation->GetChildNode("Z",bRecurse).ToSharedRef(), NotifyHook, PropertyUtilities ) ) );
 	}
 	if( Implementation->GetNumChildren() > 3 )
 	{
 		// a 4 component vector
-		VectorComponents.Add( MakeShareable( new FPropertyHandleFloat( Implementation->GetChildNode("W").ToSharedRef(), NotifyHook, PropertyUtilities ) ) );
+		VectorComponents.Add( MakeShareable( new FPropertyHandleFloat( Implementation->GetChildNode("W",bRecurse).ToSharedRef(), NotifyHook, PropertyUtilities ) ) );
 	}
 }
 
@@ -2785,7 +2786,7 @@ bool FPropertyHandleRotator::Supports( TSharedRef<FPropertyNode> PropertyNode )
 {
 	UProperty* Property = PropertyNode->GetProperty();
 
-	if ( Property == NULL )
+	if ( Property == nullptr )
 	{
 		return false;
 	}
@@ -2797,12 +2798,13 @@ bool FPropertyHandleRotator::Supports( TSharedRef<FPropertyNode> PropertyNode )
 FPropertyHandleRotator::FPropertyHandleRotator( TSharedRef<class FPropertyNode> PropertyNode, FNotifyHook* NotifyHook, TSharedPtr<IPropertyUtilities> PropertyUtilities )
 	: FPropertyHandleBase( PropertyNode, NotifyHook, PropertyUtilities ) 
 {
+	const bool bRecurse = false;
 	// A vector is a struct property that has 3 children.  We get/set the values from the children
-	RollValue = MakeShareable( new FPropertyHandleFloat( Implementation->GetChildNode("Roll").ToSharedRef(), NotifyHook, PropertyUtilities ) );
+	RollValue = MakeShareable( new FPropertyHandleFloat( Implementation->GetChildNode("Roll", bRecurse).ToSharedRef(), NotifyHook, PropertyUtilities ) );
 
-	PitchValue = MakeShareable( new FPropertyHandleFloat( Implementation->GetChildNode("Pitch").ToSharedRef(), NotifyHook, PropertyUtilities ) );
+	PitchValue = MakeShareable( new FPropertyHandleFloat( Implementation->GetChildNode("Pitch", bRecurse).ToSharedRef(), NotifyHook, PropertyUtilities ) );
 
-	YawValue = MakeShareable( new FPropertyHandleFloat( Implementation->GetChildNode("Yaw").ToSharedRef(), NotifyHook, PropertyUtilities ) );
+	YawValue = MakeShareable( new FPropertyHandleFloat( Implementation->GetChildNode("Yaw", bRecurse).ToSharedRef(), NotifyHook, PropertyUtilities ) );
 }
 
 
diff --git a/Engine/Source/Editor/PropertyEditor/Private/PropertyHandleImpl.h b/Engine/Source/Editor/PropertyEditor/Private/PropertyHandleImpl.h
--- a/Engine/Source/Editor/PropertyEditor/Private/PropertyHandleImpl.h
+++ b/Engine/Source/Editor/PropertyEditor/Private/PropertyHandleImpl.h
@@ -6,8 +6,8 @@ class FObjectBaseAddress
 {
 public:
 	FObjectBaseAddress()
-		:	Object( NULL )
-		,	BaseAddress( NULL )
+		:	Object( nullptr )
+		,	BaseAddress( nullptr )
 		,	bIsStruct(false)
 	{}
 	FObjectBaseAddress(UObject* InObject, uint8* InBaseAddress, bool InIsStruct)
@@ -207,15 +207,15 @@ public:
 	 * Gets a child node of the property node
 	 * 
 	 * @param ChildName	The name of the child to get
-	 * @return The child property node or NULL if it doesnt exist
+	 * @return The child property node or nullptr if it doesnt exist
 	 */
-	TSharedPtr<FPropertyNode> GetChildNode( FName ChildName ) const;
+	TSharedPtr<FPropertyNode> GetChildNode( FName ChildName, bool bRecurse = true ) const;
 
 	/**
 	 * Gets a child node of the property node
 	 * 
 	 * @param ChildIndex	The child index where the child is stored
-	 * @return The child property node or NULL if it doesnt exist
+	 * @return The child property node or nullptr if it doesnt exist
 	 */
 	TSharedPtr<FPropertyNode> GetChildNode( int32 ChildIndex ) const ;
 
@@ -379,14 +379,14 @@ public:
 	virtual FPropertyAccess::Result GetValueAsDisplayText( FText& OutValue ) const override;
 	virtual FPropertyAccess::Result SetValueFromFormattedString( const FString& InValue,  EPropertyValueSetFlags::Type Flags = EPropertyValueSetFlags::DefaultFlags ) override;
 	virtual TSharedPtr<IPropertyHandle> GetChildHandle( uint32 ChildIndex ) const override;
-	virtual TSharedPtr<IPropertyHandle> GetChildHandle( FName ChildName ) const override;
+	virtual TSharedPtr<IPropertyHandle> GetChildHandle( FName ChildName, bool bRecurse = true ) const override;
 	virtual TSharedPtr<IPropertyHandle> GetParentHandle() const override;
 	virtual void AccessRawData( TArray<void*>& RawData ) override;
 	virtual void AccessRawData( TArray<const void*>& RawData ) const override;
 	virtual uint32 GetNumOuterObjects() const override;
 	virtual void GetOuterObjects( TArray<UObject*>& OuterObjects ) const override;
 	virtual FPropertyAccess::Result GetNumChildren( uint32& OutNumChildren ) const override;
-	virtual TSharedPtr<IPropertyHandleArray> AsArray() override { return NULL; }
+	virtual TSharedPtr<IPropertyHandleArray> AsArray() override { return nullptr; }
 	virtual const UClass* GetPropertyClass() const override;
 	virtual UProperty* GetProperty() const override;
 	virtual UProperty* GetMetaDataProperty() const override;
diff --git a/Engine/Source/Editor/PropertyEditor/Public/PropertyHandle.h b/Engine/Source/Editor/PropertyEditor/Public/PropertyHandle.h
--- a/Engine/Source/Editor/PropertyEditor/Public/PropertyHandle.h
+++ b/Engine/Source/Editor/PropertyEditor/Public/PropertyHandle.h
@@ -263,10 +263,11 @@ public:
 	 * Gets a child handle of this handle.  Useful for accessing properties in structs.  
 	 * Array elements cannot be accessed in this way  
 	 *
-	 * @param The name of the child
+	 * @param ChildName The name of the child
+	 * @param bRecurse	Whether or not to recurse into children of children and so on. If false will only search all immediate children
 	 * @return The property handle for the child if it exists
 	 */
-	virtual TSharedPtr<IPropertyHandle> GetChildHandle( FName ChildName ) const = 0;
+	virtual TSharedPtr<IPropertyHandle> GetChildHandle( FName ChildName, bool bRecurse = true ) const = 0;
 
 	/**
 	 * Gets a child handle of this handle.  Useful for accessing properties in structs.  
