Fix customized external property nodes not remembering expansion state

diff --git a/Engine/Source/Editor/PropertyEditor/Private/SDetailsView.cpp b/Engine/Source/Editor/PropertyEditor/Private/SDetailsView.cpp
--- a/Engine/Source/Editor/PropertyEditor/Private/SDetailsView.cpp
+++ b/Engine/Source/Editor/PropertyEditor/Private/SDetailsView.cpp
@@ -25,7 +25,11 @@
 
 SDetailsView::~SDetailsView()
 {
-	SaveExpandedItems();
+	auto RootNode = GetRootNode();
+	if( RootNode.IsValid() )
+	{
+		SaveExpandedItems( RootNode.ToSharedRef() );
+	}
 };
 
 /**
@@ -564,10 +568,21 @@ void SDetailsView::RemoveDeletedObjects( const TArray<UObject*>& DeletedObjects
 /** Called before during SetObjectArray before we change the objects being observed */
 void SDetailsView::PreSetObject()
 {
-	ExternalRootPropertyNodes.Empty();
-
 	// Save existing expanded items first
-	SaveExpandedItems();
+	if( GetRootNode().IsValid() )
+	{
+		SaveExpandedItems( GetRootNode().ToSharedRef() );
+	}
+
+	for( auto ExternalRootNode : ExternalRootPropertyNodes )
+	{
+		if( ExternalRootNode.IsValid() )
+		{
+			SaveExpandedItems( ExternalRootNode.Pin().ToSharedRef() );
+		}
+	}
+
+	ExternalRootPropertyNodes.Empty();
 
 	RootNodePendingKill = RootPropertyNode;
 
@@ -612,9 +627,23 @@ void SDetailsView::PostSetObject()
 	bool bInitiallySeen = true;
 	bool bParentAllowsVisible = true;
 	// Restore existing expanded items
-	RestoreExpandedItems();
+
+	if( GetRootNode().IsValid() )
+	{
+		RestoreExpandedItems( GetRootNode().ToSharedRef() );
+	}
 
 	UpdatePropertyMap();
+
+	for( auto ExternalRootNode : ExternalRootPropertyNodes )
+	{
+		if( ExternalRootNode.IsValid() )
+		{
+			RestoreExpandedItems( ExternalRootNode.Pin().ToSharedRef() );
+		}
+	}
+
+	UpdateFilteredDetails();
 }
 
 void SDetailsView::SetOnObjectArrayChanged(FOnObjectArrayChanged OnObjectArrayChangedDelegate)
diff --git a/Engine/Source/Editor/PropertyEditor/Private/SDetailsViewBase.cpp b/Engine/Source/Editor/PropertyEditor/Private/SDetailsViewBase.cpp
--- a/Engine/Source/Editor/PropertyEditor/Private/SDetailsViewBase.cpp
+++ b/Engine/Source/Editor/PropertyEditor/Private/SDetailsViewBase.cpp
@@ -696,7 +696,7 @@ void SDetailsViewBase::Tick( const FGeometry& AllottedGeometry, const double InC
 	{
 		// Any deferred actions are likely to cause the node  tree to be at least partially rebuilt
 		// Save the expansion state of existing nodes so we can expand them later
-		SaveExpandedItems();
+		SaveExpandedItems( RootPropertyNode.ToSharedRef() );
 
 		// Execute any deferred actions
 		for (int32 ActionIndex = 0; ActionIndex < DeferredActions.Num(); ++ActionIndex)
@@ -711,7 +711,7 @@ void SDetailsViewBase::Tick( const FGeometry& AllottedGeometry, const double InC
 		// Reaquire the root property node.  It may have been changed by the deferred actions if something like a blueprint editor forcefully resets a details panel during a posteditchange
 		RootPropertyNode = GetRootNode();
 
-		RestoreExpandedItems();
+		RestoreExpandedItems( RootPropertyNode.ToSharedRef() );
 	}
 
 
@@ -719,12 +719,13 @@ void SDetailsViewBase::Tick( const FGeometry& AllottedGeometry, const double InC
 	FPropertyNode::DataValidationResult Result = RootPropertyNode->EnsureDataIsValid();
 	if (Result == FPropertyNode::PropertiesChanged || Result == FPropertyNode::EditInlineNewValueChanged)
 	{
-		RestoreExpandedItems();
+		RestoreExpandedItems( RootPropertyNode.ToSharedRef() );
 		UpdatePropertyMap();
+		UpdateFilteredDetails();
 	}
 	else if (Result == FPropertyNode::ArraySizeChanged)
 	{
-		RestoreExpandedItems();
+		RestoreExpandedItems( RootPropertyNode.ToSharedRef() );
 		UpdateFilteredDetails();
 	}
 	else if (Result == FPropertyNode::ObjectInvalid)
@@ -746,15 +747,16 @@ void SDetailsViewBase::Tick( const FGeometry& AllottedGeometry, const double InC
 				Result = PropertyNode->EnsureDataIsValid();
 				if (Result == FPropertyNode::PropertiesChanged || Result == FPropertyNode::EditInlineNewValueChanged)
 				{
-					RestoreExpandedItems(PropertyNode);
+					RestoreExpandedItems(PropertyNode.ToSharedRef());
 					UpdatePropertyMap();
+					UpdateFilteredDetails();
 					// Note this will invalidate all the external root nodes so there is no need to continue
 					ExternalRootPropertyNodes.Empty();
 					break;
 				}
 				else if (Result == FPropertyNode::ArraySizeChanged)
 				{
-					RestoreExpandedItems(PropertyNode);
+					RestoreExpandedItems(PropertyNode.ToSharedRef());
 					UpdateFilteredDetails();
 				}
 			}
@@ -848,14 +850,12 @@ void SetExpandedItems(TSharedPtr<FPropertyNode> InPropertyNode, const TArray<FSt
 	}
 }
 
-void SDetailsViewBase::SaveExpandedItems()
+void SDetailsViewBase::SaveExpandedItems( TSharedRef<FPropertyNode> StartNode )
 {
-	auto RootPropertyNode = GetRootNode();
-	check(RootPropertyNode.IsValid());
-	UStruct* BestBaseStruct = RootPropertyNode->GetBaseStructure();
+	UStruct* BestBaseStruct = StartNode->FindComplexParent()->GetBaseStructure();
 
 	TArray<FString> ExpandedPropertyItems;
-	GetExpandedItems(RootPropertyNode, ExpandedPropertyItems);
+	GetExpandedItems(StartNode, ExpandedPropertyItems);
 
 	// Handle spaces in expanded node names by wrapping them in quotes
 	for( FString& String : ExpandedPropertyItems )
@@ -878,7 +878,7 @@ void SDetailsViewBase::SaveExpandedItems()
 	//while a valid class, and we're either the same as the base class (for multiple actors being selected and base class is AActor) OR we're not down to AActor yet)
 	for (UStruct* Struct = BestBaseStruct; Struct && ((BestBaseStruct == Struct) || (Struct != AActor::StaticClass())); Struct = Struct->GetSuperStruct())
 	{
-		if (RootPropertyNode->GetNumChildNodes() > 0)
+		if (StartNode->GetNumChildNodes() > 0)
 		{
 			bool bShouldSave = ExpandedPropertyItems.Num() > 0;
 			if (!bShouldSave)
@@ -911,21 +911,17 @@ void SDetailsViewBase::SaveExpandedItems()
 	}
 }
 
-void SDetailsViewBase::RestoreExpandedItems(TSharedPtr<FPropertyNode> InitialStartNode)
+void SDetailsViewBase::RestoreExpandedItems(TSharedRef<FPropertyNode> InitialStartNode)
 {
-	auto RootPropertyNode = GetRootNode();
-	check(RootPropertyNode.IsValid());
 	TSharedPtr<FPropertyNode> StartNode = InitialStartNode;
-	if (!StartNode.IsValid())
-	{
-		StartNode = RootPropertyNode;
-	}
 
 	ExpandedDetailNodes.Empty();
 
 	TArray<FString> ExpandedPropertyItems;
 	FString ExpandedCustomItems;
-	UStruct* BestBaseStruct = RootPropertyNode->GetBaseStructure();
+
+	UStruct* BestBaseStruct = StartNode->FindComplexParent()->GetBaseStructure();
+
 	//while a valid class, and we're either the same as the base class (for multiple actors being selected and base class is AActor) OR we're not down to AActor yet)
 	for (UStruct* Struct = BestBaseStruct; Struct && ((BestBaseStruct == Struct) || (Struct != AActor::StaticClass())); Struct = Struct->GetSuperStruct())
 	{
@@ -1193,7 +1189,5 @@ void SDetailsViewBase::UpdatePropertyMap()
 	}
 	
 	DetailLayout->GenerateDetailLayout();
-
-	UpdateFilteredDetails();
 }
 
diff --git a/Engine/Source/Editor/PropertyEditor/Private/SDetailsViewBase.h b/Engine/Source/Editor/PropertyEditor/Private/SDetailsViewBase.h
--- a/Engine/Source/Editor/PropertyEditor/Private/SDetailsViewBase.h
+++ b/Engine/Source/Editor/PropertyEditor/Private/SDetailsViewBase.h
@@ -278,14 +278,12 @@ public:
 	// End of SWidget interface
 
 	/** Saves the expansion state of property nodes for the selected object set */
-	void SaveExpandedItems();
+	void SaveExpandedItems( TSharedRef<FPropertyNode> StartNode );
 
 	/**
 	* Restores the expansion state of property nodes for the selected object set
-	*
-	* @param InitialStartNode The starting node if any.  If one is not supplied the expansion state is restored from the root node
 	*/
-	void RestoreExpandedItems(TSharedPtr<FPropertyNode> InitialStartNode = NULL);
+	void RestoreExpandedItems(TSharedRef<FPropertyNode> StartNode);
 
 	virtual TSharedPtr<FComplexPropertyNode> GetRootNode() = 0;
 
diff --git a/Engine/Source/Editor/PropertyEditor/Private/SStructureDetailsView.cpp b/Engine/Source/Editor/PropertyEditor/Private/SStructureDetailsView.cpp
--- a/Engine/Source/Editor/PropertyEditor/Private/SStructureDetailsView.cpp
+++ b/Engine/Source/Editor/PropertyEditor/Private/SStructureDetailsView.cpp
@@ -19,7 +19,12 @@
 
 SStructureDetailsView::~SStructureDetailsView()
 {
-	SaveExpandedItems();
+	auto RootNode = GetRootNode();
+
+	if( RootNode.IsValid() )
+	{
+		SaveExpandedItems( RootNode.ToSharedRef() );
+	}
 }
 
 UStruct* SStructureDetailsView::GetBaseScriptStruct() const
@@ -170,7 +175,7 @@ void SStructureDetailsView::Construct(const FArguments& InArgs)
 void SStructureDetailsView::SetStructureData(TSharedPtr<FStructOnScope> InStructData)
 {
 	//PRE SET
-	SaveExpandedItems();
+	SaveExpandedItems( RootNode.ToSharedRef() );
 	RootNode->SetStructure(NULL);
 	RootNodePendingKill = RootNode;
 	RootNode = MakeShareable(new FStructurePropertyNode);
@@ -199,9 +204,11 @@ void SStructureDetailsView::SetStructureData(TSharedPtr<FStructOnScope> InStruct
 	RootNode->InitNode(InitParams);
 	RootNode->SetDisplayNameOverride(CustomName);
 
-	RestoreExpandedItems();
+	RestoreExpandedItems(RootNode.ToSharedRef());
 
 	UpdatePropertyMap();
+
+	UpdateFilteredDetails();
 }
 
 void SStructureDetailsView::ForceRefresh()
