Fix shadowed variable CIS

diff --git a/Engine/Source/Editor/PropertyEditor/Private/SStructureDetailsView.cpp b/Engine/Source/Editor/PropertyEditor/Private/SStructureDetailsView.cpp
--- a/Engine/Source/Editor/PropertyEditor/Private/SStructureDetailsView.cpp
+++ b/Engine/Source/Editor/PropertyEditor/Private/SStructureDetailsView.cpp
@@ -19,11 +19,11 @@
 
 SStructureDetailsView::~SStructureDetailsView()
 {
-	auto RootNode = GetRootNode();
+	auto RootNodeLocal = GetRootNode();
 
-	if( RootNode.IsValid() )
+	if (RootNodeLocal.IsValid())
 	{
-		SaveExpandedItems( RootNode.ToSharedRef() );
+		SaveExpandedItems(RootNodeLocal.ToSharedRef());
 	}
 }
 
