Allow creation of default property value widgets to specify whether or not the default property buttons are displayed.

diff --git a/Engine/Source/Editor/PropertyEditor/Private/PropertyHandleImpl.cpp b/Engine/Source/Editor/PropertyEditor/Private/PropertyHandleImpl.cpp
--- a/Engine/Source/Editor/PropertyEditor/Private/PropertyHandleImpl.cpp
+++ b/Engine/Source/Editor/PropertyEditor/Private/PropertyHandleImpl.cpp
@@ -1550,13 +1550,14 @@ TSharedRef<SWidget> FPropertyHandleBase::CreatePropertyNameWidget( const FText&
 	return SNullWidget::NullWidget;
 }
 
-TSharedRef<SWidget> FPropertyHandleBase::CreatePropertyValueWidget() const
+TSharedRef<SWidget> FPropertyHandleBase::CreatePropertyValueWidget( bool bDisplayDefaultPropertyButtons ) const
 {
 	if( Implementation.IsValid() && Implementation->GetPropertyNode().IsValid() )
 	{
 		TSharedPtr<FPropertyEditor> PropertyEditor = FPropertyEditor::Create( Implementation->GetPropertyNode().ToSharedRef(), Implementation->GetPropertyUtilities().ToSharedRef() );
 
-		return SNew( SPropertyValueWidget, PropertyEditor, Implementation->GetPropertyUtilities() );
+		return SNew( SPropertyValueWidget, PropertyEditor, Implementation->GetPropertyUtilities() )
+				.ShowPropertyButtons( bDisplayDefaultPropertyButtons );
 	}
 
 	return SNullWidget::NullWidget;
diff --git a/Engine/Source/Editor/PropertyEditor/Private/PropertyHandleImpl.h b/Engine/Source/Editor/PropertyEditor/Private/PropertyHandleImpl.h
--- a/Engine/Source/Editor/PropertyEditor/Private/PropertyHandleImpl.h
+++ b/Engine/Source/Editor/PropertyEditor/Private/PropertyHandleImpl.h
@@ -369,7 +369,7 @@ public:
 	virtual bool IsCustomized() const override;
 	virtual FString GeneratePathToProperty() const override;
 	virtual TSharedRef<SWidget> CreatePropertyNameWidget( const FText& NameOverride = FText::GetEmpty(), const FText& ToolTipOverride = FText::GetEmpty(), bool bDisplayResetToDefault = false, bool bDisplayText = true, bool bDisplayThumbnail = true ) const override;
-	virtual TSharedRef<SWidget> CreatePropertyValueWidget() const override;
+	virtual TSharedRef<SWidget> CreatePropertyValueWidget( bool bDisplayDefaultPropertyButtons = true ) const override;
 	virtual bool IsEditConst() const override;
 	virtual void SetOnPropertyValueChanged( const FSimpleDelegate& InOnPropertyValueChanged ) override;
 	virtual int32 GetIndexInArray() const override;
diff --git a/Engine/Source/Editor/PropertyEditor/Public/PropertyHandle.h b/Engine/Source/Editor/PropertyEditor/Public/PropertyHandle.h
--- a/Engine/Source/Editor/PropertyEditor/Public/PropertyHandle.h
+++ b/Engine/Source/Editor/PropertyEditor/Public/PropertyHandle.h
@@ -372,7 +372,7 @@ public:
 
 	 * @return the value widget for this property
 	 */
-	virtual TSharedRef<SWidget> CreatePropertyValueWidget() const = 0;
+	virtual TSharedRef<SWidget> CreatePropertyValueWidget( bool bDisplayDefaultPropertyButtons = true ) const = 0;
 
 	/**
 	 * Adds a restriction to the possible values for this property.
