Allow detail panel categories to be dynamically shown or hidden.  Added this to the material customization for primitive components.  Fixes issue where material customization never appears if there was no valid material on a component when the details panel was created

diff --git a/Engine/Source/Editor/DetailCustomizations/Private/ComponentMaterialCategory.cpp b/Engine/Source/Editor/DetailCustomizations/Private/ComponentMaterialCategory.cpp
--- a/Engine/Source/Editor/DetailCustomizations/Private/ComponentMaterialCategory.cpp
+++ b/Engine/Source/Editor/DetailCustomizations/Private/ComponentMaterialCategory.cpp
@@ -145,6 +145,7 @@ private:
 FComponentMaterialCategory::FComponentMaterialCategory( TArray< TWeakObjectPtr<USceneComponent> >& InSelectedComponents )
 	: SelectedComponents( InSelectedComponents )
 	, NotifyHook( nullptr )
+	, MaterialCategory( nullptr )
 {
 }
 
@@ -172,14 +173,12 @@ void FComponentMaterialCategory::Create( IDetailLayoutBuilder& DetailBuilder )
 	}
 
 
-	// only show the category if there are materials to display
-	if( bAnyMaterialsToDisplay )
-	{
-		// Make a category for the materials.
-		IDetailCategoryBuilder& MaterialCategory = DetailBuilder.EditCategory("Materials", FText::GetEmpty(), ECategoryPriority::TypeSpecific );
+	// Make a category for the materials.
+	MaterialCategory = &DetailBuilder.EditCategory("Materials", FText::GetEmpty(), ECategoryPriority::TypeSpecific );
 
-		MaterialCategory.AddCustomBuilder( MaterialList );
-	}
+	MaterialCategory->AddCustomBuilder( MaterialList );
+	MaterialCategory->SetCategoryVisibility( bAnyMaterialsToDisplay );
+	
 }
 
 void FComponentMaterialCategory::OnGetMaterialsForView( IMaterialListBuilder& MaterialList )
@@ -191,6 +190,8 @@ void FComponentMaterialCategory::OnGetMaterialsForView( IMaterialListBuilder& Ma
 	{	
 		int32 MaterialIndex = It.GetMaterialIndex();
 
+		MaterialCategory->SetCategoryVisibility( true );
+
 		UActorComponent* CurrentComponent = It.GetComponent();
 
 		if( CurrentComponent )
diff --git a/Engine/Source/Editor/DetailCustomizations/Private/ComponentMaterialCategory.h b/Engine/Source/Editor/DetailCustomizations/Private/ComponentMaterialCategory.h
--- a/Engine/Source/Editor/DetailCustomizations/Private/ComponentMaterialCategory.h
+++ b/Engine/Source/Editor/DetailCustomizations/Private/ComponentMaterialCategory.h
@@ -46,4 +46,6 @@ private:
 
 	/** Notify hook to use */
 	FNotifyHook* NotifyHook;
+
+	IDetailCategoryBuilder* MaterialCategory;
 };
\ No newline at end of file
diff --git a/Engine/Source/Editor/PropertyEditor/Private/DetailCategoryBuilderImpl.cpp b/Engine/Source/Editor/PropertyEditor/Private/DetailCategoryBuilderImpl.cpp
--- a/Engine/Source/Editor/PropertyEditor/Private/DetailCategoryBuilderImpl.cpp
+++ b/Engine/Source/Editor/PropertyEditor/Private/DetailCategoryBuilderImpl.cpp
@@ -86,6 +86,7 @@ FDetailCategoryImpl::FDetailCategoryImpl( FName InCategoryName, TSharedRef<FDeta
 	, bForceAdvanced( false )
 	, bHasFilterStrings( false )
 	, bHasVisibleDetails( true )
+	, bIsCategoryVisible( true )
 {
 	const UStruct* BaseStruct = InDetailLayout->GetDetailsView().GetBaseStruct();
 	// Use the base class name if there is one otherwise this is a generic category not specific to a class
@@ -93,7 +94,10 @@ FDetailCategoryImpl::FDetailCategoryImpl( FName InCategoryName, TSharedRef<FDeta
 
 	CategoryPathName = BaseStructName.ToString() + TEXT(".") + CategoryName.ToString();
 
-	GConfig->GetBool( TEXT("DetailCategoriesAdvanced"), *CategoryPathName, bUserShowAdvanced, GEditorPerProjectIni );
+	bool bUserShowAdvancedConfigValue = false;
+	GConfig->GetBool( TEXT("DetailCategoriesAdvanced"), *CategoryPathName, bUserShowAdvancedConfigValue, GEditorPerProjectIni );
+
+	bUserShowAdvanced = bUserShowAdvancedConfigValue;
 
 }
 
@@ -176,6 +180,15 @@ void FDetailCategoryImpl::GetDefaultProperties( TArray<TSharedRef<IPropertyHandl
 	}
 }
 
+void FDetailCategoryImpl::SetCategoryVisibility( bool bIsVisible )
+{
+	if( bIsVisible != bIsCategoryVisible )
+	{
+		bIsCategoryVisible = bIsVisible;
+
+		GetDetailsView().RerunCurrentFilter();
+	}
+}
 
 IDetailCategoryBuilder& FDetailCategoryImpl::InitiallyCollapsed( bool bInShouldBeInitiallyCollapsed )
 {
@@ -574,7 +587,7 @@ bool FDetailCategoryImpl::ShouldBeExpanded() const
 
 ENodeVisibility::Type FDetailCategoryImpl::GetVisibility() const
 {
-	return bHasVisibleDetails ? ENodeVisibility::Visible : ENodeVisibility::ForcedHidden;
+	return bHasVisibleDetails && bIsCategoryVisible ? ENodeVisibility::Visible : ENodeVisibility::ForcedHidden;
 }
 
 bool IsCustomProperty( const TSharedPtr<FPropertyNode>& PropertyNode )
diff --git a/Engine/Source/Editor/PropertyEditor/Private/DetailCategoryBuilderImpl.h b/Engine/Source/Editor/PropertyEditor/Private/DetailCategoryBuilderImpl.h
--- a/Engine/Source/Editor/PropertyEditor/Private/DetailCategoryBuilderImpl.h
+++ b/Engine/Source/Editor/PropertyEditor/Private/DetailCategoryBuilderImpl.h
@@ -144,6 +144,7 @@ public:
 	virtual IDetailGroup& AddGroup( FName GroupName, const FText& LocalizedDisplayName, bool bForAdvanced = false, bool bStartExpanded = false ) override;
 	virtual void GetDefaultProperties( TArray<TSharedRef<IPropertyHandle> >& OutAllProperties, bool bSimpleProperties = true, bool bAdvancedProperties = true ) override;
 	virtual const FText& GetDisplayName() const override { return DisplayName; }
+	virtual void SetCategoryVisibility( bool bIsVisible ) override;
 
 	/** IDetailTreeNode interface */
 	virtual IDetailsViewPrivate& GetDetailsView() const override{ return DetailLayoutBuilder.Pin()->GetDetailsView(); }
@@ -359,15 +360,17 @@ private:
 	/** The sort order of this category (amongst all categories) */
 	int32 SortOrder;
 	/** Whether or not to restore the expansion state between sessions */
-	bool bRestoreExpansionState;
+	bool bRestoreExpansionState:1;
 	/** Whether or not the category should be initially collapsed */
-	bool bShouldBeInitiallyCollapsed;
+	bool bShouldBeInitiallyCollapsed:1;
 	/** Whether or not advanced properties should be shown (as specified by the user) */
-	bool bUserShowAdvanced;
+	bool bUserShowAdvanced:1;
 	/** Whether or not advanced properties are forced to be shown (this is an independent toggle from bShowAdvanced which is user driven)*/
-	bool bForceAdvanced;
+	bool bForceAdvanced:1;
 	/** Whether or not the content in the category is being filtered */
-	bool bHasFilterStrings;
+	bool bHasFilterStrings:1;
 	/** true if anything is visible in the category */
-	bool bHasVisibleDetails;
+	bool bHasVisibleDetails:1;
+	/** true if the category is visible at all */
+	bool bIsCategoryVisible:1;
 };
diff --git a/Engine/Source/Editor/PropertyEditor/Private/IDetailsViewPrivate.h b/Engine/Source/Editor/PropertyEditor/Private/IDetailsViewPrivate.h
--- a/Engine/Source/Editor/PropertyEditor/Private/IDetailsViewPrivate.h
+++ b/Engine/Source/Editor/PropertyEditor/Private/IDetailsViewPrivate.h
@@ -116,6 +116,11 @@ public:
 	 */
 	virtual void NotifyFinishedChangingProperties(const FPropertyChangedEvent& PropertyChangedEvent) = 0;
 
+	/**
+	 * Reruns the current search filter applied to the details panel to account for any new changes                                                              
+	 */
+	virtual void RerunCurrentFilter() = 0;
+
 	/** If a customization standalone widget is used, the value should be update only once, when its window is closed */
 	virtual bool DontUpdateValueWhileEditing() const = 0;
 };
diff --git a/Engine/Source/Editor/PropertyEditor/Private/SDetailsViewBase.cpp b/Engine/Source/Editor/PropertyEditor/Private/SDetailsViewBase.cpp
--- a/Engine/Source/Editor/PropertyEditor/Private/SDetailsViewBase.cpp
+++ b/Engine/Source/Editor/PropertyEditor/Private/SDetailsViewBase.cpp
@@ -178,6 +178,11 @@ void SDetailsViewBase::SetOnDisplayedPropertiesChanged(FOnDisplayedPropertiesCha
 	OnDisplayedPropertiesChangedDelegate = InOnDisplayedPropertiesChangedDelegate;
 }
 
+void SDetailsViewBase::RerunCurrentFilter()
+{
+	UpdateFilteredDetails();
+}
+
 EVisibility SDetailsViewBase::GetTreeVisibility() const
 {
 	return DetailLayout.IsValid() && DetailLayout->HasDetails() ? EVisibility::Visible : EVisibility::Collapsed;
diff --git a/Engine/Source/Editor/PropertyEditor/Private/SDetailsViewBase.h b/Engine/Source/Editor/PropertyEditor/Private/SDetailsViewBase.h
--- a/Engine/Source/Editor/PropertyEditor/Private/SDetailsViewBase.h
+++ b/Engine/Source/Editor/PropertyEditor/Private/SDetailsViewBase.h
@@ -113,13 +113,14 @@ public:
 	virtual void HideFilterArea(bool bHide) override;
 
 	/** 
-	 * Implementation of IDetailsView:
+	 * IDetailsView interface 
 	 */
 	virtual TArray< FPropertyPath > GetPropertiesInOrderDisplayed() const override;
 	virtual void HighlightProperty(const FPropertyPath& Property) override;
 	virtual void ShowAllAdvancedProperties() override;
 	virtual void SetOnDisplayedPropertiesChanged(FOnDisplayedPropertiesChanged InOnDisplayedPropertiesChangedDelegate) override;
 	virtual void SetDisableCustomDetailLayouts( bool bInDisableCustomDetailLayouts ) override { bDisableCustomDetailLayouts = bInDisableCustomDetailLayouts; }
+	virtual void RerunCurrentFilter() override;
 
 	virtual FOnFinishedChangingProperties& OnFinishedChangingProperties() override
 	{ 
diff --git a/Engine/Source/Editor/PropertyEditor/Public/DetailCategoryBuilder.h b/Engine/Source/Editor/PropertyEditor/Public/DetailCategoryBuilder.h
--- a/Engine/Source/Editor/PropertyEditor/Public/DetailCategoryBuilder.h
+++ b/Engine/Source/Editor/PropertyEditor/Public/DetailCategoryBuilder.h
@@ -137,4 +137,10 @@ public:
 	 * @return The localized display name of the category
 	 */
 	virtual const FText& GetDisplayName() const = 0;
+
+	/**
+	 * Sets whether or not this category is hidden or shown
+	 * This is designed to be used for dynamic category visibility after construction of the category
+	 */
+	virtual void SetCategoryVisibility( bool bVisible ) = 0;
 };
\ No newline at end of file
