Duplicating, in a Property Tree, an array of instanced objects will correctly deep copy the objects instead of a shallow (referenced) copy.

diff --git a/Engine/Source/Editor/PropertyEditor/Private/PropertyHandleImpl.cpp b/Engine/Source/Editor/PropertyEditor/Private/PropertyHandleImpl.cpp
--- a/Engine/Source/Editor/PropertyEditor/Private/PropertyHandleImpl.cpp
+++ b/Engine/Source/Editor/PropertyEditor/Private/PropertyHandleImpl.cpp
@@ -1396,10 +1396,22 @@ void FPropertyValueImpl::DuplicateChild( TSharedPtr<FPropertyNode> ChildNodeToDu
 
 		ArrayHelper.InsertValues(Index);
 
-
 		// Copy the selected item's value to the new item.
-		NodeProperty->CopyCompleteValue(ArrayHelper.GetRawPtr(ChildNodePtr->GetArrayIndex()), ArrayHelper.GetRawPtr(ChildNodePtr->GetArrayIndex() + 1));
+		NodeProperty->CopyCompleteValue(ArrayHelper.GetRawPtr(Index), ArrayHelper.GetRawPtr(Index + 1));
+
+		if (UObjectProperty* ObjProp = Cast<UObjectProperty>(NodeProperty))
+		{
+			UObject* CurrentObject = ObjProp->GetObjectPropertyValue(ArrayHelper.GetRawPtr(Index));
 
+			// Fpr DefaultSubObjects and ArchetypeObjects we need to do a deep copy instead of a shallow copy
+			if (CurrentObject && CurrentObject->HasAnyFlags(RF_DefaultSubObject | RF_ArchetypeObject))
+			{
+				// Make a deep copy and assign it into the array.
+				UObject* DuplicatedObject = DuplicateObject(CurrentObject, CurrentObject->GetOuter());
+				ObjProp->SetObjectPropertyValue(ArrayHelper.GetRawPtr(Index + 1), DuplicatedObject);
+			}
+		}
+		
 		// Find the object that owns the array and instance any subobjects
 		if (FObjectPropertyNode* ObjectPropertyNode = ChildNodePtr->FindObjectItemParent())
 		{
