Fixed "Reset to Default" logic error (left from changes in UETOOL-125)
[UE-20444] - Reset to defaults icon is missing

diff --git a/Engine/Source/Editor/PropertyEditor/Private/DetailPropertyRow.cpp b/Engine/Source/Editor/PropertyEditor/Private/DetailPropertyRow.cpp
--- a/Engine/Source/Editor/PropertyEditor/Private/DetailPropertyRow.cpp
+++ b/Engine/Source/Editor/PropertyEditor/Private/DetailPropertyRow.cpp
@@ -466,7 +466,7 @@ void FDetailPropertyRow::MakeValueWidget( FDetailWidgetRow& Row, const TSharedPt
 		[
 			SNew( SResetToDefaultPropertyEditor, PropertyEditor.ToSharedRef() )
 			.IsEnabled( IsEnabledAttrib )
-			.CustomResetToDefault(CustomResetToDefault.Get(FResetToDefaultOverride()))
+			.CustomResetToDefault(CustomResetToDefault)
 		];
 	}
 
