PropertyEditor: Adjusted padding to prevent gap at top

diff --git a/Engine/Source/Editor/PropertyEditor/Private/DetailCategoryGroupNode.cpp b/Engine/Source/Editor/PropertyEditor/Private/DetailCategoryGroupNode.cpp
--- a/Engine/Source/Editor/PropertyEditor/Private/DetailCategoryGroupNode.cpp
+++ b/Engine/Source/Editor/PropertyEditor/Private/DetailCategoryGroupNode.cpp
@@ -12,7 +12,7 @@ void SDetailCategoryTableRow::Construct( const FArguments& InArgs, TSharedRef<ID
 	TSharedPtr<SHorizontalBox> MyContent;
 
 	ChildSlot
-	.Padding( 0.0f, 2.0f, 0.0f, 0.0f )
+	.Padding( 0.0f, 0.0f, 0.0f, 2.0f )
 	[	
 		SNew( SBorder )
 		.BorderImage( this, &SDetailCategoryTableRow::GetBackgroundImage )
