PropertyEditor: Exposing underlying details view from StructureDetailsView; cleanup pass.

diff --git a/Engine/Source/Editor/PropertyEditor/Private/SStructureDetailsView.h b/Engine/Source/Editor/PropertyEditor/Private/SStructureDetailsView.h
--- a/Engine/Source/Editor/PropertyEditor/Private/SStructureDetailsView.h
+++ b/Engine/Source/Editor/PropertyEditor/Private/SStructureDetailsView.h
@@ -13,17 +13,18 @@ public:
 
 	SLATE_BEGIN_ARGS(SStructureDetailsView)
 		: _DetailsViewArgs()
-		{}
-		/** The user defined args for the details view */
+	{ }
+		/** User defined arguments for the details view */
 		SLATE_ARGUMENT(FDetailsViewArgs, DetailsViewArgs)
+
+		/** Custom name for the root property node. */
 		SLATE_ARGUMENT(FText, CustomName)
 	SLATE_END_ARGS()
 
+	/** Destructor. */
 	~SStructureDetailsView();
 
-	/**
-	 * Constructs the property view widgets                   
-	 */
+	/** Constructs the property view widgets. */
 	void Construct(const FArguments& InArgs);
 
 	UStruct* GetBaseScriptStruct() const;
@@ -35,18 +36,31 @@ public:
 		return true;
 	}
 
-	/** IStructureDetailsView interface */
+public:
+
+	// IStructureDetailsView interface
+
+	virtual IDetailsView& GetDetailsView() override
+	{
+		return *this;
+	}
+
 	virtual TSharedPtr<SWidget> GetWidget() override
 	{
 		return AsShared();
 	}
 
 	virtual void SetStructureData(TSharedPtr<FStructOnScope> InStructData) override;
+
 	virtual FOnFinishedChangingProperties& GetOnFinishedChangingPropertiesDelegate() override
 	{
 		return OnFinishedChangingProperties();
 	}
-	/** IDetailsViewPrivate interface */
+
+public:
+
+	// IDetailsViewPrivate interface
+
 	virtual const UClass* GetBaseClass() const override
 	{
 		return NULL;
@@ -65,14 +79,19 @@ public:
 	virtual void ForceRefresh() override;
 	virtual void AddExternalRootPropertyNode(TSharedRef<FPropertyNode> ExternalRootNode) override {}
 
-	/** IDetailsView interface */
+public:
+
+	// IDetailsView interface
+
 	virtual UStruct* GetBaseStruct() const override
 	{
 		return GetBaseScriptStruct();
 	}
+
 	virtual const TArray< TWeakObjectPtr<UObject> >& GetSelectedObjects() const override;
 	virtual const TArray< TWeakObjectPtr<AActor> >& GetSelectedActors() const override;
 	virtual const struct FSelectedActorInfo& GetSelectedActorInfo() const override;
+
 	virtual bool HasClassDefaultObject() const override
 	{
 		return false;
@@ -87,7 +106,8 @@ public:
 	virtual void RemoveInvalidObjects() override {}
 
 	virtual TSharedPtr<class FComplexPropertyNode> GetRootNode() override;
+
 protected:
+
 	virtual void CustomUpdatePropertyMap() override;
 };
-
diff --git a/Engine/Source/Editor/PropertyEditor/Public/IStructureDetailsView.h b/Engine/Source/Editor/PropertyEditor/Public/IStructureDetailsView.h
--- a/Engine/Source/Editor/PropertyEditor/Public/IStructureDetailsView.h
+++ b/Engine/Source/Editor/PropertyEditor/Public/IStructureDetailsView.h
@@ -10,8 +10,34 @@
 class IStructureDetailsView
 {
 public:
+
+	/**
+	 * Get an interface the underlying details view.
+	 *
+	 * @return Details view interface.
+	 */
+	virtual IDetailsView& GetDetailsView() = 0;
+
+	/**
+	 * Get this view's Slate widget.
+	 *
+	 * @return The widget for this view.
+	 */
 	virtual TSharedPtr<class SWidget> GetWidget() = 0;
 
+	/**
+	 * Set the structure data to be displayed in this view.
+	 *
+	 * @param StructData The structure data to display.
+	 */
 	virtual void SetStructureData(TSharedPtr<class FStructOnScope> StructData) = 0;
+
+public:
+
+	/**
+	 * Get a delegate that is executed when the view has finished changing properties.
+	 *
+	 * @return The delegate.
+	 */
 	virtual FOnFinishedChangingProperties& GetOnFinishedChangingPropertiesDelegate() = 0;
 };
