Fixed some places that were surrounding FName properties in quotes, but weren't providing the PPF_Delimited flag

UE-21078 - FBoneReference text import no longer works correctly.

Some property editing code used to forcibly surround FName properties in quotes, as UNameProperty::ImportText_Internal didn't handle the PPF_Delimited correctly prior to CL# 2682467, so this was the only way to ensure that names containing spaces would be imported.

This change removes those places that were forcibly adding quotes, as UNameProperty::ImportText_Internal will now do the correct thing as the default import flags do not contain PPF_Delimited.

diff --git a/Engine/Source/Editor/PropertyEditor/Private/Presentation/PropertyTable/PropertyTable.cpp b/Engine/Source/Editor/PropertyEditor/Private/Presentation/PropertyTable/PropertyTable.cpp
--- a/Engine/Source/Editor/PropertyEditor/Private/Presentation/PropertyTable/PropertyTable.cpp
+++ b/Engine/Source/Editor/PropertyEditor/Private/Presentation/PropertyTable/PropertyTable.cpp
@@ -376,17 +376,15 @@ void SetCellValue( const TSharedRef< IPropertyTableCell >& Cell, FString Value )
 		return;
 	}
 
-	// We need to put quotes around property name strings
+	// We need to sanitize property name strings
 	TSharedPtr<FPropertyNode> PropertyNode = Cell->GetNode();
 	if (PropertyNode.IsValid())
 	{
 		UProperty* NodeProperty = PropertyNode->GetProperty();
 		if (NodeProperty->IsA(UNameProperty::StaticClass()))
 		{
-			// Remove any pre-existing quotes and return characters
-			Value = Value.TrimQuotes().Replace(TEXT("\n"), TEXT("")).Replace(TEXT("\r"), TEXT(""));
-			// Add quotes
-			Value = FString::Printf(TEXT("\"%s\""), *Value);
+			// Remove any pre-existing return characters
+			Value = Value.TrimQuotes().Replace(TEXT("\n"), TEXT(""));
 		}
 	}
 
diff --git a/Engine/Source/Editor/PropertyEditor/Private/PropertyHandleImpl.cpp b/Engine/Source/Editor/PropertyEditor/Private/PropertyHandleImpl.cpp
--- a/Engine/Source/Editor/PropertyEditor/Private/PropertyHandleImpl.cpp
+++ b/Engine/Source/Editor/PropertyEditor/Private/PropertyHandleImpl.cpp
@@ -695,12 +695,6 @@ FPropertyAccess::Result FPropertyValueImpl::SetValueAsString( const FString& InV
 					break;
 				}
 			}
-
-			// Ensure the name is enclosed with quotes.
-			if( Value.Len() )
-			{
-				Value = FString::Printf( TEXT("\"%s\""), *Value.TrimQuotes() );
-			}
 		}
 
 		// If more than one object is selected, an empty field indicates their values for this property differ.
