Crash fix for the details panel

--------
Integrated using branch UE4-To-UE4-Orion (reversed) of change#2704869 by [REDACTED_PERSON] on 2015/09/24 14:35:41.

diff --git a/Engine/Source/Editor/PropertyEditor/Public/DetailWidgetRow.h b/Engine/Source/Editor/PropertyEditor/Public/DetailWidgetRow.h
--- a/Engine/Source/Editor/PropertyEditor/Public/DetailWidgetRow.h
+++ b/Engine/Source/Editor/PropertyEditor/Public/DetailWidgetRow.h
@@ -7,13 +7,14 @@ class FDetailWidgetDecl
 {
 public:
 	FDetailWidgetDecl( class FDetailWidgetRow& InParentDecl, float InMinWidth, float InMaxWidth, EHorizontalAlignment InHAlign, EVerticalAlignment InVAlign )
-		: Widget( SNullWidget::NullWidget )
+		: Widget( SNew( SSpacer) )
 		, HorizontalAlignment( InHAlign )
 		, VerticalAlignment( InVAlign )
 		, MinWidth( InMinWidth )
 		, MaxWidth( InMaxWidth )
 		, ParentDecl( InParentDecl )
 	{
+
 	}
 
 	FDetailWidgetRow& operator[]( TSharedRef<SWidget> InWidget )
