Fix broken details panels caused by CL 2704872

diff --git a/Engine/Source/Editor/PropertyEditor/Public/DetailWidgetRow.h b/Engine/Source/Editor/PropertyEditor/Public/DetailWidgetRow.h
--- a/Engine/Source/Editor/PropertyEditor/Public/DetailWidgetRow.h
+++ b/Engine/Source/Editor/PropertyEditor/Public/DetailWidgetRow.h
@@ -6,8 +6,9 @@
 class FDetailWidgetDecl
 {
 public:
+	
 	FDetailWidgetDecl( class FDetailWidgetRow& InParentDecl, float InMinWidth, float InMaxWidth, EHorizontalAlignment InHAlign, EVerticalAlignment InVAlign )
-		: Widget( SNew( SSpacer) )
+		: Widget( SNew( SInvalidDetailWidget ) )
 		, HorizontalAlignment( InHAlign )
 		, VerticalAlignment( InVAlign )
 		, MinWidth( InMinWidth )
@@ -46,6 +47,19 @@ public:
 		MaxWidth = InMaxWidth;
 		return *this;
 	}
+private:
+	class SInvalidDetailWidget : public SSpacer
+	{
+		SLATE_BEGIN_ARGS( SInvalidDetailWidget )
+		{}
+		SLATE_END_ARGS()
+
+		void Construct( const FArguments& InArgs )
+		{
+			Visibility = EVisibility::Collapsed;
+		}
+
+	};
 public:
 	TSharedRef<SWidget> Widget;
 	EHorizontalAlignment HorizontalAlignment;
@@ -56,6 +70,9 @@ private:
 	class FDetailWidgetRow& ParentDecl;
 };
 
+
+static FName InvalidDetailWidgetName = TEXT("SInvalidDetailWidget");
+
 /**
  * Represents a single row of custom widgets in a details panel 
  */
@@ -158,7 +175,7 @@ public:
 	 */
 	bool HasColumns() const
 	{
-		return NameWidget.Widget != SNullWidget::NullWidget || ValueWidget.Widget != SNullWidget::NullWidget;
+		return NameWidget.Widget->GetType() != InvalidDetailWidgetName || ValueWidget.Widget->GetType() != InvalidDetailWidgetName;
 	}
 
 	/**
@@ -166,7 +183,7 @@ public:
 	 */
 	bool HasAnyContent() const
 	{
-		return WholeRowWidget.Widget != SNullWidget::NullWidget || HasColumns();
+		return WholeRowWidget.Widget->GetType() != InvalidDetailWidgetName || HasColumns();
 	}
 
 	/** @return true if a custom copy/paste is bound on this row */
