Extended the structure detail view asset detection support to always treat "asset" properties as an asset

eg) Anything used by a TAssetPtr

diff --git a/Engine/Source/Editor/PropertyEditor/Private/PropertyEditorModule.cpp b/Engine/Source/Editor/PropertyEditor/Private/PropertyEditorModule.cpp
--- a/Engine/Source/Editor/PropertyEditor/Private/PropertyEditorModule.cpp
+++ b/Engine/Source/Editor/PropertyEditor/Private/PropertyEditorModule.cpp
@@ -688,11 +688,21 @@ TSharedRef<class IStructureDetailsView> FPropertyEditorModule::CreateStructureDe
 			const auto ObjectProperty = Cast<UObjectPropertyBase>(PropertyToTest);
 			if( ObjectProperty )
 			{
-				if( InStructureDetailsViewArgs.bShowAssets && ObjectProperty->PropertyClass )
+				if( InStructureDetailsViewArgs.bShowAssets )
 				{
-					// We can use the asset tools module to see whether this type has asset actions (which likely means it's an asset class type)
-					FAssetToolsModule& AssetToolsModule = FAssetToolsModule::GetModule();
-					return AssetToolsModule.Get().GetAssetTypeActionsForClass(ObjectProperty->PropertyClass).IsValid();
+					// Is this an "asset" property?
+					if( PropertyToTest->IsA<UAssetObjectProperty>())
+					{
+						return true;
+					}
+
+					// Not an "asset" property, but it may still be a property using an asset class type (such as a raw pointer)
+					if( ObjectProperty->PropertyClass )
+					{
+						// We can use the asset tools module to see whether this type has asset actions (which likely means it's an asset class type)
+						FAssetToolsModule& AssetToolsModule = FAssetToolsModule::GetModule();
+						return AssetToolsModule.Get().GetAssetTypeActionsForClass(ObjectProperty->PropertyClass).IsValid();
+					}
 				}
 
 				return InStructureDetailsViewArgs.bShowObjects;
