PropertyEditor: Fixed crash when unregistering details customizations on shutdown

diff --git a/Engine/Source/Editor/PropertyEditor/Private/PropertyEditorModule.cpp b/Engine/Source/Editor/PropertyEditor/Private/PropertyEditorModule.cpp
--- a/Engine/Source/Editor/PropertyEditor/Private/PropertyEditorModule.cpp
+++ b/Engine/Source/Editor/PropertyEditor/Private/PropertyEditorModule.cpp
@@ -391,7 +391,7 @@ void FPropertyEditorModule::RegisterCustomClassLayout( FName ClassName, FOnGetDe
 
 void FPropertyEditorModule::UnregisterCustomClassLayout( FName ClassName )
 {
-	if (ClassName != NAME_None)
+	if (ClassName.IsValid() && (ClassName != NAME_None))
 	{
 		ClassNameToDetailLayoutNameMap.Remove(ClassName);
 	}
@@ -441,30 +441,32 @@ void FPropertyEditorModule::RegisterCustomPropertyTypeLayout( FName PropertyType
 
 void FPropertyEditorModule::UnregisterCustomPropertyTypeLayout( FName PropertyTypeName, TSharedPtr<IPropertyTypeIdentifier> Identifier, TSharedPtr<IDetailsView> ForSpecificInstance )
 {
-	if( PropertyTypeName != NAME_None )
+	if (!PropertyTypeName.IsValid() || (PropertyTypeName == NAME_None))
 	{
-		if( ForSpecificInstance.IsValid() )
+		return;
+	}
+
+	if (ForSpecificInstance.IsValid())
+	{
+		FCustomPropertyTypeLayoutMap* PropertyTypeToLayoutMap = InstancePropertyTypeLayoutMap.Find(ForSpecificInstance);
+			
+		if (PropertyTypeToLayoutMap)
 		{
-			FCustomPropertyTypeLayoutMap* PropertyTypeToLayoutMap = InstancePropertyTypeLayoutMap.Find( ForSpecificInstance );
+			FPropertyTypeLayoutCallbackList* LayoutCallbacks = PropertyTypeToLayoutMap->Find(PropertyTypeName);
 			
-			if( PropertyTypeToLayoutMap )
+			if (LayoutCallbacks)
 			{
-				FPropertyTypeLayoutCallbackList* LayoutCallbacks = PropertyTypeToLayoutMap->Find( PropertyTypeName );
-			
-				if( LayoutCallbacks )
-				{
-					LayoutCallbacks->Remove( Identifier );
-				}
+				LayoutCallbacks->Remove(Identifier);
 			}
 		}
-		else
-		{
-			FPropertyTypeLayoutCallbackList* LayoutCallbacks = GlobalPropertyTypeToLayoutMap.Find( PropertyTypeName );
+	}
+	else
+	{
+		FPropertyTypeLayoutCallbackList* LayoutCallbacks = GlobalPropertyTypeToLayoutMap.Find(PropertyTypeName);
 
-			if( LayoutCallbacks )
-			{
-				LayoutCallbacks->Remove( Identifier );
-			}
+		if (LayoutCallbacks)
+		{
+			LayoutCallbacks->Remove(Identifier);
 		}
 	}
 }
