Undoing CL 2741389 (an attempted fix to constrain EditInlineNew properties). This was causing  FRawDistributionVector's details selection to break, as users could no longer select the distribution type. An alternative fix for the original problem was separately submitted for 4.10, with CL 2738129.

diff --git a/Engine/Source/Editor/PropertyEditor/Private/PropertyNode.cpp b/Engine/Source/Editor/PropertyEditor/Private/PropertyNode.cpp
--- a/Engine/Source/Editor/PropertyEditor/Private/PropertyNode.cpp
+++ b/Engine/Source/Editor/PropertyEditor/Private/PropertyNode.cpp
@@ -133,9 +133,6 @@ void FPropertyNode::InitNode( const FPropertyNodeInitParams& InitParams )
 		const bool bEditInline = bIsObjectOrInterface && GotReadAddresses && MyProperty->HasMetaData(TEXT("EditInline"));
 		SetNodeFlags(EPropertyNodeFlags::EditInline, bEditInline);
 
-		const bool bPersistentInstance = MyProperty->HasAllPropertyFlags(CPF_PersistentInstance);
-		SetNodeFlags(EPropertyNodeFlags::PersistentInstance, bPersistentInstance);
-
 		//Get the property max child depth
 		if (Property->HasMetaData(TEXT("MaxPropertyDepth")))
 		{
diff --git a/Engine/Source/Editor/PropertyEditor/Private/PropertyNode.h b/Engine/Source/Editor/PropertyEditor/Private/PropertyNode.h
--- a/Engine/Source/Editor/PropertyEditor/Private/PropertyNode.h
+++ b/Engine/Source/Editor/PropertyEditor/Private/PropertyNode.h
@@ -45,8 +45,6 @@ namespace EPropertyNodeFlags
 
 	const Type	SkipChildValidation				= 1 << 21; /** true if this node should skip child validation */
 
-	const Type	PersistentInstance				= 1 << 22; /** true if this node supports creating 'editinlinenew' sub-objects */
-
 	const Type 	NoFlags							= 0;
 
 };
diff --git a/Engine/Source/Editor/PropertyEditor/Private/UserInterface/PropertyEditor/SPropertyEditorEditInline.cpp b/Engine/Source/Editor/PropertyEditor/Private/UserInterface/PropertyEditor/SPropertyEditorEditInline.cpp
--- a/Engine/Source/Editor/PropertyEditor/Private/UserInterface/PropertyEditor/SPropertyEditorEditInline.cpp
+++ b/Engine/Source/Editor/PropertyEditor/Private/UserInterface/PropertyEditor/SPropertyEditorEditInline.cpp
@@ -130,11 +130,10 @@ void SPropertyEditorEditInline::GetDesiredWidth( float& OutMinDesiredWidth, floa
 
 bool SPropertyEditorEditInline::Supports( const FPropertyNode* InTreeNode, int32 InArrayIdx )
 {
-	return InTreeNode 
+	return InTreeNode
 		&& !InTreeNode->IsEditConst()
 		&& InTreeNode->FindObjectItemParent()
-		&& InTreeNode->HasNodeFlags(EPropertyNodeFlags::EditInline)
-		&& InTreeNode->HasNodeFlags(EPropertyNodeFlags::PersistentInstance);
+		&& InTreeNode->HasNodeFlags(EPropertyNodeFlags::EditInline);
 }
 
 bool SPropertyEditorEditInline::Supports( const TSharedRef< class FPropertyEditor >& InPropertyEditor )
