Copying //UE4/Dev-Blueprints to Dev-Main (//UE4/Dev-Main)

diff --git a/Engine/Source/Editor/Kismet/Private/Profiler/ScriptPerfData.cpp b/Engine/Source/Editor/Kismet/Private/Profiler/ScriptPerfData.cpp
--- a/Engine/Source/Editor/Kismet/Private/Profiler/ScriptPerfData.cpp
+++ b/Engine/Source/Editor/Kismet/Private/Profiler/ScriptPerfData.cpp
@@ -10,18 +10,34 @@ FNumberFormattingOptions FScriptPerfData::StatNumberFormat;
 float FScriptPerfData::RecentSampleBias = 0.2f;
 float FScriptPerfData::HistoricalSampleBias = 0.8f;
 float FScriptPerfData::NodePerformanceThreshold = 1.f;
-float FScriptPerfData::PureNodePerformanceThreshold = 1.f;
 float FScriptPerfData::InclusivePerformanceThreshold = 1.5f;
 float FScriptPerfData::MaxPerformanceThreshold = 2.f;
 
-void FScriptPerfData::AddEventTiming(const double PureNodeTimingIn, const double NodeTimingIn)
+void FScriptPerfData::AddEventTiming(const double NodeTimingIn)
 {
 	const double TimingInMs = NodeTimingIn * 1000;
-	const double PureNodeTimingInMs = PureNodeTimingIn * 1000;
 	NodeTiming = (TimingInMs * RecentSampleBias) + (NodeTiming * HistoricalSampleBias);
-	PureNodeTiming = (PureNodeTimingInMs * RecentSampleBias) + (PureNodeTiming * HistoricalSampleBias);
-	InclusiveTiming = NodeTiming + PureNodeTiming;
-	TotalTiming += NodeTimingIn + PureNodeTimingIn;
+	MaxTiming = FMath::Max<double>(MaxTiming, NodeTiming);
+	MinTiming = FMath::Min<double>(MinTiming, NodeTiming);
+	TotalTiming += NodeTimingIn;
+	NumSamples++;
+}
+
+void FScriptPerfData::AddPureChainTiming(const double PureNodeTimingIn)
+{
+	const double PureTimingInMs = PureNodeTimingIn * 1000;
+	const double NewInclusiveTimingInMs = NodeTiming + PureTimingInMs;
+	InclusiveTiming = (NewInclusiveTimingInMs * RecentSampleBias) + (InclusiveTiming * HistoricalSampleBias);
+	TotalTiming += PureNodeTimingIn;
+	MaxTiming = FMath::Max<double>(MaxTiming, InclusiveTiming);
+	MinTiming = FMath::Min<double>(MinTiming, InclusiveTiming);
+}
+
+void FScriptPerfData::AddInclusiveTiming(const double InclusiveNodeTimingIn)
+{
+	const double TimingInMs = InclusiveNodeTimingIn * 1000;
+	InclusiveTiming = (TimingInMs * RecentSampleBias) + (InclusiveTiming * HistoricalSampleBias);
+	TotalTiming += InclusiveNodeTimingIn;
 	NumSamples++;
 	MaxTiming = FMath::Max<double>(MaxTiming, InclusiveTiming);
 	MinTiming = FMath::Min<double>(MinTiming, InclusiveTiming);
@@ -33,8 +49,10 @@ void FScriptPerfData::AddData(const FScriptPerfData& DataIn)
 	{
 		// Accumulate data, find min, max values
 		NodeTiming += DataIn.NodeTiming;
-		PureNodeTiming += DataIn.PureNodeTiming;
-		InclusiveTiming = NodeTiming + PureNodeTiming;
+		if (DataIn.InclusiveTiming > 0.0)
+		{
+			InclusiveTiming += DataIn.InclusiveTiming;
+		}
 		TotalTiming += DataIn.TotalTiming;
 		NumSamples += DataIn.NumSamples;
 		MaxTiming = FMath::Max<double>(MaxTiming, DataIn.MaxTiming);
@@ -53,7 +71,6 @@ void FScriptPerfData::AddBranchData(const FScriptPerfData& DataIn)
 void FScriptPerfData::Reset()
 {
 	NodeTiming = 0.0;
-	PureNodeTiming = 0.0;
 	InclusiveTiming = 0.0;
 	TotalTiming = 0.0;
 	NumSamples = 0;
@@ -78,17 +95,6 @@ FSlateColor FScriptPerfData::GetNodeHeatColor() const
 	return FLinearColor(1.f, Value, Value);
 }
 
-void FScriptPerfData::SetPureNodePerformanceThreshold(const float PureNodePerformanceThresholdIn)
-{
-	PureNodePerformanceThreshold = PureNodePerformanceThresholdIn;
-}
-
-FSlateColor FScriptPerfData::GetPureNodeHeatColor() const
-{
-	const float Value = 1.f - FMath::Min<float>(PureNodeTiming / PureNodePerformanceThreshold, 1.f);
-	return FLinearColor(1.f, Value, Value);
-}
-
 void FScriptPerfData::SetInclusivePerformanceThreshold(const float InclusivePerformanceThresholdIn)
 {
 	InclusivePerformanceThreshold = InclusivePerformanceThresholdIn;
@@ -113,56 +119,54 @@ FSlateColor FScriptPerfData::GetMaxTimeHeatColor() const
 
 FText FScriptPerfData::GetNodeTimingText() const
 {
-	return FText::AsNumber(GetNodeTiming(), &StatNumberFormat);
-}
-
-FText FScriptPerfData::GetPureNodeTimingText() const
-{
-	const double CurrentPureTiming = GetPureTiming();
-	if (CurrentPureTiming == 0.0)
-	{
-		return FText::GetEmpty();
-	}
-	else
+	if (NodeTiming > 0.0)
 	{
-		return FText::AsNumber(CurrentPureTiming, &StatNumberFormat);
+		return FText::AsNumber(NodeTiming, &StatNumberFormat);
 	}
+	return FText::GetEmpty();
 }
 
 FText FScriptPerfData::GetInclusiveTimingText() const
 {
-	const double CurrentPureTiming = GetPureTiming();
-	if (CurrentPureTiming == 0.0)
-	{
-		return FText::GetEmpty();
-	}
-	else
+	if (InclusiveTiming > 0.0)
 	{
-		return FText::AsNumber(GetInclusiveTiming(), &StatNumberFormat);
+		return FText::AsNumber(InclusiveTiming, &StatNumberFormat);
 	}
+	return FText::GetEmpty();
 }
 
 FText FScriptPerfData::GetMaxTimingText() const
 {
-	const double MaximumTime = GetMaxTiming();
-	const bool bMaxTimeInitialised = MaximumTime != -MAX_dbl;
-	return FText::AsNumber(bMaxTimeInitialised ? MaximumTime : 0.0, &StatNumberFormat);
+	if (MaxTiming != -MAX_dbl)
+	{
+		return FText::AsNumber(MaxTiming, &StatNumberFormat);
+	}
+	return FText::GetEmpty();
 }
 
 FText FScriptPerfData::GetMinTimingText() const
 {
-	const double MinimumTime = GetMinTiming();
-	const bool bMinTimeInitialised = MinimumTime != MAX_dbl;
-	return FText::AsNumber(bMinTimeInitialised ? MinimumTime : 0.0, &StatNumberFormat);
+	if (MinTiming != MAX_dbl)
+	{
+		return FText::AsNumber(MinTiming, &StatNumberFormat);
+	}
+	return FText::GetEmpty();
 }
 
 FText FScriptPerfData::GetTotalTimingText() const
 {
-	// Covert back up to seconds for total time.
-	return FText::AsNumber(GetTotalTiming(), &StatNumberFormat);
+	if (TotalTiming > 0.0)
+	{
+		return FText::AsNumber(TotalTiming, &StatNumberFormat);
+	}
+	return FText::GetEmpty();
 }
 
 FText FScriptPerfData::GetSamplesText() const
 {
-	return FText::AsNumber(GetSampleCount());
+	if (NumSamples > 0)
+	{
+		return FText::AsNumber(NumSamples);
+	}
+	return FText::GetEmpty();
 }
\ No newline at end of file
diff --git a/Engine/Source/Editor/Kismet/Private/SKismetInspector.cpp b/Engine/Source/Editor/Kismet/Private/SKismetInspector.cpp
--- a/Engine/Source/Editor/Kismet/Private/SKismetInspector.cpp
+++ b/Engine/Source/Editor/Kismet/Private/SKismetInspector.cpp
@@ -17,6 +17,7 @@
 #include "BlueprintDetailsCustomization.h"
 #include "UserDefinedEnumEditor.h"
 #include "UserDefinedStructureEditor.h"
+#include "BitmaskLiteralDetails.h"
 #include "FormatTextDetails.h"
 #include "Engine/SCS_Node.h"
 #include "Components/ChildActorComponent.h"
@@ -345,6 +346,9 @@ void SKismetInspector::Construct(const FArguments& InArgs)
 		FOnGetDetailCustomizationInstance LayoutFormatTextDetails = FOnGetDetailCustomizationInstance::CreateStatic(&FFormatTextDetails::MakeInstance);
 		PropertyView->RegisterInstancedCustomPropertyLayout(UK2Node_FormatText::StaticClass(), LayoutFormatTextDetails);
 
+		FOnGetDetailCustomizationInstance LayoutBitmaskLiteralDetails = FOnGetDetailCustomizationInstance::CreateStatic(&FBitmaskLiteralDetails::MakeInstance);
+		PropertyView->RegisterInstancedCustomPropertyLayout(UK2Node_BitmaskLiteral::StaticClass(), LayoutBitmaskLiteralDetails);
+
 		FOnGetDetailCustomizationInstance LayoutDocumentationDetails = FOnGetDetailCustomizationInstance::CreateStatic(&FBlueprintDocumentationDetails::MakeInstance, BlueprintEditorPtr);
 		PropertyView->RegisterInstancedCustomPropertyLayout(UEdGraphNode_Documentation::StaticClass(), LayoutDocumentationDetails);
 
diff --git a/Engine/Source/Editor/Kismet/Private/SMyBlueprint.cpp b/Engine/Source/Editor/Kismet/Private/SMyBlueprint.cpp
--- a/Engine/Source/Editor/Kismet/Private/SMyBlueprint.cpp
+++ b/Engine/Source/Editor/Kismet/Private/SMyBlueprint.cpp
@@ -1201,7 +1201,7 @@ void SMyBlueprint::CollectAllActions(FGraphActionListBuilderBase& OutAllActions)
 
 		//@TODO: Should be a bit more generic (or the AnimGraph shouldn't be stored as a FunctionGraph...)
 		int32 SectionID = Graph->IsA<UAnimationGraph>() ? NodeSectionID::GRAPH : NodeSectionID::FUNCTION;
-		TSharedPtr<FEdGraphSchemaAction_K2Graph> NewFuncAction = MakeShareable(new FEdGraphSchemaAction_K2Graph(EEdGraphSchemaAction_K2Graph::Function, FunctionCategory, DisplayInfo.DisplayName, DisplayInfo.Tooltip, bIsConstructionScript ? 2 : 1, SectionID));
+		TSharedPtr<FEdGraphSchemaAction_K2Graph> NewFuncAction = MakeShareable(new FEdGraphSchemaAction_K2Graph(EEdGraphSchemaAction_K2Graph::Function, FunctionCategory, DisplayInfo.PlainName, DisplayInfo.Tooltip, bIsConstructionScript ? 2 : 1, SectionID));
 		NewFuncAction->FuncName = Graph->GetFName();
 		NewFuncAction->EdGraph = Graph;
 
@@ -1225,7 +1225,7 @@ void SMyBlueprint::CollectAllActions(FGraphActionListBuilderBase& OutAllActions)
 
 		FText MacroCategory = GetGraphCategory(Graph);
 
-		TSharedPtr<FEdGraphSchemaAction_K2Graph> NewMacroAction = MakeShareable(new FEdGraphSchemaAction_K2Graph(EEdGraphSchemaAction_K2Graph::Macro, MacroCategory, DisplayInfo.DisplayName, DisplayInfo.Tooltip, 1, NodeSectionID::MACRO));
+		TSharedPtr<FEdGraphSchemaAction_K2Graph> NewMacroAction = MakeShareable(new FEdGraphSchemaAction_K2Graph(EEdGraphSchemaAction_K2Graph::Macro, MacroCategory, DisplayInfo.PlainName, DisplayInfo.Tooltip, 1, NodeSectionID::MACRO));
 		NewMacroAction->FuncName = MacroName;
 		NewMacroAction->EdGraph = Graph;
 		OutAllActions.AddAction(NewMacroAction);
@@ -1343,7 +1343,7 @@ void SMyBlueprint::CollectAllActions(FGraphActionListBuilderBase& OutAllActions)
 		FGraphDisplayInfo DisplayInfo;
 		Graph->GetSchema()->GetGraphDisplayInformation(*Graph, DisplayInfo);
 
-		TSharedPtr<FEdGraphSchemaAction_K2Graph> NeUbergraphAction = MakeShareable(new FEdGraphSchemaAction_K2Graph(EEdGraphSchemaAction_K2Graph::Graph, FText::GetEmpty(), DisplayInfo.DisplayName, DisplayInfo.Tooltip, 2, NodeSectionID::GRAPH));
+		TSharedPtr<FEdGraphSchemaAction_K2Graph> NeUbergraphAction = MakeShareable(new FEdGraphSchemaAction_K2Graph(EEdGraphSchemaAction_K2Graph::Graph, FText::GetEmpty(), DisplayInfo.PlainName, DisplayInfo.Tooltip, 2, NodeSectionID::GRAPH));
 		NeUbergraphAction->FuncName = Graph->GetFName();
 		NeUbergraphAction->EdGraph = Graph;
 		OutAllActions.AddAction(NeUbergraphAction);
diff --git a/Engine/Source/Editor/Kismet/Private/UserDefinedEnumEditor.cpp b/Engine/Source/Editor/Kismet/Private/UserDefinedEnumEditor.cpp
--- a/Engine/Source/Editor/Kismet/Private/UserDefinedEnumEditor.cpp
+++ b/Engine/Source/Editor/Kismet/Private/UserDefinedEnumEditor.cpp
@@ -6,6 +6,7 @@
 
 #include "PropertyCustomizationHelpers.h"
 #include "SDockTab.h"
+#include "IDocumentation.h"
 #include "Engine/UserDefinedEnum.h"
 
 #define LOCTEXT_NAMESPACE "UserDefinedEnumEditor"
@@ -158,6 +159,8 @@ void FEnumDetails::CustomizeDetails( IDetailLayoutBuilder& DetailLayout )
 		TargetEnum = CastChecked<UUserDefinedEnum>(Objects[0].Get());
 		TSharedRef<IPropertyHandle> PropertyHandle = DetailLayout.GetProperty(FName("Names"), UEnum::StaticClass());
 
+		const FString DocLink = TEXT("Shared/Editors/BlueprintEditor/EnumDetails");
+
 		IDetailCategoryBuilder& InputsCategory = DetailLayout.EditCategory("Enumerators", LOCTEXT("EnumDetailsEnumerators", "Enumerators"));
 
 		InputsCategory.AddCustomRow( LOCTEXT("FunctionNewInputArg", "New") )
@@ -173,6 +176,23 @@ void FEnumDetails::CustomizeDetails( IDetailLayoutBuilder& DetailLayout )
 
 		Layout = MakeShareable( new FUserDefinedEnumLayout(TargetEnum.Get()) );
 		InputsCategory.AddCustomBuilder( Layout.ToSharedRef() );
+
+		TSharedPtr<SToolTip> BitmaskFlagsTooltip = IDocumentation::Get()->CreateToolTip(LOCTEXT("BitmaskFlagsTooltip", "When enabled, this enumeration can be used as a set of explicitly-named bitmask flags. Each enumerator's value will correspond to the index of the bit (flag) in the mask."), nullptr, DocLink, TEXT("Bitmask Flags"));
+
+		InputsCategory.AddCustomRow(LOCTEXT("BitmaskFlagsAttributeLabel", "Bitmask Flags"), true)
+		.NameContent()
+		[
+			SNew(STextBlock)
+			.Text(LOCTEXT("BitmaskFlagsAttributeLabel", "Bitmask Flags"))
+			.ToolTip(BitmaskFlagsTooltip)
+		]
+		.ValueContent()
+		[
+			SNew(SCheckBox)
+			.IsChecked(this, &FEnumDetails::OnGetBitmaskFlagsAttributeState)
+			.OnCheckStateChanged(this, &FEnumDetails::OnBitmaskFlagsAttributeStateChanged)
+			.ToolTip(BitmaskFlagsTooltip)
+		];
 	}
 }
 
@@ -206,6 +226,36 @@ FReply FEnumDetails::OnAddNewEnumerator()
 	return FReply::Handled();
 }
 
+ECheckBoxState FEnumDetails::OnGetBitmaskFlagsAttributeState() const
+{
+	ECheckBoxState Result = ECheckBoxState::Unchecked;
+	const UUserDefinedEnum* Enum = TargetEnum.Get();
+	if (Enum && Enum->HasMetaData(TEXT("Bitflags")))
+	{
+		Result = ECheckBoxState::Checked;
+	}
+
+	return Result;
+}
+
+void FEnumDetails::OnBitmaskFlagsAttributeStateChanged(ECheckBoxState InNewState)
+{
+	UUserDefinedEnum* Enum = TargetEnum.Get();
+	if (Enum)
+	{
+		Enum->Modify();
+
+		if (InNewState == ECheckBoxState::Checked)
+		{
+			Enum->SetMetaData(TEXT("Bitflags"), TEXT(""));
+		}
+		else
+		{
+			Enum->RemoveMetaData(TEXT("Bitflags"));
+		}
+	}
+}
+
 bool FUserDefinedEnumLayout::CausedChange() const
 {
 	for(auto Iter = Children.CreateConstIterator(); Iter; ++Iter)
diff --git a/Engine/Source/Editor/Kismet/Private/UserDefinedEnumEditor.h b/Engine/Source/Editor/Kismet/Private/UserDefinedEnumEditor.h
--- a/Engine/Source/Editor/Kismet/Private/UserDefinedEnumEditor.h
+++ b/Engine/Source/Editor/Kismet/Private/UserDefinedEnumEditor.h
@@ -77,6 +77,10 @@ private:
 	/** Handles new enum element request */
 	FReply OnAddNewEnumerator();
 
+	/** Handles the optional bitmask flags attribute */
+	ECheckBoxState OnGetBitmaskFlagsAttributeState() const;
+	void OnBitmaskFlagsAttributeStateChanged(ECheckBoxState InNewState);
+
 private:
 	TSharedPtr<class FUserDefinedEnumLayout> Layout;
 	/** The target node that this argument is on */
diff --git a/Engine/Source/Editor/Kismet/Public/BlueprintEditor.h b/Engine/Source/Editor/Kismet/Public/BlueprintEditor.h
--- a/Engine/Source/Editor/Kismet/Public/BlueprintEditor.h
+++ b/Engine/Source/Editor/Kismet/Public/BlueprintEditor.h
@@ -507,6 +507,9 @@ public:
 	/** Can generate native code for current blueprint */
 	bool CanGenerateNativeCode() const;
 
+	/** Returns if the blueprint profiler is enabled in editor settings. */
+	bool IsProfilerAvailable() const;
+
 	/** Returns if the blueprint profiler is currently active. */
 	bool IsProfilerActive() const;
 
@@ -914,9 +917,6 @@ protected:
 	/** Attempts to invoke the details tab if it's currently possible to. */
 	void TryInvokingDetailsTab(bool bFlash = true);
 
-	/** Returns true if the blueprint profiler is valid for the current editor and mode */
-	bool IsBlueprintProfilerSupported() const;
-
 private:
 
 	/** Returns true if modules can be recompiled */
diff --git a/Engine/Source/Editor/Kismet/Public/Profiler/EventExecution.h b/Engine/Source/Editor/Kismet/Public/Profiler/EventExecution.h
--- a/Engine/Source/Editor/Kismet/Public/Profiler/EventExecution.h
+++ b/Engine/Source/Editor/Kismet/Public/Profiler/EventExecution.h
@@ -9,22 +9,33 @@ namespace EScriptExecutionNodeFlags
 {
 	enum Type
 	{
-		None				= 0x00000000,	// No Flags
-		Class				= 0x00000001,	// Class
-		Instance			= 0x00000002,	// Instance
-		Event				= 0x00000004,	// Event
-		CustomEvent			= 0x00000008,	// Custom Event
-		Container			= 0x00000007,	// Container type - Class, Instance or Event.
-		FunctionCall		= 0x00000010,	// Function Call
-		MacroCall			= 0x00000020,	// Macro Call
-		CallSite			= 0x00000030,	// Function / Macro call site
-		MacroNode			= 0x00000100,	// Macro Node
-		ConditionalBranch	= 0x00000200,	// Node has multiple exit pins using a jump
-		SequentialBranch	= 0x00000400,	// Node has multiple exit pins ran in sequence
-		BranchNode			= 0x00000600,	// BranchNode
-		Node				= 0x00000800,	// Node timing
-		ExecPin				= 0x00001000,	// Exec pin dummy node
-		PureNode			= 0x00002000	// Pure node (no exec pins)
+		None						= 0x00000000,	// No Flags
+		Class						= 0x00000001,	// Class
+		Instance					= 0x00000002,	// Instance
+		Event						= 0x00000004,	// Event
+		CustomEvent					= 0x00000008,	// Custom Event
+		FunctionCall				= 0x00000010,	// Function Call
+		MacroCall					= 0x00000020,	// Macro Call
+		MacroNode					= 0x00000100,	// Macro Node
+		ConditionalBranch			= 0x00000200,	// Node has multiple exit pins using a jump
+		SequentialBranch			= 0x00000400,	// Node has multiple exit pins ran in sequence
+		Node						= 0x00000800,	// Node timing
+		ExecPin						= 0x00001000,	// Exec pin dummy node
+		PureNode					= 0x00002000,	// Pure node (no exec pins)
+		FunctionTunnel				= 0x00010000,	// Tunnel function
+		TunnelEntryPin				= 0x00020000,	// Tunnel entry pin
+		TunnelExitPin				= 0x00040000,	// Tunnel exit pin
+		ReEntrantTunnelPin			= 0x00080000,	// Re-Entrant tunnel pin
+		TunnelThenPin				= 0x00100000,	// Tunnel final exit or exec then pin
+		PureChain					= 0x00200000,	// Pure node call chain
+		CyclicLinkage				= 0x00400000,	// Marks execution path as cyclic.
+		InvalidTrace				= 0x00800000,	// Indicates that node doesn't contain a valid script trace.
+									// Groups
+		CallSite					= FunctionCall|MacroCall|FunctionTunnel,
+		BranchNode					= ConditionalBranch|SequentialBranch,
+		TunnelPin					= TunnelEntryPin|TunnelExitPin,
+		PureStats					= PureNode|PureChain,
+		Container					= Class|Instance|Event|CustomEvent|TunnelPin|ExecPin|PureChain
 	};
 }
 
@@ -59,21 +70,6 @@ public:
 		TSharedPtr<class FScriptExecutionNode> LinkedNode;
 	};
 
-	/** Returns the parent node */
-	TSharedPtr<class FScriptExecutionNode> GetParentNode() { return ParentNode; }
-
-	/** Sets the parent node */
-	void SetParentNode(TSharedPtr<class FScriptExecutionNode> InParentNode) { ParentNode = InParentNode; }
-
-	/** Returns the pure nodes map */
-	TMap<int32, TSharedPtr<class FScriptExecutionNode>>& GetPureNodes() { return PureNodes; }
-
-	/** Returns the number of pure nodes */
-	int32 GetNumPureNodes() const { return PureNodes.Num(); }
-
-	/** Add pure node */
-	void AddPureNode(const int32 PinScriptOffset, TSharedPtr<class FScriptExecutionNode> PureNode);
-
 	/** Returns the linked nodes map */
 	TMap<int32, TSharedPtr<class FScriptExecutionNode>>& GetLinkedNodes() { return LinkedNodes; }
 
@@ -83,6 +79,9 @@ public:
 	/** Add linked node */
 	void AddLinkedNode(const int32 PinScriptOffset, TSharedPtr<class FScriptExecutionNode> LinkedNode);
 
+	/** Returns linked node by matching script offset */
+	TSharedPtr<class FScriptExecutionNode> GetLinkedNodeByScriptOffset(const int32 PinScriptOffset);
+
 	/** Returns the child nodes map */
 	TArray<TSharedPtr<class FScriptExecutionNode>>& GetChildNodes() { return ChildNodes; }
 
@@ -97,10 +96,6 @@ public:
 
 protected:
 
-	/** Parent node */
-	TSharedPtr<class FScriptExecutionNode> ParentNode;
-	/** Pure nodes */
-	TMap<int32, TSharedPtr<class FScriptExecutionNode>> PureNodes;
 	/** Linked nodes */
 	TMap<int32, TSharedPtr<class FScriptExecutionNode>> LinkedNodes;
 	/** Child nodes */
@@ -159,8 +154,6 @@ struct KISMET_API FScriptExecNodeParams
 //////////////////////////////////////////////////////////////////////////
 // FScriptExecutionNode
 
-typedef TSet<TWeakObjectPtr<const UObject>> FExecNodeFilter;
-
 class KISMET_API FScriptExecutionNode : public FScriptNodeExecLinkage, public FScriptNodePerfData, public TSharedFromThis<FScriptExecutionNode>
 {
 public:
@@ -210,6 +203,9 @@ public:
 	/** Returns if this node is also a pure node (i.e. no exec input pin) */
 	bool IsPureNode() const { return (NodeFlags & EScriptExecutionNodeFlags::PureNode) != 0U; }
 
+	/** Returns if this node is a pure chain node */
+	bool IsPureChain() const { return (NodeFlags & EScriptExecutionNodeFlags::PureChain) != 0U; }
+
 	/** Returns if this event potentially multiple exit sites */
 	bool IsBranch() const { return (NodeFlags & EScriptExecutionNodeFlags::BranchNode) != 0U; }
 
@@ -246,6 +242,12 @@ public:
 	/** Sets the current expansion state for widget UI */
 	void SetExpanded(bool bIsExpanded) { bExpansionState = bIsExpanded; }
 
+	/** Returns the pure chain node associated with this exec node (if one exists) */
+	TSharedPtr<FScriptExecutionNode> GetPureChainNode();
+
+	/** Returns pure node script code range */
+	FInt32Range GetPureNodeScriptCodeRange() const { return PureNodeScriptCodeRange; }
+
 	/** Sets the pure node script code range */
 	void SetPureNodeScriptCodeRange(FInt32Range InScriptCodeRange) { PureNodeScriptCodeRange = InScriptCodeRange; }
 
@@ -266,15 +268,12 @@ public:
 
 protected:
 
+	/** Returns Tunnel Linear Execution Trace */
+	void MapTunnelLinearExecution(FTracePath& TraceInOut) const;
+
 	/** Get all pure nodes - private implementation */
 	void GetAllPureNodes_Internal(TMap<int32, TSharedPtr<FScriptExecutionNode>>& PureNodesOut, const FInt32Range& ScriptCodeRange);
 
-	/** Get linear execution path - private implementation */
-	void GetLinearExecutionPath_Internal(FExecNodeFilter& Filter, TArray<FLinearExecPath>& LinearExecutionNodes, const FTracePath& TracePath);
-
-	/** Refresh Stats - private implementation */
-	void RefreshStats_Internal(const FTracePath& TracePath, FExecNodeFilter& VisitedStats);
-
 protected:
 
 	/** Node name */
diff --git a/Engine/Source/Editor/Kismet/Public/Profiler/ScriptPerfData.h b/Engine/Source/Editor/Kismet/Public/Profiler/ScriptPerfData.h
--- a/Engine/Source/Editor/Kismet/Public/Profiler/ScriptPerfData.h
+++ b/Engine/Source/Editor/Kismet/Public/Profiler/ScriptPerfData.h
@@ -12,7 +12,6 @@ public:
 	FScriptPerfData()
 		: InclusiveTiming(0.0)
 		, NodeTiming(0.0)
-		, PureNodeTiming(0.0)
 		, MaxTiming(-MAX_dbl)
 		, MinTiming(MAX_dbl)
 		, TotalTiming(0.0)
@@ -21,7 +20,13 @@ public:
 	}
 
 	/** Add a single event node timing into the dataset */
-	void AddEventTiming(const double PureNodeTimingIn, const double NodeTimingIn);
+	void AddEventTiming(const double NodeTimingIn);
+
+	/** Add a single event pure node timing into the dataset */
+	void AddPureChainTiming(const double PureTimingIn);
+
+	/** Add a inclusive timing into the dataset */
+	void AddInclusiveTiming(const double InclusiveNodeTimingIn);
 
 	/** Add data */
 	void AddData(const FScriptPerfData& DataIn);
@@ -39,14 +44,12 @@ public:
 	static void SetNumberFormattingForStats(const FNumberFormattingOptions& FormatIn) { StatNumberFormat = FormatIn; }
 	static void SetRecentSampleBias(const float RecentSampleBiasIn);
 	static void SetNodePerformanceThreshold(const float NodePerformanceThresholdIn);
-	static void SetPureNodePerformanceThreshold(const float PureNodePerformanceThresholdIn);
 	static void SetInclusivePerformanceThreshold(const float InclusivePerformanceThresholdIn);
 	static void SetMaxPerformanceThreshold(const float MaxPerformanceThresholdIn);
 
 	// Returns the various stats this container holds
 	double GetNodeTiming() const { return NodeTiming; }
-	double GetPureTiming() const { return PureNodeTiming; }
-	float GetInclusiveTiming() const { return InclusiveTiming; }
+	double GetInclusiveTiming() const { return InclusiveTiming; }
 	double GetMaxTiming() const { return MaxTiming; }
 	double GetMinTiming() const { return MinTiming; }
 	double GetTotalTiming() const { return TotalTiming; }
@@ -54,7 +57,6 @@ public:
 
 	// Returns various performance colors for visual display
 	FSlateColor GetNodeHeatColor() const;
-	FSlateColor GetPureNodeHeatColor() const;
 	FSlateColor GetInclusiveHeatColor() const;
 	FSlateColor GetMaxTimeHeatColor() const;
 
@@ -62,7 +64,6 @@ public:
 	FText GetTotalTimingText() const;
 	FText GetInclusiveTimingText() const;
 	FText GetNodeTimingText() const;
-	FText GetPureNodeTimingText() const;
 	FText GetMaxTimingText() const;
 	FText GetMinTimingText() const;
 	FText GetSamplesText() const;
@@ -73,8 +74,6 @@ private:
 	double InclusiveTiming;
 	/** Node timing */
 	double NodeTiming;
-	/** Pure node timing */
-	double PureNodeTiming;
 	/** Max exclusive timing */
 	double MaxTiming;
 	/** Min exclusive timing */
@@ -93,8 +92,6 @@ private:
 	static float HistoricalSampleBias;
 	/** Cached node performance threshold */
 	static float NodePerformanceThreshold;
-	/** Cached pure node performance threshold */
-	static float PureNodePerformanceThreshold;
 	/** Cached inclusive performance threshold */
 	static float InclusivePerformanceThreshold;
 	/** Cached max performance threshold */
diff --git a/Engine/Source/Editor/KismetCompiler/Private/KismetCompilerVMBackend.cpp b/Engine/Source/Editor/KismetCompiler/Private/KismetCompilerVMBackend.cpp
--- a/Engine/Source/Editor/KismetCompiler/Private/KismetCompilerVMBackend.cpp
+++ b/Engine/Source/Editor/KismetCompiler/Private/KismetCompilerVMBackend.cpp
@@ -469,6 +469,20 @@ public:
 					Writer << Value;
 				}
 			}
+			else if (CoerceProperty->IsA(UInt64Property::StaticClass()))
+			{
+				int64 Value = 0;
+				LexicalConversion::FromString(Value, *(Term->Name));
+				Writer << EX_Int64Const;
+				Writer << Value;
+			}
+			else if (CoerceProperty->IsA(UUInt64Property::StaticClass()))
+			{
+				uint64 Value = 0;
+				LexicalConversion::FromString(Value, *(Term->Name));
+				Writer << EX_UInt64Const;
+				Writer << Value;
+			}
 			else if (CoerceProperty->IsA(UByteProperty::StaticClass()))
 			{
 				UByteProperty* ByteProp = CastChecked< UByteProperty >( CoerceProperty );
diff --git a/Engine/Source/Editor/KismetWidgets/Private/SPinTypeSelector.cpp b/Engine/Source/Editor/KismetWidgets/Private/SPinTypeSelector.cpp
--- a/Engine/Source/Editor/KismetWidgets/Private/SPinTypeSelector.cpp
+++ b/Engine/Source/Editor/KismetWidgets/Private/SPinTypeSelector.cpp
@@ -191,8 +191,9 @@ void SPinTypeSelector::Construct(const FArguments& InArgs, FGetPinTypeTree GetPi
 
 FText SPinTypeSelector::GetTypeDescription() const
 {
+	const FString& PinSubCategory = TargetPinType.Get().PinSubCategory;
 	const UObject* PinSubCategoryObject = TargetPinType.Get().PinSubCategoryObject.Get();
-	if (PinSubCategoryObject)
+	if (PinSubCategory != UEdGraphSchema_K2::PSC_Bitmask && PinSubCategoryObject)
 	{
 		if (auto Field = Cast<const UField>(PinSubCategoryObject))
 		{
@@ -386,7 +387,7 @@ TSharedRef< SWidget > SPinTypeSelector::GetAllowedObjectTypes(FPinTypeTreeItem I
 
 	FFormatNamedArguments Args;
 
-	if(PinType.PinSubCategoryObject.IsValid())
+	if(PinType.PinSubCategory != UEdGraphSchema_K2::PSC_Bitmask && PinType.PinSubCategoryObject.IsValid())
 	{
 		Args.Add(TEXT("TypeName"), InItem->GetDescription());
 	}
diff --git a/Engine/Source/Editor/PropertyEditor/Private/UserInterface/PropertyEditor/PropertyEditorConstants.cpp b/Engine/Source/Editor/PropertyEditor/Private/UserInterface/PropertyEditor/PropertyEditorConstants.cpp
--- a/Engine/Source/Editor/PropertyEditor/Private/UserInterface/PropertyEditor/PropertyEditorConstants.cpp
+++ b/Engine/Source/Editor/PropertyEditor/Private/UserInterface/PropertyEditor/PropertyEditorConstants.cpp
@@ -7,6 +7,9 @@
 const FName PropertyEditorConstants::PropertyFontStyle( TEXT("PropertyWindow.NormalFont") );
 const FName PropertyEditorConstants::CategoryFontStyle( TEXT("PropertyWindow.BoldFont") );
 
+const FName PropertyEditorConstants::MD_Bitmask( TEXT("Bitmask") );
+const FName PropertyEditorConstants::MD_BitmaskEnum( TEXT("BitmaskEnum") );
+
 const FSlateBrush* PropertyEditorConstants::GetOverlayBrush( const TSharedRef< class FPropertyEditor > PropertyEditor )
 {
 	return FEditorStyle::GetBrush( TEXT("PropertyWindow.NoOverlayColor") );
diff --git a/Engine/Source/Editor/PropertyEditor/Private/UserInterface/PropertyEditor/PropertyEditorConstants.h b/Engine/Source/Editor/PropertyEditor/Private/UserInterface/PropertyEditor/PropertyEditorConstants.h
--- a/Engine/Source/Editor/PropertyEditor/Private/UserInterface/PropertyEditor/PropertyEditorConstants.h
+++ b/Engine/Source/Editor/PropertyEditor/Private/UserInterface/PropertyEditor/PropertyEditorConstants.h
@@ -8,5 +8,8 @@ public:
 	static const FName PropertyFontStyle;
 	static const FName CategoryFontStyle;
 
+	static const FName MD_Bitmask;
+	static const FName MD_BitmaskEnum;
+
 	static const FSlateBrush* GetOverlayBrush( const TSharedRef< class FPropertyEditor > PropertyEditor );
 };
\ No newline at end of file
diff --git a/Engine/Source/Editor/PropertyEditor/Private/UserInterface/PropertyEditor/SPropertyEditorNumeric.h b/Engine/Source/Editor/PropertyEditor/Private/UserInterface/PropertyEditor/SPropertyEditorNumeric.h
--- a/Engine/Source/Editor/PropertyEditor/Private/UserInterface/PropertyEditor/SPropertyEditorNumeric.h
+++ b/Engine/Source/Editor/PropertyEditor/Private/UserInterface/PropertyEditor/SPropertyEditorNumeric.h
@@ -29,112 +29,250 @@ public:
 		const TSharedRef< FPropertyNode > PropertyNode = InPropertyEditor->GetPropertyNode();
 		const UProperty* Property = InPropertyEditor->GetProperty();
 
-		// Instance metadata overrides per-class metadata.
-		auto GetMetaDataFromKey = [&PropertyNode, &Property](const FName& Key) -> const FString&
+		if(!Property->IsA(UFloatProperty::StaticClass()) && Property->HasMetaData(PropertyEditorConstants::MD_Bitmask))
 		{
-			const FString* InstanceValue = PropertyNode->GetInstanceMetaData(Key);
-			return (InstanceValue != nullptr) ? *InstanceValue : Property->GetMetaData(Key);
-		};
+			const int32 BitmaskBitCount = sizeof(NumericType) << 3;
 
-		const FString& MetaUIMinString = GetMetaDataFromKey("UIMin");
-		const FString& MetaUIMaxString = GetMetaDataFromKey("UIMax");
-		const FString& SliderExponentString = GetMetaDataFromKey("SliderExponent");
-		const FString& DeltaString = GetMetaDataFromKey("Delta");
-		const FString& ClampMinString = GetMetaDataFromKey("ClampMin");
-		const FString& ClampMaxString = GetMetaDataFromKey("ClampMax");
+			BitmaskFlags.Empty(BitmaskBitCount);
 
-		// If no UIMin/Max was specified then use the clamp string
-		const FString& UIMinString = MetaUIMinString.Len() ? MetaUIMinString : ClampMinString;
-		const FString& UIMaxString = MetaUIMaxString.Len() ? MetaUIMaxString : ClampMaxString;
+			const UEnum* BitmaskEnum = nullptr;
+			const FString& BitmaskEnumName = Property->GetMetaData(PropertyEditorConstants::MD_BitmaskEnum);
+			if (!BitmaskEnumName.IsEmpty())
+			{
+				// @TODO: Potentially replace this with a parameter passed in from a member variable on the UProperty (e.g. UByteProperty::Enum)
+				BitmaskEnum = FindObject<UEnum>(ANY_PACKAGE, *BitmaskEnumName);
+			}
 
-		NumericType ClampMin = TNumericLimits<NumericType>::Lowest();
-		NumericType ClampMax = TNumericLimits<NumericType>::Max();
+			if (BitmaskEnum)
+			{
+				for (int32 BitmaskEnumIndex = 0; BitmaskEnumIndex < BitmaskEnum->NumEnums() - 1; ++BitmaskEnumIndex)
+				{
+					int8 BitmaskFlagIndex = BitmaskEnum->GetValueByIndex(BitmaskEnumIndex);
+					if (BitmaskFlagIndex >= 0 && BitmaskFlagIndex < BitmaskBitCount)
+					{
+						FBitmaskFlagInfo* BitmaskFlag = new(BitmaskFlags) FBitmaskFlagInfo();
+
+						BitmaskFlag->Value = TBitmaskValueHelpers<NumericType>::LeftShift(static_cast<NumericType>(1), static_cast<int32>(BitmaskFlagIndex));
+
+						BitmaskFlag->DisplayName = BitmaskEnum->GetDisplayNameText(BitmaskEnumIndex);
+						if (BitmaskFlag->DisplayName.IsEmpty())
+						{
+							BitmaskFlag->DisplayName = FText::FromString(BitmaskEnum->GetEnumName(BitmaskEnumIndex));
+						}
+
+						BitmaskFlag->ToolTipText = BitmaskEnum->GetToolTipText(BitmaskEnumIndex);
+						if (BitmaskFlag->ToolTipText.IsEmpty())
+						{
+							BitmaskFlag->ToolTipText = FText::Format(LOCTEXT("BitmaskDefaultFlagToolTipText", "Toggle {0} on/off"), BitmaskFlag->DisplayName);
+						}
+					}
+				}
+			}
+			else
+			{
+				for (int32 BitmaskFlagIndex = 0; BitmaskFlagIndex < BitmaskBitCount; ++BitmaskFlagIndex)
+				{
+					FBitmaskFlagInfo* BitmaskFlag = new(BitmaskFlags) FBitmaskFlagInfo();
 
-		if (!ClampMinString.IsEmpty())
-		{
-			TTypeFromString<NumericType>::FromString(ClampMin, *ClampMinString);
-		}
+					BitmaskFlag->Value = TBitmaskValueHelpers<NumericType>::LeftShift(static_cast<NumericType>(1), BitmaskFlagIndex);
+					BitmaskFlag->DisplayName = FText::Format(LOCTEXT("BitmaskDefaultFlagDisplayName", "Flag {0}"), FText::AsNumber(BitmaskFlagIndex + 1));
+					BitmaskFlag->ToolTipText = FText::Format(LOCTEXT("BitmaskDefaultFlagToolTipText", "Toggle {0} on/off"), BitmaskFlag->DisplayName);
+				}
+			}
 
-		if (!ClampMaxString.IsEmpty())
-		{
-			TTypeFromString<NumericType>::FromString(ClampMax, *ClampMaxString);
-		}
+			const FComboBoxStyle& ComboBoxStyle = FCoreStyle::Get().GetWidgetStyle< FComboBoxStyle >("ComboBox");
 
-		NumericType UIMin = TNumericLimits<NumericType>::Lowest();
-		NumericType UIMax = TNumericLimits<NumericType>::Max();
-		TTypeFromString<NumericType>::FromString(UIMin, *UIMinString);
-		TTypeFromString<NumericType>::FromString(UIMax, *UIMaxString);
+			const auto& GetComboButtonText = [this]() -> FText
+			{
+				TOptional<NumericType> Value = OnGetValue();
+				if (Value.IsSet())
+				{
+					NumericType BitmaskValue = Value.GetValue();
+					if (BitmaskValue > 0)
+					{
+						if (TBitmaskValueHelpers<NumericType>::BitwiseAND(BitmaskValue, BitmaskValue - static_cast<NumericType>(1)))
+						{
+							return LOCTEXT("BitmaskButtonContentMultipleBitsSet", "(Mixed Flags)");
+						}
+						else
+						{
+							for (int i = 0; i < BitmaskFlags.Num(); ++i)
+							{
+								if (TBitmaskValueHelpers<NumericType>::BitwiseAND(BitmaskValue, BitmaskFlags[i].Value))
+								{
+									return BitmaskFlags[i].DisplayName;
+								}
+							}
+						}
+					}
+
+					return LOCTEXT("BitmaskButtonContentNoFlagsSet", "(No Flags Set)");
+				}
+				else
+				{
+					return LOCTEXT("MultipleValues", "Multiple Values");
+				}
+			};
 
-		NumericType SliderExponent = NumericType(1);
-		if (SliderExponentString.Len())
-		{
-			TTypeFromString<NumericType>::FromString(SliderExponent, *SliderExponentString);
-		}
+			// Constructs the UI for bitmask property editing.
+			SAssignNew(PrimaryWidget, SComboButton)
+			.ComboButtonStyle(&ComboBoxStyle.ComboButtonStyle)
+			.ContentPadding(FMargin(4.0, 2.0))
+			.ButtonContent()
+			[
+				SNew(STextBlock)
+				.Font(InArgs._Font)
+				.Text_Lambda(GetComboButtonText)
+			]
+			.OnGetMenuContent_Lambda([this]()
+			{
+				FMenuBuilder MenuBuilder(false, nullptr);
+
+				for (int i = 0; i < BitmaskFlags.Num(); ++i)
+				{
+					MenuBuilder.AddMenuEntry(
+						BitmaskFlags[i].DisplayName,
+						BitmaskFlags[i].ToolTipText,
+						FSlateIcon(),
+						FUIAction
+						(
+							FExecuteAction::CreateLambda([this, i]()
+							{
+								TOptional<NumericType> Value = OnGetValue();
+								if (Value.IsSet())
+								{
+									OnValueCommitted(TBitmaskValueHelpers<NumericType>::BitwiseXOR(Value.GetValue(), BitmaskFlags[i].Value), ETextCommit::Default);
+								}
+							}),
+							FCanExecuteAction(),
+							FIsActionChecked::CreateLambda([this, i]() -> bool
+							{
+								TOptional<NumericType> Value = OnGetValue();
+								if (Value.IsSet())
+								{
+									return TBitmaskValueHelpers<NumericType>::BitwiseAND(Value.GetValue(), BitmaskFlags[i].Value) != static_cast<NumericType>(0);
+								}
+
+								return false;
+							})
+						),
+						NAME_None,
+						EUserInterfaceActionType::Check);
+				}
 
-		NumericType Delta = NumericType(0);
-		if (DeltaString.Len())
-		{
-			TTypeFromString<NumericType>::FromString(Delta, *DeltaString);
-		}
+				return MenuBuilder.MakeWidget();
+			});
 
-		if (ClampMin >= ClampMax && (ClampMinString.Len() || ClampMaxString.Len()))
-		{
-			UE_LOG(LogPropertyNode, Warning, TEXT("Clamp Min (%s) >= Clamp Max (%s) for Ranged Numeric property %s"), *ClampMinString, *ClampMaxString, *Property->GetPathName());
+			ChildSlot.AttachWidget(PrimaryWidget.ToSharedRef());
 		}
+		else
+		{
+			// Instance metadata overrides per-class metadata.
+			auto GetMetaDataFromKey = [&PropertyNode, &Property](const FName& Key) -> const FString&
+			{
+				const FString* InstanceValue = PropertyNode->GetInstanceMetaData(Key);
+				return (InstanceValue != nullptr) ? *InstanceValue : Property->GetMetaData(Key);
+			};
 
-		const NumericType ActualUIMin = FMath::Max(UIMin, ClampMin);
-		const NumericType ActualUIMax = FMath::Min(UIMax, ClampMax);
+			const FString& MetaUIMinString = GetMetaDataFromKey("UIMin");
+			const FString& MetaUIMaxString = GetMetaDataFromKey("UIMax");
+			const FString& SliderExponentString = GetMetaDataFromKey("SliderExponent");
+			const FString& DeltaString = GetMetaDataFromKey("Delta");
+			const FString& ClampMinString = GetMetaDataFromKey("ClampMin");
+			const FString& ClampMaxString = GetMetaDataFromKey("ClampMax");
 
-		TOptional<NumericType> MinValue = ClampMinString.Len() ? ClampMin : TOptional<NumericType>();
-		TOptional<NumericType> MaxValue = ClampMaxString.Len() ? ClampMax : TOptional<NumericType>();
-		TOptional<NumericType> SliderMinValue = (UIMinString.Len()) ? ActualUIMin : TOptional<NumericType>();
-		TOptional<NumericType> SliderMaxValue = (UIMaxString.Len()) ? ActualUIMax : TOptional<NumericType>();
+			// If no UIMin/Max was specified then use the clamp string
+			const FString& UIMinString = MetaUIMinString.Len() ? MetaUIMinString : ClampMinString;
+			const FString& UIMaxString = MetaUIMaxString.Len() ? MetaUIMaxString : ClampMaxString;
 
-		if ((ActualUIMin >= ActualUIMax) && (SliderMinValue.IsSet() && SliderMaxValue.IsSet()))
-		{
-			UE_LOG(LogPropertyNode, Warning, TEXT("UI Min (%s) >= UI Max (%s) for Ranged Numeric property %s"), *UIMinString, *UIMaxString, *Property->GetPathName());
-		}
+			NumericType ClampMin = TNumericLimits<NumericType>::Lowest();
+			NumericType ClampMax = TNumericLimits<NumericType>::Max();
 
-		FObjectPropertyNode* ObjectPropertyNode = PropertyNode->FindObjectItemParent();
-		const bool bAllowSpin = (!ObjectPropertyNode || (1 == ObjectPropertyNode->GetNumObjects()))
-			&& !PropertyNode->GetProperty()->GetBoolMetaData("NoSpinbox");
+			if (!ClampMinString.IsEmpty())
+			{
+				TTypeFromString<NumericType>::FromString(ClampMin, *ClampMinString);
+			}
 
-		// Set up the correct type interface if we want to display units on the property editor
+			if (!ClampMaxString.IsEmpty())
+			{
+				TTypeFromString<NumericType>::FromString(ClampMax, *ClampMaxString);
+			}
 
-		// First off, check for ForceUnits= meta data. This meta tag tells us to interpret, and always display the value in these units. FUnitConversion::Settings().ShouldDisplayUnits does not apply to suce properties
-		const FString& ForcedUnits = InPropertyEditor->GetProperty()->GetMetaData(TEXT("ForceUnits"));
-		auto PropertyUnits = FUnitConversion::UnitFromString(*ForcedUnits);
-		if (PropertyUnits.IsSet())
-		{
-			// Create the type interface and set up the default input units if they are compatible
-			TypeInterface = MakeShareable(new TNumericUnitTypeInterface<NumericType>(PropertyUnits.GetValue()));
-			TypeInterface->FixedDisplayUnits = PropertyUnits.GetValue();
-		}
-		// If that's not set, we fall back to Units=xxx which calculates the most appropriate unit to display in
-		else
-		{
-			if (FUnitConversion::Settings().ShouldDisplayUnits())
+			NumericType UIMin = TNumericLimits<NumericType>::Lowest();
+			NumericType UIMax = TNumericLimits<NumericType>::Max();
+			TTypeFromString<NumericType>::FromString(UIMin, *UIMinString);
+			TTypeFromString<NumericType>::FromString(UIMax, *UIMaxString);
+
+			NumericType SliderExponent = NumericType(1);
+			if (SliderExponentString.Len())
 			{
-				const FString& DynamicUnits = InPropertyEditor->GetProperty()->GetMetaData(TEXT("Units"));
-				PropertyUnits = FUnitConversion::UnitFromString(*DynamicUnits);
+				TTypeFromString<NumericType>::FromString(SliderExponent, *SliderExponentString);
 			}
 
-			if (!PropertyUnits.IsSet())
+			NumericType Delta = NumericType(0);
+			if (DeltaString.Len())
 			{
-				PropertyUnits = EUnit::Unspecified;
+				TTypeFromString<NumericType>::FromString(Delta, *DeltaString);
 			}
 
-			// Create the type interface and set up the default input units if they are compatible
-			TypeInterface = MakeShareable(new TNumericUnitTypeInterface<NumericType>(PropertyUnits.GetValue()));
-			auto Value = OnGetValue();
+			if (ClampMin >= ClampMax && (ClampMinString.Len() || ClampMaxString.Len()))
+			{
+				UE_LOG(LogPropertyNode, Warning, TEXT("Clamp Min (%s) >= Clamp Max (%s) for Ranged Numeric property %s"), *ClampMinString, *ClampMaxString, *Property->GetPathName());
+			}
+
+			const NumericType ActualUIMin = FMath::Max(UIMin, ClampMin);
+			const NumericType ActualUIMax = FMath::Min(UIMax, ClampMax);
 
-			if (Value.IsSet())
+			TOptional<NumericType> MinValue = ClampMinString.Len() ? ClampMin : TOptional<NumericType>();
+			TOptional<NumericType> MaxValue = ClampMaxString.Len() ? ClampMax : TOptional<NumericType>();
+			TOptional<NumericType> SliderMinValue = (UIMinString.Len()) ? ActualUIMin : TOptional<NumericType>();
+			TOptional<NumericType> SliderMaxValue = (UIMaxString.Len()) ? ActualUIMax : TOptional<NumericType>();
+
+			if ((ActualUIMin >= ActualUIMax) && (SliderMinValue.IsSet() && SliderMaxValue.IsSet()))
 			{
-				TypeInterface->SetupFixedDisplay(Value.GetValue());
+				UE_LOG(LogPropertyNode, Warning, TEXT("UI Min (%s) >= UI Max (%s) for Ranged Numeric property %s"), *UIMinString, *UIMaxString, *Property->GetPathName());
+			}
+
+			FObjectPropertyNode* ObjectPropertyNode = PropertyNode->FindObjectItemParent();
+			const bool bAllowSpin = (!ObjectPropertyNode || (1 == ObjectPropertyNode->GetNumObjects()))
+				&& !PropertyNode->GetProperty()->GetBoolMetaData("NoSpinbox");
+
+			// Set up the correct type interface if we want to display units on the property editor
+
+			// First off, check for ForceUnits= meta data. This meta tag tells us to interpret, and always display the value in these units. FUnitConversion::Settings().ShouldDisplayUnits does not apply to suce properties
+			const FString& ForcedUnits = InPropertyEditor->GetProperty()->GetMetaData(TEXT("ForceUnits"));
+			auto PropertyUnits = FUnitConversion::UnitFromString(*ForcedUnits);
+			if (PropertyUnits.IsSet())
+			{
+				// Create the type interface and set up the default input units if they are compatible
+				TypeInterface = MakeShareable(new TNumericUnitTypeInterface<NumericType>(PropertyUnits.GetValue()));
+				TypeInterface->FixedDisplayUnits = PropertyUnits.GetValue();
+			}
+			// If that's not set, we fall back to Units=xxx which calculates the most appropriate unit to display in
+			else
+			{
+				if (FUnitConversion::Settings().ShouldDisplayUnits())
+				{
+					const FString& DynamicUnits = InPropertyEditor->GetProperty()->GetMetaData(TEXT("Units"));
+					PropertyUnits = FUnitConversion::UnitFromString(*DynamicUnits);
+				}
+
+				if (!PropertyUnits.IsSet())
+				{
+					PropertyUnits = EUnit::Unspecified;
+				}
+
+				// Create the type interface and set up the default input units if they are compatible
+				TypeInterface = MakeShareable(new TNumericUnitTypeInterface<NumericType>(PropertyUnits.GetValue()));
+				auto Value = OnGetValue();
+
+				if (Value.IsSet())
+				{
+					TypeInterface->SetupFixedDisplay(Value.GetValue());
+				}
 			}
-		}
 
-		ChildSlot
+			ChildSlot
 			[
 				SAssignNew(PrimaryWidget, SNumericEntryBox<NumericType>)
 				// Only allow spinning if we have a single value
@@ -155,6 +293,7 @@ public:
 				.OnEndSliderMovement(this, &SPropertyEditorNumeric<NumericType>::OnEndSliderMovement)
 				.TypeInterface(TypeInterface)
 			];
+		}
 
 		SetEnabled(TAttribute<bool>(this, &SPropertyEditorNumeric<NumericType>::CanEdit));
 	}
@@ -173,9 +312,10 @@ public:
 	void GetDesiredWidth( float& OutMinDesiredWidth, float& OutMaxDesiredWidth )
 	{
 		const UProperty* Property = PropertyEditor->GetProperty();
+		const bool bIsBitmask = (!Property->IsA(UFloatProperty::StaticClass()) && Property->HasMetaData(PropertyEditorConstants::MD_Bitmask));
 		const bool bIsNonEnumByte = (Property->IsA(UByteProperty::StaticClass()) && Cast<const UByteProperty>(Property)->Enum == NULL);
 
-		if( bIsNonEnumByte )
+		if( bIsNonEnumByte && !bIsBitmask )
 		{
 			OutMinDesiredWidth = 75.0f;
 			OutMaxDesiredWidth = 75.0f;
@@ -183,7 +323,7 @@ public:
 		else
 		{
 			OutMinDesiredWidth = 125.0f;
-			OutMaxDesiredWidth = 125.0f;
+			OutMaxDesiredWidth = bIsBitmask ? 400.0f : 125.0f;
 		}
 	}
 
@@ -368,6 +508,32 @@ private:
 		return PropertyEditor.IsValid() ? !PropertyEditor->IsEditConst() : true;
 	}
 
+	/** Flag data for bitmasks. */
+	struct FBitmaskFlagInfo
+	{
+		NumericType Value;
+		FText DisplayName;
+		FText ToolTipText;
+	};
+
+	/** Integral bitmask value helper methods. */
+	template<typename T, typename U = void>
+	struct TBitmaskValueHelpers
+	{
+		static T BitwiseAND(T Base, T Mask) { return Base & Mask; }
+		static T BitwiseXOR(T Base, T Mask) { return Base ^ Mask; }
+		static T LeftShift(T Base, int32 Shift) { return Base << Shift; }
+	};
+
+	/** Explicit specialization for numeric 'float' types (these will not be used). */
+	template<typename U>
+	struct TBitmaskValueHelpers<float, U>
+	{
+		static float BitwiseAND(float Base, float Mask) { return 0.0f; }
+		static float BitwiseXOR(float Base, float Mask) { return 0.0f; }
+		static float LeftShift(float Base, int32 Shift) { return 0.0f; }
+	};
+
 private:
 
 	TSharedPtr<TNumericUnitTypeInterface<NumericType>> TypeInterface;	
@@ -378,6 +544,9 @@ private:
 
 	/** True if the slider is being used to change the value of the property */
 	bool bIsUsingSlider;
+
+	/** Internal array of per-flag info, used for bitmask editing. */
+	TArray<FBitmaskFlagInfo> BitmaskFlags;
 };
 
 #undef LOCTEXT_NAMESPACE
\ No newline at end of file
diff --git a/Engine/Source/Editor/UnrealEd/Private/ComponentTypeRegistry.cpp b/Engine/Source/Editor/UnrealEd/Private/ComponentTypeRegistry.cpp
--- a/Engine/Source/Editor/UnrealEd/Private/ComponentTypeRegistry.cpp
+++ b/Engine/Source/Editor/UnrealEd/Private/ComponentTypeRegistry.cpp
@@ -288,23 +288,22 @@ void FComponentTypeRegistryData::ForceRefreshComponentList()
 			// that the Asset Registry could just keep asset paths:
 			TArray<FAssetData> BlueprintAssetData;
 			AssetRegistryModule.Get().GetAssetsByClass(UBlueprint::StaticClass()->GetFName(), BlueprintAssetData);
+			TMap<FString, FAssetData> BlueprintNames;
+			for (const FAssetData& Blueprint : BlueprintAssetData)
+			{
+				const auto& name = Blueprint.AssetName.ToString();
+				BlueprintNames.Add(name, Blueprint);
+			}
 
-			for (auto OnDiskClass : OnDiskClasses)
+			for (const FName& OnDiskClass : OnDiskClasses)
 			{
 				FString FixedString = OnDiskClass.ToString();
 				FixedString.RemoveFromEnd(TEXT("_C"));
 
 				FAssetData AssetData;
+				if (const FAssetData* Value = BlueprintNames.Find(FixedString))
 				{
-					// find asset data for this unloaded class, this makes our logic n^2 and may provide opportunity for further performance improvement:
-					for (auto& Blueprint : BlueprintAssetData)
-					{
-						if (Blueprint.AssetName.ToString() == FixedString)
-						{
-							AssetData = Blueprint;
-							break;
-						}
-					}
+					AssetData = *Value;
 				}
 
 				FComponentTypeEntry Entry = { FixedString, AssetData.ObjectPath.ToString(), nullptr };
diff --git a/Engine/Source/Editor/UnrealEd/Private/EdGraphUtilities.cpp b/Engine/Source/Editor/UnrealEd/Private/EdGraphUtilities.cpp
--- a/Engine/Source/Editor/UnrealEd/Private/EdGraphUtilities.cpp
+++ b/Engine/Source/Editor/UnrealEd/Private/EdGraphUtilities.cpp
@@ -6,6 +6,8 @@
 #include "Factories.h"
 #include "Editor/UnrealEd/Public/Kismet2/BlueprintEditorUtils.h"
 #include "UnrealExporter.h"
+#include "SNotificationList.h"
+#include "NotificationManager.h"
 
 /////////////////////////////////////////////////////
 // FGraphObjectTextFactory
@@ -15,7 +17,9 @@ struct FGraphObjectTextFactory : public FCustomizableTextObjectFactory
 {
 public:
 	TSet<UEdGraphNode*> SpawnedNodes;
+	TSet<UEdGraphNode*> SubstituteNodes;
 	const UEdGraph* DestinationGraph;
+	TArray<FName> ExtraNamesInUse;
 public:
 	FGraphObjectTextFactory(const UEdGraph* InDestinationGraph)
 		: FCustomizableTextObjectFactory(GWarn)
@@ -54,7 +58,8 @@ protected:
 			if(!Node->CanPasteHere(DestinationGraph))
 			{
 				// Attempt to create a substitute node if it cannot be pasted (note: the return value can be NULL, indicating that the node cannot be pasted into the graph)
-				Node = DestinationGraph->GetSchema()->CreateSubstituteNode(Node, DestinationGraph, &InstanceGraph);
+				Node = DestinationGraph->GetSchema()->CreateSubstituteNode(Node, DestinationGraph, &InstanceGraph, ExtraNamesInUse);
+				SubstituteNodes.Add(Node);
 			}
 
 			if(Node != CreatedObject)
@@ -72,6 +77,23 @@ protected:
 			}
 		}
 	}
+
+	virtual void PostProcessConstructedObjects() override
+	{
+		if (SubstituteNodes.Num() > 0)
+		{
+			// Display a notification to inform the user that the variable type was invalid (likely due to corruption), it should no longer appear in the list.
+			FNotificationInfo Info(NSLOCTEXT("EdGraphUtilities", "SubstituteNodesWarning", "Conflicting nodes substituted during paste!"));
+			Info.ExpireDuration = 3.0f;
+			Info.bUseLargeFont = false;
+			Info.Image = FCoreStyle::Get().GetBrush(TEXT("MessageLog.Warning"));
+			TSharedPtr<SNotificationItem> Notification = FSlateNotificationManager::Get().AddNotification(Info);
+			if (Notification.IsValid())
+			{
+				Notification->SetCompletionState(SNotificationItem::CS_None);
+			}
+		}
+	}
 };
 
 //////////////////////////////////////////////////////////////////////////
diff --git a/Engine/Source/Editor/UnrealEd/Private/EditorCategoryUtils.cpp b/Engine/Source/Editor/UnrealEd/Private/EditorCategoryUtils.cpp
--- a/Engine/Source/Editor/UnrealEd/Private/EditorCategoryUtils.cpp
+++ b/Engine/Source/Editor/UnrealEd/Private/EditorCategoryUtils.cpp
@@ -285,7 +285,7 @@ FString FEditorCategoryUtils::GetCategoryDisplayString(FString const& Unsanitize
 }
 
 //------------------------------------------------------------------------------
-void FEditorCategoryUtils::GetClassHideCategories(UClass const* Class, TArray<FString>& CategoriesOut)
+void FEditorCategoryUtils::GetClassHideCategories(UClass const* Class, TArray<FString>& CategoriesOut, bool bHomogenize)
 {
 	CategoriesOut.Empty();
 
@@ -296,9 +296,12 @@ void FEditorCategoryUtils::GetClassHideCategories(UClass const* Class, TArray<FS
 
 		HideCategories.ParseIntoArray(CategoriesOut, TEXT(" "), /*InCullEmpty =*/true);
 		
-		for (FString& Category : CategoriesOut)
+		if (bHomogenize)
 		{
-			Category = GetCategoryDisplayString(Category);
+			for (FString& Category : CategoriesOut)
+			{
+				Category = GetCategoryDisplayString(Category);
+			}
 		}
 	}
 }
diff --git a/Engine/Source/Editor/UnrealEd/Private/Factories/EditorFactories.cpp b/Engine/Source/Editor/UnrealEd/Private/Factories/EditorFactories.cpp
--- a/Engine/Source/Editor/UnrealEd/Private/Factories/EditorFactories.cpp
+++ b/Engine/Source/Editor/UnrealEd/Private/Factories/EditorFactories.cpp
@@ -5346,6 +5346,7 @@ void FCustomizableTextObjectFactory::ProcessBuffer(UObject* InParent, EObjectFla
 		ImportObjectProperties((uint8*)CreatedObject, *PropText, CreatedObject->GetClass(), CreatedObject, CreatedObject, WarningContext, 0, 0, &InstanceGraph);
 		ProcessConstructedObject(CreatedObject);
 	}
+	PostProcessConstructedObjects();
 }
 
 
diff --git a/Engine/Source/Editor/UnrealEd/Private/Kismet2/BlueprintEditorUtils.cpp b/Engine/Source/Editor/UnrealEd/Private/Kismet2/BlueprintEditorUtils.cpp
--- a/Engine/Source/Editor/UnrealEd/Private/Kismet2/BlueprintEditorUtils.cpp
+++ b/Engine/Source/Editor/UnrealEd/Private/Kismet2/BlueprintEditorUtils.cpp
@@ -1770,7 +1770,10 @@ void FBlueprintEditorUtils::RecreateClassMetaData(UBlueprint* Blueprint, UClass*
 	{
 		if (!ParentClass->HasMetaData(FBlueprintMetadata::MD_IgnoreCategoryKeywordsInSubclasses))
 		{
-			FEditorCategoryUtils::GetClassHideCategories(ParentClass, AllHideCategories);
+			// we want the categories just as they appear in the parent class 
+			// (set bHomogenize to false) - especially since homogenization 
+			// could inject spaces
+			FEditorCategoryUtils::GetClassHideCategories(ParentClass, AllHideCategories, /*bHomogenize =*/false);
 			if (ParentClass->HasMetaData(TEXT("ShowCategories")))
 			{
 				Class->SetMetaData(TEXT("ShowCategories"), *ParentClass->GetMetaData("ShowCategories"));
@@ -1820,11 +1823,26 @@ void FBlueprintEditorUtils::RecreateClassMetaData(UBlueprint* Blueprint, UClass*
 
 	for (FString HideCategory : Blueprint->HideCategories)
 	{
-		HideCategory.ReplaceInline(TEXT(" "), TEXT(""));
+		auto& CharArray = HideCategory.GetCharArray();
+
+		int32 SpaceIndex = CharArray.Find(TEXT(' '));
+		while (SpaceIndex != INDEX_NONE)
+		{
+			CharArray.RemoveAt(SpaceIndex);
+			if (SpaceIndex >= CharArray.Num())
+			{
+				break;
+			}
+
+			TCHAR& WordStartingChar = CharArray[SpaceIndex];
+			WordStartingChar = FChar::ToUpper(WordStartingChar);
+
+			CharArray.Find(TEXT(' '), SpaceIndex);
+		}
 		AllHideCategories.Add(HideCategory);
 	}
 
-	if (AllHideCategories.Num())
+	if (AllHideCategories.Num() > 0)
 	{
 		Class->SetMetaData(TEXT("HideCategories"), *FString::Join(AllHideCategories, TEXT(" ")));
 	}
@@ -2799,7 +2817,7 @@ UK2Node_Event* FBlueprintEditorUtils::FindOverrideForFunction(const UBlueprint*
 		UK2Node_Event* EventNode = AllEvents[i];
 		check(EventNode);
 		if(	EventNode->bOverrideFunction == true &&
-			EventNode->EventReference.GetMemberParentClass(EventNode->GetBlueprintClassFromNode()) == SignatureClass &&
+			EventNode->EventReference.GetMemberParentClass(EventNode->GetBlueprintClassFromNode())->IsChildOf(SignatureClass) &&
 			EventNode->EventReference.GetMemberName() == SignatureName )
 		{
 			return EventNode;
@@ -2809,6 +2827,28 @@ UK2Node_Event* FBlueprintEditorUtils::FindOverrideForFunction(const UBlueprint*
 	return NULL;
 }
 
+UK2Node_Event* FBlueprintEditorUtils::FindCustomEventNode(const UBlueprint* Blueprint, FName const CustomName)
+{
+	UK2Node_Event* FoundNode = nullptr;
+
+	if (CustomName != NAME_None)
+	{
+		TArray<UK2Node_Event*> AllEvents;
+		FBlueprintEditorUtils::GetAllNodesOfClass<UK2Node_Event>(Blueprint, AllEvents);
+
+		for (UK2Node_Event* EventNode : AllEvents)
+		{
+			if (EventNode->CustomFunctionName == CustomName)
+			{
+				FoundNode = EventNode;
+				break;
+			}
+		}
+	}
+	return FoundNode;
+}
+
+
 void FBlueprintEditorUtils::GatherDependencies(const UBlueprint* InBlueprint, TSet<TWeakObjectPtr<UBlueprint>>& Dependencies, TSet<TWeakObjectPtr<UStruct>>& OutUDSDependencies)
 {
 	struct FGatherDependenciesHelper
@@ -2976,8 +3016,9 @@ bool FBlueprintEditorUtils::IsDataOnlyBlueprint(const UBlueprint* Blueprint)
 		return false;
 	}
 	
-	// No extra functions, other than the user construction script
-	if ((Blueprint->FunctionGraphs.Num() > 1) || (Blueprint->MacroGraphs.Num() > 0))
+	// No extra functions, other than the user construction script(only AActor and subclasses of AActor have)
+	auto DefaultFunctionNum = (Blueprint->ParentClass && Blueprint->ParentClass->IsChildOf(AActor::StaticClass())) ? 1 : 0;
+	if ((Blueprint->FunctionGraphs.Num() > DefaultFunctionNum) || (Blueprint->MacroGraphs.Num() > 0))
 	{
 		return false;
 	}
@@ -5784,7 +5825,8 @@ namespace
 			}
 			else
 			{
-				UEdGraphNode* CustomEventNode = CurrentGraph->GetSchema()->CreateSubstituteNode(EventNode, CurrentGraph, NULL);
+				TArray<FName> DummyExtraNameList;
+				UEdGraphNode* CustomEventNode = CurrentGraph->GetSchema()->CreateSubstituteNode(EventNode, CurrentGraph, NULL, DummyExtraNameList);
 				if (ensure(CustomEventNode))
 				{
 					// Destroy the old event node (this will also break all pin links and remove it from the graph)
@@ -5899,6 +5941,7 @@ static void ConformInterfaceByName(UBlueprint* Blueprint, FBPInterfaceDescriptio
 
 		// check to make sure that there aren't any interface methods that we originally 
 		// implemented as events, but have since switched to functions 
+		TArray<FName> ExtraNameList;
 		for (UK2Node_Event* EventNode : ImplementedEvents)
 		{
 			// if this event belongs to something other than this interface
@@ -5916,7 +5959,7 @@ static void ConformInterfaceByName(UBlueprint* Blueprint, FBPInterfaceDescriptio
 
 			UEdGraph* EventGraph = EventNode->GetGraph();
 			// we've already implemented this interface function as an event (which we need to replace)
-			UK2Node_CustomEvent* CustomEventNode = Cast<UK2Node_CustomEvent>(EventGraph->GetSchema()->CreateSubstituteNode(EventNode, EventGraph, NULL));
+			UK2Node_CustomEvent* CustomEventNode = Cast<UK2Node_CustomEvent>(EventGraph->GetSchema()->CreateSubstituteNode(EventNode, EventGraph, NULL, ExtraNameList));
 			if (CustomEventNode == NULL)
 			{
 				continue;
@@ -8353,7 +8396,7 @@ void FBlueprintEditorUtils::HandleDisableEditableWhenInherited(UObject* Modified
 void FBlueprintEditorUtils::BuildComponentInstancingData(UActorComponent* ComponentTemplate, FBlueprintCookedComponentInstancingData& OutData)
 {
 	// Recursively gathers properties that differ from class/struct defaults, and fills out the cooked property list structure.
-	TFunction<void(UStruct*, const uint8*, const uint8*)> RecursivePropertyGatherLambda = [&](UStruct* InStruct, const uint8* DataPtr, const uint8* DefaultDataPtr)
+	TFunction<void(UStruct*, const uint8*, const uint8*)> RecursivePropertyGatherLambda = [&OutData, &RecursivePropertyGatherLambda](UStruct* InStruct, const uint8* DataPtr, const uint8* DefaultDataPtr)
 	{
 		for (UProperty* Property = InStruct->PropertyLink; Property; Property = Property->PropertyLinkNext)
 		{
@@ -8375,6 +8418,59 @@ void FBlueprintEditorUtils::BuildComponentInstancingData(UActorComponent* Compon
 
 						RecursivePropertyGatherLambda(StructProperty->Struct, PropertyValue, DefaultPropertyValue);
 
+						// Prepend the struct property only if there is at least one changed sub-property.
+						if (NumChangedProperties < OutData.ChangedPropertyList.Num())
+						{
+							OutData.ChangedPropertyList.Insert(ChangedPropertyInfo, NumChangedProperties);
+						}
+					}
+					else if (UArrayProperty* ArrayProperty = Cast<UArrayProperty>(Property))
+					{
+						FScriptArrayHelper ArrayValueHelper(ArrayProperty, PropertyValue);
+						FScriptArrayHelper DefaultArrayValueHelper(ArrayProperty, DefaultPropertyValue);
+
+						int32 NumChangedProperties = OutData.ChangedPropertyList.Num();
+						FBlueprintComponentChangedPropertyInfo ChangedArrayPropertyInfo = ChangedPropertyInfo;
+
+						for (int32 ArrayValueIndex = 0; ArrayValueIndex < ArrayValueHelper.Num(); ++ArrayValueIndex)
+						{
+							ChangedArrayPropertyInfo.ArrayIndex = ArrayValueIndex;
+							const uint8* ArrayPropertyValue = ArrayValueHelper.GetRawPtr(ArrayValueIndex);
+
+							if (ArrayValueIndex < DefaultArrayValueHelper.Num())
+							{
+								const uint8* DefaultArrayPropertyValue = DefaultArrayValueHelper.GetRawPtr(ArrayValueIndex);
+
+								if (UStructProperty* InnerStructProperty = Cast<UStructProperty>(ArrayProperty->Inner))
+								{
+									int32 NumChangedArrayProperties = OutData.ChangedPropertyList.Num();
+
+									RecursivePropertyGatherLambda(InnerStructProperty->Struct, ArrayPropertyValue, DefaultArrayPropertyValue);
+
+									// Prepend the struct property only if there is at least one changed sub-property.
+									if (NumChangedArrayProperties < OutData.ChangedPropertyList.Num())
+									{
+										OutData.ChangedPropertyList.Insert(ChangedArrayPropertyInfo, NumChangedArrayProperties);
+									}
+								}
+								else if(!ArrayProperty->Inner->Identical(ArrayPropertyValue, DefaultArrayPropertyValue, PPF_None))
+								{
+									// Emit the index of the individual array value that differs from the default value
+									OutData.ChangedPropertyList.Add(ChangedArrayPropertyInfo);
+								}
+							}
+							else
+							{
+								// Emit the "end" of differences with the default value (signals that remaining values should be copied in full)
+								ChangedArrayPropertyInfo.PropertyName = NAME_None;
+								OutData.ChangedPropertyList.Add(ChangedArrayPropertyInfo);
+
+								// Don't need to record anything else.
+								break;
+							}
+						}
+
+						// Prepend the array property as changed only if we also wrote out any of the inner value as changed.
 						if (NumChangedProperties < OutData.ChangedPropertyList.Num())
 						{
 							OutData.ChangedPropertyList.Insert(ChangedPropertyInfo, NumChangedProperties);
diff --git a/Engine/Source/Editor/UnrealEd/Private/Kismet2/CompilerResultsLog.cpp b/Engine/Source/Editor/UnrealEd/Private/Kismet2/CompilerResultsLog.cpp
--- a/Engine/Source/Editor/UnrealEd/Private/Kismet2/CompilerResultsLog.cpp
+++ b/Engine/Source/Editor/UnrealEd/Private/Kismet2/CompilerResultsLog.cpp
@@ -65,7 +65,7 @@ UObject const* FBacktrackMap::FindSourceObject(UObject const* PossiblyDuplicated
 //////////////////////////////////////////////////////////////////////////
 // FCompilerResultsLog
 
-FCompilerResultsLog::FCompilerResultsLog()
+FCompilerResultsLog::FCompilerResultsLog(bool bIsCompatibleWithEvents/* = true*/)
 	: NumErrors(0)
 	, NumWarnings(0)
 	, bSilentMode(false)
@@ -75,7 +75,7 @@ FCompilerResultsLog::FCompilerResultsLog()
 	, EventDisplayThresholdMs(0)
 {
 	CurrentEventScope = nullptr;
-	if(CurrentEventTarget == nullptr)
+	if(bIsCompatibleWithEvents && CurrentEventTarget == nullptr)
 	{
 		CurrentEventTarget = this;
 	}
@@ -281,7 +281,15 @@ void FCompilerResultsLog::InternalLogMessage(const EMessageSeverity::Type& Sever
 	{
 		if(Severity == EMessageSeverity::CriticalError || Severity == EMessageSeverity::Error)
 		{
-			UE_LOG(LogBlueprint, Error, TEXT("Compiler %s"), *Line->ToText().ToString());
+			if(IsRunningCommandlet())
+			{
+				UE_LOG(LogBlueprint, Error, TEXT("Compiler %s from Source: %s"), *Line->ToText().ToString(), *SourcePath);
+			}
+			else
+			{
+				// in editor the compiler log is 'rich' and we don't need to annotate with the blueprint name:
+				UE_LOG(LogBlueprint, Error, TEXT("Compiler %s"), *Line->ToText().ToString());
+			}
 		}
 		else if(Severity == EMessageSeverity::Warning || Severity == EMessageSeverity::PerformanceWarning)
 		{
diff --git a/Engine/Source/Editor/UnrealEd/Private/Kismet2/Kismet2.cpp b/Engine/Source/Editor/UnrealEd/Private/Kismet2/Kismet2.cpp
--- a/Engine/Source/Editor/UnrealEd/Private/Kismet2/Kismet2.cpp
+++ b/Engine/Source/Editor/UnrealEd/Private/Kismet2/Kismet2.cpp
@@ -498,7 +498,7 @@ UBlueprint* FKismetEditorUtilities::CreateBlueprint(UClass* ParentClass, UObject
 				FDefaultEventNodeData Data = DataIt.Value();
 				if (NewBP->GeneratedClass->IsChildOf(Data.TargetClass))
 				{
-					FKismetEditorUtilities::AddDefaultEventNode(NewBP, NewBP->UbergraphPages[0], Data.EventName, NewBP->GeneratedClass, NodePositionY);
+					FKismetEditorUtilities::AddDefaultEventNode(NewBP, NewBP->UbergraphPages[0], Data.EventName, Data.TargetClass, NodePositionY);
 				}
 			}
 
@@ -684,6 +684,9 @@ void FKismetEditorUtilities::CompileBlueprint(UBlueprint* BlueprintObj, bool bIs
 	FSecondsCounterScope Timer(BlueprintCompileAndLoadTimerData); 
 	BP_SCOPED_COMPILER_EVENT_STAT(EKismetCompilerStats_CompileBlueprint);
 
+	// Wipe the PreCompile log, any generated messages are now irrelevant
+	BlueprintObj->PreCompileLog.Reset();
+
 	// Broadcast pre-compile
 #if WITH_EDITOR
 	{
diff --git a/Engine/Source/Editor/UnrealEd/Private/Kismet2/KismetDebugUtilities.cpp b/Engine/Source/Editor/UnrealEd/Private/Kismet2/KismetDebugUtilities.cpp
--- a/Engine/Source/Editor/UnrealEd/Private/Kismet2/KismetDebugUtilities.cpp
+++ b/Engine/Source/Editor/UnrealEd/Private/Kismet2/KismetDebugUtilities.cpp
@@ -142,7 +142,11 @@ void FKismetDebugUtilities::OnScriptException(const UObject* ActiveObject, const
 		case EBlueprintExceptionType::AccessViolation:
 			if ( GIsEditor && GIsPlayInEditorWorld )
 			{
-				TSharedRef<FTokenizedMessage> ErrorMessage = FMessageLog("PIE").Error(LOCTEXT("RuntimeErrorMessage", "Blueprint Runtime Error from function:"));
+				// declared as its own variable since it's flushed (logs pushed
+				// to std output) on destruction - we want the full message 
+				// constructed before it's logged
+				FMessageLog PIEMessageLog("PIE");
+				TSharedRef<FTokenizedMessage> ErrorMessage = PIEMessageLog.Error(LOCTEXT("RuntimeErrorMessage", "Blueprint Runtime Error from function:"));
 
 				// NOTE: StackFrame.Node is not a blueprint node like you may think ("Node" has some legacy meaning)
 				FString GeneratedFuncName = FString::Printf(TEXT("'%s'"), *StackFrame.Node->GetName());
diff --git a/Engine/Source/Editor/UnrealEd/Public/EditorCategoryUtils.h b/Engine/Source/Editor/UnrealEd/Public/EditorCategoryUtils.h
--- a/Engine/Source/Editor/UnrealEd/Public/EditorCategoryUtils.h
+++ b/Engine/Source/Editor/UnrealEd/Public/EditorCategoryUtils.h
@@ -125,8 +125,9 @@ namespace FEditorCategoryUtils
 	 * 
 	 * @param  Class			The class you want to pull data from.
 	 * @param  CategoriesOut	An array that will be filled with a list of hidden categories.
+	 * @param  bHomogenize		Determines if the categories should be ran through expansion and display sanitation (useful even when not being displayed, for comparisons)
 	 */
-	UNREALED_API void GetClassHideCategories(UClass const* Class, TArray<FString>& CategoriesOut);
+	UNREALED_API void GetClassHideCategories(UClass const* Class, TArray<FString>& CategoriesOut, bool bHomogenize = true);
 
 	/**
 	 * Parses out the class's "ShowCategories" metadata, and returns it 
diff --git a/Engine/Source/Editor/UnrealEd/Public/Factories.h b/Engine/Source/Editor/UnrealEd/Public/Factories.h
--- a/Engine/Source/Editor/UnrealEd/Public/Factories.h
+++ b/Engine/Source/Editor/UnrealEd/Public/Factories.h
@@ -44,6 +44,9 @@ protected:
 	/** This is called on each created object after the property text is imported */
 	virtual void ProcessConstructedObject(UObject* CreatedObject);
 
+	/** Post handling of constructed objects by the factory */
+	virtual void PostProcessConstructedObjects() {};
+
 	/** Util to ensure that InName is a valid name for a new object within InParent. Will rename any existing object within InParent if it is called InName. */
 	static void ClearObjectNameUsage(UObject* InParent, FName InName);
 };
diff --git a/Engine/Source/Editor/UnrealEd/Public/Kismet2/BlueprintEditorUtils.h b/Engine/Source/Editor/UnrealEd/Public/Kismet2/BlueprintEditorUtils.h
--- a/Engine/Source/Editor/UnrealEd/Public/Kismet2/BlueprintEditorUtils.h
+++ b/Engine/Source/Editor/UnrealEd/Public/Kismet2/BlueprintEditorUtils.h
@@ -376,6 +376,9 @@ public:
 	/** Look to see if an event already exists to override a particular function */
 	static class UK2Node_Event* FindOverrideForFunction(const UBlueprint* Blueprint, const UClass* SignatureClass, FName SignatureName);
 
+	/** Find the Custom Event if it already exists in the Blueprint */
+	static class UK2Node_Event* FindCustomEventNode(const UBlueprint* Blueprint, FName const CustomName);
+
 
 	/** Returns all nodes in all graphs of the specified class */
 	template< class T > 
diff --git a/Engine/Source/Editor/UnrealEd/Public/Kismet2/CompilerResultsLog.h b/Engine/Source/Editor/UnrealEd/Public/Kismet2/CompilerResultsLog.h
--- a/Engine/Source/Editor/UnrealEd/Public/Kismet2/CompilerResultsLog.h
+++ b/Engine/Source/Editor/UnrealEd/Public/Kismet2/CompilerResultsLog.h
@@ -103,7 +103,7 @@ protected:
 	FString SourcePath;
 
 public:
-	FCompilerResultsLog();
+	FCompilerResultsLog(bool bIsCompatibleWithEvents = true);
 	virtual ~FCompilerResultsLog();
 
 	/** Register this log with the MessageLog module */
diff --git a/Engine/Source/Programs/UnrealHeaderTool/Private/HeaderParser.cpp b/Engine/Source/Programs/UnrealHeaderTool/Private/HeaderParser.cpp
--- a/Engine/Source/Programs/UnrealHeaderTool/Private/HeaderParser.cpp
+++ b/Engine/Source/Programs/UnrealHeaderTool/Private/HeaderParser.cpp
@@ -5388,7 +5388,11 @@ void FHeaderParser::CompileDelegateDeclaration(FClasses& AllClasses, const TCHAR
 			FError::Throwf(TEXT("Unexpected token following UDELEGATE(): %s"), Token.Identifier);
 
 		DelegateMacro = Token.Identifier;
-		CheckAllow(CurrentScopeName, ENestAllowFlags::TypeDecl);
+
+		//Workaround for UE-28897
+		const FStructScope* CurrentStructScope = TopNest->GetScope() ? TopNest->GetScope()->AsStructScope() : nullptr;
+		const bool bDynamicClassScope = CurrentStructScope && CurrentStructScope->GetStruct() && FClass::IsDynamic(CurrentStructScope->GetStruct());
+		CheckAllow(CurrentScopeName, bDynamicClassScope ? ENestAllowFlags::ImplicitDelegateDecl : ENestAllowFlags::TypeDecl);
 	}
 	else
 	{
diff --git a/Engine/Source/Runtime/CoreUObject/Private/Blueprint/BlueprintSupport.cpp b/Engine/Source/Runtime/CoreUObject/Private/Blueprint/BlueprintSupport.cpp
--- a/Engine/Source/Runtime/CoreUObject/Private/Blueprint/BlueprintSupport.cpp
+++ b/Engine/Source/Runtime/CoreUObject/Private/Blueprint/BlueprintSupport.cpp
@@ -863,8 +863,31 @@ void FLinkerLoad::ResolveDeferredDependencies(UStruct* LoadStruct)
 			{
 				FObjectImport& Import = ImportMap[ImportIndex];
 
-				const FLinkerLoad* SourceLinker = (Import.SourceLinker != nullptr) ? Import.SourceLinker :
-					(Import.XObject != nullptr) ? Import.XObject->GetLinker() : nullptr;
+				FLinkerLoad* SourceLinker = Import.SourceLinker;
+				// we cannot rely on Import.SourceLinker being set, if you look 
+				// at FLinkerLoad::CreateImport(), you'll see in game builds 
+				// that we try to circumvent the normal Import loading with a
+				// FindImportFast() call... if this is successful (the import 
+				// has already been somewhat loaded), then we don't fill out the 
+				// SourceLinker field
+				if (SourceLinker == nullptr)
+				{
+					if (Import.XObject != nullptr)
+					{
+						SourceLinker = Import.XObject->GetLinker();
+						//if (SourceLinker == nullptr)
+						//{
+						//	UPackage* ImportPkg = Import.XObject->GetOutermost();
+						//	// we use this package as placeholder for our 
+						//	// placeholders, so make sure to skip those (all other
+						//	// imports should belong to another package)
+						//	if (ImportPkg && ImportPkg != LinkerRoot)
+						//	{
+						//		SourceLinker = FindExistingLinkerForPackage(ImportPkg);
+						//	}
+						//}
+					}					
+				}
 
 				const UPackage* SourcePackage = (SourceLinker != nullptr) ? SourceLinker->LinkerRoot : nullptr;
 				// this package may not have introduced any (possible) cyclic 
@@ -901,9 +924,9 @@ void FLinkerLoad::ResolveDeferredDependencies(UStruct* LoadStruct)
 				{
 					// in case this is a user defined struct, we have to resolve any 
 					// deferred dependencies in the struct 
-					if (Import.SourceLinker != nullptr)
+					if (SourceLinker != nullptr)
 					{
-						Import.SourceLinker->ResolveDeferredDependencies(StructObj);
+						SourceLinker->ResolveDeferredDependencies(StructObj);
 					}
 				}
 			}
@@ -1047,10 +1070,10 @@ int32 FLinkerLoad::ResolveDependencyPlaceholder(FLinkerPlaceholderBase* Placehol
 	// if we can't rely on the Import object's RF_LoadCompleted flag, then its
 	// owner class should at least have it
 	DEFERRED_DEPENDENCY_CHECK( (RealImportObj == nullptr) || bExpectsLoadCompleteFlag ||
-		(FunctionOwner && FunctionOwner->HasAnyFlags(RF_LoadCompleted)) );
+		(FunctionOwner && FunctionOwner->HasAnyFlags(RF_LoadCompleted | RF_Dynamic)) );
 
 	DEFERRED_DEPENDENCY_CHECK(RealImportObj != PlaceholderObj);
-	DEFERRED_DEPENDENCY_CHECK(!bExpectsLoadCompleteFlag || RealImportObj->HasAnyFlags(RF_LoadCompleted));
+	DEFERRED_DEPENDENCY_CHECK(!bExpectsLoadCompleteFlag || RealImportObj->HasAnyFlags(RF_LoadCompleted | RF_Dynamic));
 #endif // USE_DEFERRED_DEPENDENCY_CHECK_VERIFICATION_TESTS
 
 	int32 ReplacementCount = 0;
@@ -1212,6 +1235,31 @@ void FLinkerLoad::FinalizeBlueprint(UClass* LoadClass)
 	// we need it ran for any super-classes before we regen
 	ResolveAllImports();
 
+	// interfaces can invalidate classes which implement them (when the  
+	// interface is regenerated), they essentially define the makeup of the  
+	// implementing class; so here, like we do with the parent class above, we  
+	// ensure that all implemented interfaces are finalized first - this helps  
+	// avoid cyclic issues where an interface ends up invalidating a dependent  
+	// class by being regenerated after the class (see UE-28846)
+	for (const FImplementedInterface& InterfaceDesc : LoadClass->Interfaces)
+	{
+		FLinkerLoad* InterfaceLinker = (InterfaceDesc.Class) ? InterfaceDesc.Class->GetLinker() : nullptr;
+		if ((InterfaceLinker != nullptr) && InterfaceLinker->IsBlueprintFinalizationPending())
+		{
+#if USE_DEFERRED_DEPENDENCY_CHECK_VERIFICATION_TESTS
+			// the interface import should have been properly resolved above, in 
+			// ResolveAllImports()
+			if (!ensure(!InterfaceLinker->HasUnresolvedDependencies()))
+#else 
+			if (InterfaceLinker->HasUnresolvedDependencies())
+#endif // USE_DEFERRED_DEPENDENCY_CHECK_VERIFICATION_TESTS
+			{
+				InterfaceLinker->ResolveDeferredDependencies(InterfaceDesc.Class);
+			}
+			InterfaceLinker->FinalizeBlueprint(InterfaceDesc.Class);
+		}
+	}
+
 	// replace any export placeholders that were created, and serialize in the 
 	// class's CDO
 	ResolveDeferredExports(LoadClass);
diff --git a/Engine/Source/Runtime/CoreUObject/Private/Serialization/AsyncLoading.cpp b/Engine/Source/Runtime/CoreUObject/Private/Serialization/AsyncLoading.cpp
--- a/Engine/Source/Runtime/CoreUObject/Private/Serialization/AsyncLoading.cpp
+++ b/Engine/Source/Runtime/CoreUObject/Private/Serialization/AsyncLoading.cpp
@@ -2014,6 +2014,16 @@ EAsyncPackageState::Type FAsyncPackage::PostLoadDeferredObjects(double InTickSta
 
 			if (Linker)
 			{
+				// Dynamic Class doesn't require/use pre-loading (or post-loading). 
+				// The CDO is created at this point, because now it's safe to solve cyclic dependencies.
+				for (FObjectExport& Export : Linker->ExportMap)
+				{
+					if (UDynamicClass* DynamicClass = Cast<UDynamicClass>(Export.Object))
+					{
+						DynamicClass->GetDefaultObject(true);
+					}
+				}
+
 				CreateClustersFromPackage(Linker);
 
 				// give a hint to the IO system that we are done with this file for now
diff --git a/Engine/Source/Runtime/CoreUObject/Private/UObject/Class.cpp b/Engine/Source/Runtime/CoreUObject/Private/UObject/Class.cpp
--- a/Engine/Source/Runtime/CoreUObject/Private/UObject/Class.cpp
+++ b/Engine/Source/Runtime/CoreUObject/Private/UObject/Class.cpp
@@ -4976,12 +4976,6 @@ void UDynamicClass::AddReferencedObjects(UObject* InThis, FReferenceCollector& C
 	Super::AddReferencedObjects(This, Collector);
 }
 
-void UDynamicClass::PostLoad()
-{
-	Super::PostLoad();
-	GetDefaultObject(true);
-}
-
 void UDynamicClass::PurgeClass(bool bRecompilingOnLoad)
 {
 	Super::PurgeClass(bRecompilingOnLoad);
diff --git a/Engine/Source/Runtime/CoreUObject/Private/UObject/LinkerLoad.cpp b/Engine/Source/Runtime/CoreUObject/Private/UObject/LinkerLoad.cpp
--- a/Engine/Source/Runtime/CoreUObject/Private/UObject/LinkerLoad.cpp
+++ b/Engine/Source/Runtime/CoreUObject/Private/UObject/LinkerLoad.cpp
@@ -2916,7 +2916,15 @@ void FLinkerLoad::LoadAllObjects( bool bForcePreload )
 		}
 #endif // USE_CIRCULAR_DEPENDENCY_LOAD_DEFERRING
 
-		CreateExportAndPreload(ExportIndex, bForcePreload);
+		UObject* LoadedObject = CreateExportAndPreload(ExportIndex, bForcePreload);
+		// DynamicClass could be created without calling CreateImport. The imported objects will be required later when a CDO is created.
+		if (UDynamicClass* DynamicClass = Cast<UDynamicClass>(LoadedObject))
+		{
+			for (int32 ImportIndex = 0; ImportIndex < ImportMap.Num(); ++ImportIndex)
+			{
+				CreateImport(ImportIndex);
+			}
+		}
 
 		// If needed send a heartbeat, but no need to do it too often
 		if (bShouldTickHeartBeat && (ExportIndex % 10) == 0)
@@ -3487,6 +3495,15 @@ UObject* FLinkerLoad::CreateExport( int32 Index )
 			if (Export.Object)
 			{
 				Export.Object->SetLinker(this, Index);
+				if (UDynamicClass* DynamicClass = Cast<UDynamicClass>(Export.Object))
+				{
+					// Dynamic Class doesn't require/use pre-loading (or post-loading), but at this point the class is not fully initialized. 
+					// The CDO is created (in a custom code) at the end of loading (when it's safe to solve cyclic dependencies).
+					if (!DynamicClass->GetDefaultObject(false))
+					{
+						FUObjectThreadContext::Get().ObjLoaded.Add(Export.Object);
+					}
+				}
 			}
 			return Export.Object;
 		}
diff --git a/Engine/Source/Runtime/CoreUObject/Private/UObject/Obj.cpp b/Engine/Source/Runtime/CoreUObject/Private/UObject/Obj.cpp
--- a/Engine/Source/Runtime/CoreUObject/Private/UObject/Obj.cpp
+++ b/Engine/Source/Runtime/CoreUObject/Private/UObject/Obj.cpp
@@ -738,7 +738,15 @@ void UObject::ConditionalPostLoad()
 
 		{
 			FExclusiveLoadPackageTimeTracker::FScopedPostLoadTracker Tracker(this);
-			PostLoad();
+
+			if (HasAnyFlags(RF_ClassDefaultObject))
+			{
+				GetClass()->PostLoadDefaultObject(this);
+			}
+			else
+			{
+				PostLoad();
+			}
 		}
 
 #if !(UE_BUILD_SHIPPING || UE_BUILD_TEST)
diff --git a/Engine/Source/Runtime/CoreUObject/Private/UObject/PropertyArray.cpp b/Engine/Source/Runtime/CoreUObject/Private/UObject/PropertyArray.cpp
--- a/Engine/Source/Runtime/CoreUObject/Private/UObject/PropertyArray.cpp
+++ b/Engine/Source/Runtime/CoreUObject/Private/UObject/PropertyArray.cpp
@@ -113,9 +113,61 @@ void UArrayProperty::SerializeItem( FArchive& Ar, void* Value, void const* Defau
 	// need to know how much data this call to SerializeItem consumes, so mark where we are
 	int32 DataOffset = Ar.Tell();
 
-	for( int32 i=0; i<n; i++ )
+	// If we're using a custom property list, first serialize any explicit indices
+	int32 i = 0;
+	bool bSerializeRemainingItems = true;
+	bool bUsingCustomPropertyList = Ar.ArUseCustomPropertyList;
+	if (bUsingCustomPropertyList && Ar.ArCustomPropertyList != nullptr)
 	{
-		Inner->SerializeItem( Ar, ArrayHelper.GetRawPtr(i) );
+		// Initially we only serialize indices that are explicitly specified (in order)
+		bSerializeRemainingItems = false;
+
+		const FCustomPropertyListNode* CustomPropertyList = Ar.ArCustomPropertyList;
+		const FCustomPropertyListNode* PropertyNode = CustomPropertyList;
+		while (PropertyNode && i < n && !bSerializeRemainingItems)
+		{
+			if (PropertyNode->Property != Inner)
+			{
+				// A null property value signals that we should serialize the remaining array values in full starting at this index
+				if (PropertyNode->Property == nullptr)
+				{
+					i = PropertyNode->ArrayIndex;
+				}
+
+				bSerializeRemainingItems = true;
+			}
+			else
+			{
+				// Set a temporary node to represent the item
+				FCustomPropertyListNode ItemNode = *PropertyNode;
+				ItemNode.ArrayIndex = 0;
+				ItemNode.PropertyListNext = nullptr;
+				Ar.ArCustomPropertyList = &ItemNode;
+
+				// Serialize the item at this array index
+				i = PropertyNode->ArrayIndex;
+				Inner->SerializeItem(Ar, ArrayHelper.GetRawPtr(i));
+				PropertyNode = PropertyNode->PropertyListNext;
+
+				// Restore the current property list
+				Ar.ArCustomPropertyList = CustomPropertyList;
+			}
+		}
+	}
+
+	if (bSerializeRemainingItems)
+	{
+		// Temporarily suspend the custom property list (as we need these items to be serialized in full)
+		Ar.ArUseCustomPropertyList = false;
+
+		// Serialize each item until we get to the end of the array
+		while (i < n)
+		{
+			Inner->SerializeItem(Ar, ArrayHelper.GetRawPtr(i++));
+		}
+
+		// Restore use of the custom property list (if it was previously enabled)
+		Ar.ArUseCustomPropertyList = bUsingCustomPropertyList;
 	}
 
 	if (Ar.UE4Ver() >= VER_UE4_INNER_ARRAY_TAG_INFO && Ar.IsSaving() && InnerTag.Type == NAME_StructProperty)
diff --git a/Engine/Source/Runtime/CoreUObject/Private/UObject/ScriptCore.cpp b/Engine/Source/Runtime/CoreUObject/Private/UObject/ScriptCore.cpp
--- a/Engine/Source/Runtime/CoreUObject/Private/UObject/ScriptCore.cpp
+++ b/Engine/Source/Runtime/CoreUObject/Private/UObject/ScriptCore.cpp
@@ -1534,7 +1534,7 @@ IMPLEMENT_VM_FUNCTION( EX_WireTracepoint, execWireTracepoint );
 void UObject::execInstrumentation( FFrame& Stack, RESULT_DECL )
 {
 #if !UE_BUILD_SHIPPING
-	const EScriptInstrumentation::Type EventType = static_cast<EScriptInstrumentation::Type>(Stack.ReadInt());
+	const EScriptInstrumentation::Type EventType = static_cast<EScriptInstrumentation::Type>(Stack.ReadInt<int32>());
 #if WITH_EDITORONLY_DATA
 	if (GIsEditor)
 	{
@@ -2284,10 +2284,22 @@ IMPLEMENT_VM_FUNCTION( EX_ClearMulticastDelegate, execClearMulticastDelegate );
 
 void UObject::execIntConst( FFrame& Stack, RESULT_DECL )
 {
-	*(int32*)RESULT_PARAM = Stack.ReadInt();
+	*(int32*)RESULT_PARAM = Stack.ReadInt<int32>();
 }
 IMPLEMENT_VM_FUNCTION( EX_IntConst, execIntConst );
 
+void UObject::execInt64Const(FFrame& Stack, RESULT_DECL)
+{
+	*(int64*)RESULT_PARAM = Stack.ReadInt<int64>();
+}
+IMPLEMENT_VM_FUNCTION(EX_Int64Const, execInt64Const);
+
+void UObject::execUInt64Const(FFrame& Stack, RESULT_DECL)
+{
+	*(uint64*)RESULT_PARAM = Stack.ReadInt<uint64>();
+}
+IMPLEMENT_VM_FUNCTION(EX_UInt64Const, execUInt64Const);
+
 void UObject::execSkipOffsetConst( FFrame& Stack, RESULT_DECL )
 {
 	CodeSkipSizeType Literal = Stack.ReadCodeSkipCount();
@@ -2473,7 +2485,7 @@ IMPLEMENT_VM_FUNCTION( EX_TransformConst, execTransformConst );
 void UObject::execStructConst( FFrame& Stack, RESULT_DECL )
 {
 	UScriptStruct* ScriptStruct = CastChecked<UScriptStruct>(Stack.ReadObject());
-	int32 SerializedSize = Stack.ReadInt();
+	int32 SerializedSize = Stack.ReadInt<int32>();
 
 	// Temporarily disabling this check because we can't assume the serialized size
 	// will match the struct size on all platforms (like win64 vs win32 cooked)
@@ -2517,7 +2529,7 @@ IMPLEMENT_VM_FUNCTION( EX_SetArray, execSetArray );
 void UObject::execArrayConst(FFrame& Stack, RESULT_DECL)
 {
 	UProperty* InnerProperty = CastChecked<UProperty>(Stack.ReadObject());
-	int32 Num = Stack.ReadInt();
+	int32 Num = Stack.ReadInt<int32>();
 	check(RESULT_PARAM);
 	FScriptArrayHelper ArrayHelper = FScriptArrayHelper::CreateHelperFormInnerProperty(InnerProperty, RESULT_PARAM);
 	ArrayHelper.EmptyValues(Num);
diff --git a/Engine/Source/Runtime/CoreUObject/Private/UObject/UObjectGlobals.cpp b/Engine/Source/Runtime/CoreUObject/Private/UObject/UObjectGlobals.cpp
--- a/Engine/Source/Runtime/CoreUObject/Private/UObject/UObjectGlobals.cpp
+++ b/Engine/Source/Runtime/CoreUObject/Private/UObject/UObjectGlobals.cpp
@@ -1460,6 +1460,16 @@ void EndLoad()
 				}
 			}
 
+			// Dynamic Class doesn't require/use pre-loading (or post-loading). 
+			// The CDO is created at this point, because now it's safe to solve cyclic dependencies.
+			for (UObject* Obj : ObjLoaded)
+			{
+				if (UDynamicClass* DynamicClass = Cast<UDynamicClass>(Obj))
+				{
+					DynamicClass->GetDefaultObject(true);
+				}
+			}
+
 			// Create clusters after all objects have been loaded
 			extern int32 GCreateGCClusters;
 			if (FPlatformProperties::RequiresCookedData() && !GIsInitialLoad && GCreateGCClusters && !GUObjectArray.IsOpenForDisregardForGC())
@@ -2531,10 +2541,14 @@ void FObjectInitializer::PostConstructInit()
 	if (!Obj->HasAnyFlags(RF_NeedLoad) || bIsDeferredInitializer)
 #endif // !USE_CIRCULAR_DEPENDENCY_LOAD_DEFERRING
 	{
-		if (bIsCDO || Class->HasAnyClassFlags(CLASS_PerObjectConfig))
+		if ((bIsCDO && !Class->HasAnyFlags(RF_Dynamic)) || Class->HasAnyClassFlags(CLASS_PerObjectConfig))
 		{
 			Obj->LoadConfig(NULL, NULL, bIsCDO ? UE4::LCPF_ReadParentSections : UE4::LCPF_None);
 		}
+		else if (bIsCDO && Class->HasAnyFlags(RF_Dynamic) && Class->HasAnyClassFlags(CLASS_Config))
+		{
+			Obj->LoadConfig(Class);
+		}
 		if (bAllowInstancing)
 		{
 			// Instance subobject templates for non-cdo blueprint classes or when using non-CDO template.
@@ -2799,6 +2813,14 @@ void FObjectInitializer::InitPropertiesFromCustomList(const FCustomPropertyListN
 				InitPropertiesFromCustomList(CustomPropertyListNode->SubPropertyList, StructProperty->Struct, PropertyValue, DefaultPropertyValue);
 			}
 		}
+		else if (const UArrayProperty* ArrayProperty = Cast<UArrayProperty>(CustomPropertyListNode->Property))
+		{
+			// This should never be NULL; we should not be recording the ArrayProperty without at least one sub property, but we'll verify just to be sure.
+			if (ensure(CustomPropertyListNode->SubPropertyList != nullptr))
+			{
+				InitArrayPropertyFromCustomList(ArrayProperty, CustomPropertyListNode->SubPropertyList, PropertyValue, DefaultPropertyValue);
+			}
+		}
 		else
 		{
 			CustomPropertyListNode->Property->CopySingleValue(PropertyValue, DefaultPropertyValue);
@@ -2806,6 +2828,61 @@ void FObjectInitializer::InitPropertiesFromCustomList(const FCustomPropertyListN
 	}
 }
 
+void FObjectInitializer::InitArrayPropertyFromCustomList(const UArrayProperty* ArrayProperty, const FCustomPropertyListNode* InPropertyList, uint8* DataPtr, const uint8* DefaultDataPtr)
+{
+	FScriptArrayHelper DstArrayValueHelper(ArrayProperty, DataPtr);
+	FScriptArrayHelper SrcArrayValueHelper(ArrayProperty, DefaultDataPtr);
+
+	int32 Num = SrcArrayValueHelper.Num();
+
+	if (!(ArrayProperty->Inner->PropertyFlags & CPF_IsPlainOldData))
+	{
+		DstArrayValueHelper.EmptyAndAddValues(Num);
+	}
+	else
+	{
+		DstArrayValueHelper.EmptyAndAddUninitializedValues(Num);
+	}
+
+	for (const FCustomPropertyListNode* CustomArrayPropertyListNode = InPropertyList; CustomArrayPropertyListNode; CustomArrayPropertyListNode = CustomArrayPropertyListNode->PropertyListNext)
+	{
+		int32 ArrayIndex = CustomArrayPropertyListNode->ArrayIndex;
+
+		uint8* DstArrayItemValue = DstArrayValueHelper.GetRawPtr(ArrayIndex);
+		const uint8* SrcArrayItemValue = SrcArrayValueHelper.GetRawPtr(ArrayIndex);
+
+		// A null property value signals that we should serialize the remaining array values in full starting at this index
+		if (CustomArrayPropertyListNode->Property == nullptr)
+		{
+			int32 Size = ArrayProperty->Inner->ElementSize;
+
+			if (!(ArrayProperty->Inner->PropertyFlags & CPF_IsPlainOldData))
+			{
+				for (int32 i = 0; i < Num - ArrayIndex; i++)
+				{
+					ArrayProperty->Inner->CopyCompleteValue(DstArrayItemValue + i * Size, SrcArrayItemValue + i * Size);
+				}
+			}
+			else
+			{
+				FMemory::Memcpy(DstArrayItemValue, SrcArrayItemValue, (Num - ArrayIndex) * Size);
+			}
+		}
+		else if (const UStructProperty* InnerStructProperty = Cast<UStructProperty>(ArrayProperty->Inner))
+		{
+			InitPropertiesFromCustomList(CustomArrayPropertyListNode->SubPropertyList, InnerStructProperty->Struct, DstArrayItemValue, SrcArrayItemValue);
+		}
+		else if (const UArrayProperty* InnerArrayProperty = Cast<UArrayProperty>(ArrayProperty->Inner))
+		{
+			InitArrayPropertyFromCustomList(InnerArrayProperty, CustomArrayPropertyListNode->SubPropertyList, DstArrayItemValue, SrcArrayItemValue);
+		}
+		else
+		{
+			ArrayProperty->Inner->CopyCompleteValue(DstArrayItemValue, SrcArrayItemValue);
+		}
+	}
+}
+
 bool FObjectInitializer::IslegalOverride(FName InComponentName, class UClass *DerivedComponentClass, class UClass *BaseComponentClass) const
 {
 	if (DerivedComponentClass && BaseComponentClass && !DerivedComponentClass->IsChildOf(BaseComponentClass) )
diff --git a/Engine/Source/Runtime/CoreUObject/Public/UObject/Class.h b/Engine/Source/Runtime/CoreUObject/Public/UObject/Class.h
--- a/Engine/Source/Runtime/CoreUObject/Public/UObject/Class.h
+++ b/Engine/Source/Runtime/CoreUObject/Public/UObject/Class.h
@@ -2384,6 +2384,11 @@ public:
 	 */
 	virtual void SerializeDefaultObject(UObject* Object, FArchive& Ar);
 
+	/** Wraps the PostLoad() call for the class default object.
+	 * @param Object the default object to call PostLoad() on
+	 */
+	virtual void PostLoadDefaultObject(UObject* Object) { Object->PostLoad(); }
+
 	/** 
 	 * Purges out the properties of this class in preparation for it to be regenerated
 	 * @param bRecompilingOnLoad - true if we are recompiling on load
@@ -2507,7 +2512,6 @@ public:
 
 	// UObject interface.
 	static void AddReferencedObjects(UObject* InThis, FReferenceCollector& Collector);
-	virtual void PostLoad() override;
 
 	// UClass interface
 	virtual void PurgeClass(bool bRecompilingOnLoad) override;
diff --git a/Engine/Source/Runtime/CoreUObject/Public/UObject/Script.h b/Engine/Source/Runtime/CoreUObject/Public/UObject/Script.h
--- a/Engine/Source/Runtime/CoreUObject/Public/UObject/Script.h
+++ b/Engine/Source/Runtime/CoreUObject/Public/UObject/Script.h
@@ -201,8 +201,8 @@ enum EExprToken
 	EX_EndArray				= 0x32,
 	//						= 0x33,
 	EX_UnicodeStringConst   = 0x34, // Unicode string constant.
-	//						= 0x35,
-	//						= 0x36,
+	EX_Int64Const			= 0x35,	// 64-bit integer constant.
+	EX_UInt64Const			= 0x36,	// 64-bit unsigned integer constant.
 	//						= 0x37,
 	EX_PrimitiveCast		= 0x38,	// A casting operator for primitives which reads the type as the subsequent byte
 	//						= 0x39,
diff --git a/Engine/Source/Runtime/CoreUObject/Public/UObject/ScriptSerialization.h b/Engine/Source/Runtime/CoreUObject/Public/UObject/ScriptSerialization.h
--- a/Engine/Source/Runtime/CoreUObject/Public/UObject/ScriptSerialization.h
+++ b/Engine/Source/Runtime/CoreUObject/Public/UObject/ScriptSerialization.h
@@ -316,6 +316,16 @@
 			XFER(int32);
 			break;
 		}
+		case EX_Int64Const:
+		{
+			XFER(int64);
+			break;
+		}
+		case EX_UInt64Const:
+		{
+			XFER(uint64);
+			break;
+		}
 		case EX_SkipOffsetConst:
 		{
 			XFER(CodeSkipSizeType);
diff --git a/Engine/Source/Runtime/CoreUObject/Public/UObject/Stack.h b/Engine/Source/Runtime/CoreUObject/Public/UObject/Stack.h
--- a/Engine/Source/Runtime/CoreUObject/Public/UObject/Stack.h
+++ b/Engine/Source/Runtime/CoreUObject/Public/UObject/Stack.h
@@ -136,7 +136,8 @@ public:
 	
 	COREUOBJECT_API static void KismetExecutionMessage(const TCHAR* Message, ELogVerbosity::Type Verbosity);
 
-	int32 ReadInt();
+	template<typename TNumericType>
+	TNumericType ReadInt();
 	float ReadFloat();
 	FName ReadName();
 	UObject* ReadObject();
@@ -199,15 +200,16 @@ inline FFrame::FFrame( UObject* InObject, UFunction* InNode, void* InLocals, FFr
 #endif
 }
 
-inline int32 FFrame::ReadInt()
+template<typename TNumericType>
+inline TNumericType FFrame::ReadInt()
 {
-	int32 Result;
+	TNumericType Result;
 #ifdef REQUIRES_ALIGNED_INT_ACCESS
-	FMemory::Memcpy(&Result, Code, sizeof(int32));
+	FMemory::Memcpy(&Result, Code, sizeof(TNumericType));
 #else
-	Result = *(int32*)Code;
+	Result = *(TNumericType*)Code;
 #endif
-	Code += sizeof(int32);
+	Code += sizeof(TNumericType);
 	return Result;
 }
 
diff --git a/Engine/Source/Runtime/CoreUObject/Public/UObject/UObject.h b/Engine/Source/Runtime/CoreUObject/Public/UObject/UObject.h
--- a/Engine/Source/Runtime/CoreUObject/Public/UObject/UObject.h
+++ b/Engine/Source/Runtime/CoreUObject/Public/UObject/UObject.h
@@ -1083,6 +1083,8 @@ public:
 
 	// Constants
 	DECLARE_FUNCTION(execIntConst);
+	DECLARE_FUNCTION(execInt64Const);
+	DECLARE_FUNCTION(execUInt64Const);
 	DECLARE_FUNCTION(execSkipOffsetConst);
 	DECLARE_FUNCTION(execFloatConst);
 	DECLARE_FUNCTION(execStringConst);
diff --git a/Engine/Source/Runtime/CoreUObject/Public/UObject/UObjectGlobals.h b/Engine/Source/Runtime/CoreUObject/Public/UObject/UObjectGlobals.h
--- a/Engine/Source/Runtime/CoreUObject/Public/UObject/UObjectGlobals.h
+++ b/Engine/Source/Runtime/CoreUObject/Public/UObject/UObjectGlobals.h
@@ -920,6 +920,16 @@ private:
 	 */
 	static void InitPropertiesFromCustomList(const FCustomPropertyListNode* InPropertyList, UStruct* InStruct, uint8* DataPtr, const uint8* DefaultDataPtr);
 
+	/**
+	* Helper method to assist with initializing from an array property with an explicit item list.
+	*
+	* @param	ArrayProperty		the array property for which the given property list applies
+	* @param	InPropertyList		only these properties (indices) will be copied from defaults
+	* @param	DataPtr				destination address (where to start copying values to)
+	* @param	DefaultDataPtr		source address (where to start copying the defaults data from)
+	*/
+	static void InitArrayPropertyFromCustomList(const UArrayProperty* ArrayProperty, const FCustomPropertyListNode* InPropertyList, uint8* DataPtr, const uint8* DefaultDataPtr);
+
 	bool IsInstancingAllowed() const;
 
 	/**
diff --git a/Engine/Source/Runtime/Engine/Classes/EdGraph/EdGraphSchema.h b/Engine/Source/Runtime/Engine/Classes/EdGraph/EdGraphSchema.h
--- a/Engine/Source/Runtime/Engine/Classes/EdGraph/EdGraphSchema.h
+++ b/Engine/Source/Runtime/Engine/Classes/EdGraph/EdGraphSchema.h
@@ -767,9 +767,10 @@ class ENGINE_API UEdGraphSchema : public UObject
 	 * @param	Node			The node to replace
 	 * @param	Graph			The destination graph
 	 * @param	InstanceGraph	Object instancing graph
+	 * @param	InOutExtraNames	List of extra names that are in-use from the substitution should be added to this list to prevent other substitutions from attempting to use them
 	 * @return					NULL if a substitute node cannot be created; otherwise, the substitute node instance
 	 */
-	virtual UEdGraphNode* CreateSubstituteNode(UEdGraphNode* Node, const UEdGraph* Graph, FObjectInstancingGraph* InstanceGraph) const { return NULL; }
+	virtual UEdGraphNode* CreateSubstituteNode(UEdGraphNode* Node, const UEdGraph* Graph, FObjectInstancingGraph* InstanceGraph, TArray<FName>& InOutExtraNames) const { return NULL; }
 
 
 	/**
diff --git a/Engine/Source/Runtime/Engine/Classes/Engine/Blueprint.h b/Engine/Source/Runtime/Engine/Classes/Engine/Blueprint.h
--- a/Engine/Source/Runtime/Engine/Classes/Engine/Blueprint.h
+++ b/Engine/Source/Runtime/Engine/Classes/Engine/Blueprint.h
@@ -730,6 +730,9 @@ public:
 	/** Message log for storing upgrade notes that were generated within the Blueprint, will be displayed to the compiler results each compiler and will remain until saving */
 	TSharedPtr<class FCompilerResultsLog> UpgradeNotesLog;
 
+	/** Message log for storing pre-compile errors/notes/warnings that will only last until the next Blueprint compile */
+	TSharedPtr<class FCompilerResultsLog> PreCompileLog;
+
 	/** 
 	 * Sends a message to the CurrentMessageLog, if there is one available.  Otherwise, defaults to logging to the normal channels.
 	 * Should use this for node and blueprint actions that happen during compilation!
diff --git a/Engine/Source/Runtime/Engine/Classes/Engine/BlueprintGeneratedClass.h b/Engine/Source/Runtime/Engine/Classes/Engine/BlueprintGeneratedClass.h
--- a/Engine/Source/Runtime/Engine/Classes/Engine/BlueprintGeneratedClass.h
+++ b/Engine/Source/Runtime/Engine/Classes/Engine/BlueprintGeneratedClass.h
@@ -518,6 +518,9 @@ protected:
 	/** Internal method used to help recursively build the cached property list for serialization. */
 	void BuildCachedPropertyList(FCustomPropertyListNode** CurrentNode, const UStruct* CurrentScope, int32* CurrentSourceIdx = nullptr) const;
 
+	/** Internal method used to help recursively build a cached sub property list from an array property for serialization. */
+	void BuildCachedArrayPropertyList(const UArrayProperty* ArraySubPropertyNode, FCustomPropertyListNode** CurrentNode, int32* CurrentSourceIdx) const;
+
 private:
 	/** Internal property list that's used in binary object serialization at component instancing time. */
 	mutable TIndirectArray<FCustomPropertyListNode> CachedPropertyListForSerialization;
@@ -616,6 +619,7 @@ public:
 	virtual UObject* GetArchetypeForCDO() const override;
 #endif //WITH_EDITOR
 	virtual void SerializeDefaultObject(UObject* Object, FArchive& Ar) override;
+	virtual void PostLoadDefaultObject(UObject* Object) override;
 	virtual bool IsFunctionImplementedInBlueprint(FName InFunctionName) const override;
 	virtual uint8* GetPersistentUberGraphFrame(UObject* Obj, UFunction* FuncToCheck) const override;
 	virtual void CreatePersistentUberGraphFrame(UObject* Obj, bool bCreateOnlyIfEmpty = false, bool bSkipSuperClass = false) const override;
@@ -691,6 +695,9 @@ protected:
 	/** Internal helper method used to recursively build the custom property list that's used for post-construct initialization. */
 	void BuildCustomPropertyListForPostConstruction(FCustomPropertyListNode*& InPropertyList, UStruct* InStruct, const uint8* DataPtr, const uint8* DefaultDataPtr);
 
+	/** Internal helper method used to recursively build a custom property list from an array property used for post-construct initialization. */
+	void BuildCustomArrayPropertyListForPostConstruction(UArrayProperty* ArrayProperty, FCustomPropertyListNode*& InPropertyList, const uint8* DataPtr, const uint8* DefaultDataPtr);
+
 private:
 	/** List of native class-owned properties that differ from defaults. This is used to optimize property initialization during post-construction by minimizing the number of native class-owned property values that get copied to the new instance. */
 	TIndirectArray<FCustomPropertyListNode> CustomPropertyListForPostConstruction;
diff --git a/Engine/Source/Runtime/Engine/Classes/Engine/EngineTypes.h b/Engine/Source/Runtime/Engine/Classes/Engine/EngineTypes.h
--- a/Engine/Source/Runtime/Engine/Classes/Engine/EngineTypes.h
+++ b/Engine/Source/Runtime/Engine/Classes/Engine/EngineTypes.h
@@ -2559,7 +2559,7 @@ struct FTimerHandle
 
 	FString ToString() const
 	{
-		return FString::Printf(TEXT("%d"), Handle);
+		return FString::Printf(TEXT("%ull"), Handle);
 	}
 
 private:
diff --git a/Engine/Source/Runtime/Engine/Classes/GameFramework/Actor.h b/Engine/Source/Runtime/Engine/Classes/GameFramework/Actor.h
--- a/Engine/Source/Runtime/Engine/Classes/GameFramework/Actor.h
+++ b/Engine/Source/Runtime/Engine/Classes/GameFramework/Actor.h
@@ -2245,6 +2245,9 @@ public:
 	 */
 	void ExecuteConstruction(const FTransform& Transform, const class FComponentInstanceDataCache* InstanceDataCache, bool bIsDefaultTransform = false);
 
+	// Forces PostLoad call on natively created components, that were originally in SCS
+	virtual void FixComponentsFromDynamicClass();
+
 	/**
 	 * Called when an instance of this class is placed (in editor) or spawned.
 	 * @param	Transform			The transform the actor was constructed at.
diff --git a/Engine/Source/Runtime/Engine/Classes/Kismet/KismetMathLibrary.h b/Engine/Source/Runtime/Engine/Classes/Kismet/KismetMathLibrary.h
--- a/Engine/Source/Runtime/Engine/Classes/Kismet/KismetMathLibrary.h
+++ b/Engine/Source/Runtime/Engine/Classes/Kismet/KismetMathLibrary.h
@@ -233,6 +233,10 @@ class ENGINE_API UKismetMathLibrary : public UBlueprintFunctionLibrary
 	UFUNCTION(BlueprintPure, meta=(DisplayName = "Bitwise OR", CompactNodeTitle = "|", Keywords = "| or", CommutativeAssociativeBinaryOperator = "true"), Category="Math|Integer")
 	static int32 Or_IntInt(int32 A, int32 B);
 
+	/* Bitwise NOT (~A) */
+	UFUNCTION(BlueprintPure, meta = (DisplayName = "Bitwise NOT", CompactNodeTitle = "~", Keywords = "~ not"), Category = "Math|Integer")
+	static int32 Not_Int(int32 A);
+
 	/* Sign (integer, returns -1 if A < 0, 0 if A is zero, and +1 if A > 0) */
 	UFUNCTION(BlueprintPure, meta=(DisplayName = "Sign (int)"), Category="Math|Integer")
 	static int32 SignOfInteger(int32 A);
