Fixed CanEditChange not being called when a property is changed that affects an unrelated property's edit constness.

diff --git a/Engine/Source/Editor/PropertyEditor/Private/PropertyNode.cpp b/Engine/Source/Editor/PropertyEditor/Private/PropertyNode.cpp
--- a/Engine/Source/Editor/PropertyEditor/Private/PropertyNode.cpp
+++ b/Engine/Source/Editor/PropertyEditor/Private/PropertyNode.cpp
@@ -1953,6 +1953,8 @@ void FPropertyNode::NotifyPostChange( FPropertyChangedEvent& InPropertyChangedEv
 	FObjectPropertyNode* ObjectNode = FindObjectItemParent();
 	if( ObjectNode )
 	{
+		ObjectNode->InvalidateCachedState();
+
 		UProperty* CurProperty = InPropertyChangedEvent.Property;
 
 		// Fire ULevel::LevelDirtiedEvent when falling out of scope.
@@ -2035,9 +2037,6 @@ void FPropertyNode::NotifyPostChange( FPropertyChangedEvent& InPropertyChangedEv
 		}
 	}
 
-	bUpdateEditConstState = true;
-	bUpdateDiffersFromDefault = true;
-
 	// Broadcast the change to any listeners
 	BroadcastPropertyChangedDelegates();
 
@@ -2207,6 +2206,17 @@ const FString* FPropertyNode::GetInstanceMetaData(const FName& Key) const
 	return InstanceMetaData.Find(Key);
 }
 
+void FPropertyNode::InvalidateCachedState()
+{
+	bUpdateDiffersFromDefault = true;
+	bUpdateEditConstState = true;
+
+	for( TSharedPtr<FPropertyNode>& ChildNode : ChildNodes )
+	{
+		ChildNode->InvalidateCachedState();
+	}
+}
+
 /**
  * Does the string compares to ensure this Name is acceptable to the filter that is passed in
  * @return		Return True if this property should be displayed.  False if it should be culled
diff --git a/Engine/Source/Editor/PropertyEditor/Private/PropertyNode.h b/Engine/Source/Editor/PropertyEditor/Private/PropertyNode.h
--- a/Engine/Source/Editor/PropertyEditor/Private/PropertyNode.h
+++ b/Engine/Source/Editor/PropertyEditor/Private/PropertyNode.h
@@ -705,6 +705,10 @@ public:
 	 */
 	const FString* GetInstanceMetaData(const FName& Key) const;
 
+	/**
+	 * Invalidates the cached state of this node in all children;
+	 */
+	void InvalidateCachedState();
 protected:
 
 	TSharedRef<FEditPropertyChain> BuildPropertyChain( UProperty* PropertyAboutToChange );
