Copying //UE4/Dev-Editor to //UE4/Dev-Main (Source: //UE4/Dev-Editor @ 3082391)

diff --git a/Engine/Source/Editor/PropertyEditor/Private/SDetailsView.cpp b/Engine/Source/Editor/PropertyEditor/Private/SDetailsView.cpp
--- a/Engine/Source/Editor/PropertyEditor/Private/SDetailsView.cpp
+++ b/Engine/Source/Editor/PropertyEditor/Private/SDetailsView.cpp
@@ -25,10 +25,10 @@
 
 SDetailsView::~SDetailsView()
 {
-	auto RootNode = GetRootNode();
-	if( RootNode.IsValid() )
+	const FRootPropertyNodeList& RootNodes = GetRootNodes();
+	for(const TSharedPtr<FComplexPropertyNode>& RootNode : RootNodes)
 	{
-		SaveExpandedItems( RootNode.ToSharedRef() );
+		SaveExpandedItems(RootNode.ToSharedRef());
 	}
 };
 
@@ -42,9 +42,6 @@ void SDetailsView::Construct(const FArguments& InArgs)
 	DetailsViewArgs = InArgs._DetailsViewArgs;
 	bViewingClassDefaultObject = false;
 
-	// Create the root property now
-	RootPropertyNode = MakeShareable( new FObjectPropertyNode );
-		
 	PropertyUtilities = MakeShareable( new FPropertyDetailsUtilities( *this ) );
 	
 	ColumnSizeData.LeftColumnWidth = TAttribute<float>( this, &SDetailsView::OnGetLeftColumnWidth );
@@ -58,7 +55,7 @@ void SDetailsView::Construct(const FArguments& InArgs)
 	TSharedRef<SScrollBar> ExternalScrollbar = SNew(SScrollBar);
 	ExternalScrollbar->SetVisibility( TAttribute<EVisibility>( this, &SDetailsView::GetScrollBarVisibility ) );
 
-		FMenuBuilder DetailViewOptions( true, NULL );
+		FMenuBuilder DetailViewOptions( true, nullptr);
 
 		if (DetailsViewArgs.bShowModifiedPropertiesOption)
 		{
@@ -270,7 +267,7 @@ EVisibility SDetailsView::GetActorNameAreaVisibility() const
 
 EVisibility SDetailsView::GetScrollBarVisibility() const
 {
-	const bool bHasObjects = RootPropertyNode.IsValid() && RootPropertyNode->GetObjectBaseClass() && RootPropertyNode->GetNumObjects();
+	const bool bHasObjects = RootPropertyNodes.Num() > 0;
 	return bHasObjects ? EVisibility::Visible : EVisibility::Collapsed; 
 }
 
@@ -278,13 +275,21 @@ void SDetailsView::ForceRefresh()
 {
 	TArray< TWeakObjectPtr< UObject > > NewObjectList;
 
-	// Simply re-add the same existing objects to cause a refresh
-	for ( TPropObjectIterator Itor( RootPropertyNode->ObjectIterator() ); Itor; ++Itor )
+	const FRootPropertyNodeList& RootNodes = GetRootNodes();
+	for(const TSharedPtr<FComplexPropertyNode>& ComplexRootNode : RootNodes)
 	{
-		TWeakObjectPtr<UObject> Object = *Itor;
-		if( Object.IsValid() )
+		FObjectPropertyNode* RootNode = ComplexRootNode->AsObjectNode();
+		if(RootNode)
 		{
-			NewObjectList.Add( Object.Get() );
+			// Simply re-add the same existing objects to cause a refresh
+			for(TPropObjectIterator Itor(RootNode->ObjectIterator()); Itor; ++Itor)
+			{
+				TWeakObjectPtr<UObject> Object = *Itor;
+				if(Object.IsValid())
+				{
+					NewObjectList.Add(Object.Get());
+				}
+			}
 		}
 	}
 
@@ -296,18 +301,6 @@ void SDetailsView::MoveScrollOffset(int32 DeltaOffset)
 	DetailTree->AddScrollOffset((float)DeltaOffset);
 }
 
-bool SDetailsView::GetCategoryInfo(FName CategoryName, int32 &SimplePropertiesNum, int32 &AdvancePropertiesNum)
-{
-	SimplePropertiesNum = 0;
-	AdvancePropertiesNum = 0;
-	bool CategoryExist = DetailLayout.IsValid() ? DetailLayout->HasCategory(CategoryName) : false;
-	if (CategoryExist)
-	{
-		DetailLayout->DefaultCategory(CategoryName).GetCategoryInformation(SimplePropertiesNum, AdvancePropertiesNum);
-	}
-	return CategoryExist;
-}
-
 void SDetailsView::SetObjects( const TArray<UObject*>& InObjects, bool bForceRefresh/* = false*/, bool bOverrideLock/* = false*/ )
 {
 	if (!IsLocked() || bOverrideLock)
@@ -350,17 +343,25 @@ void SDetailsView::RemoveInvalidObjects()
 	TArray< TWeakObjectPtr< UObject > > ResetArray;
 
 	bool bAllFound = true;
-	for (TPropObjectIterator Itor(RootPropertyNode->ObjectIterator()); Itor; ++Itor)
+	
+	for(const TSharedPtr<FComplexPropertyNode>& ComplexRootNode : RootPropertyNodes)
 	{
-		TWeakObjectPtr<UObject> Object = *Itor;
-
-		if( Object.IsValid() && !Object->IsPendingKill() )
-		{
-			ResetArray.Add(Object);
-		}
-		else
+		FObjectPropertyNode* RootPropertyNode = ComplexRootNode->AsObjectNode();
+		if(RootPropertyNode)
 		{
-			bAllFound = false;
+			for(TPropObjectIterator Itor(RootPropertyNode->ObjectIterator()); Itor; ++Itor)
+			{
+				TWeakObjectPtr<UObject> Object = *Itor;
+
+				if(Object.IsValid() && !Object->IsPendingKill())
+				{
+					ResetArray.Add(Object);
+				}
+				else
+				{
+					bAllFound = false;
+				}
+			}
 		}
 	}
 
@@ -372,10 +373,30 @@ void SDetailsView::RemoveInvalidObjects()
 
 void SDetailsView::SetObjectPackageOverrides(const TMap<TWeakObjectPtr<UObject>, TWeakObjectPtr<UPackage>>& InMapping)
 {
-	RootPropertyNode->SetObjectPackageOverrides(InMapping);
+	for(TSharedPtr<FComplexPropertyNode>& ComplexRootNode : RootPropertyNodes)
+	{
+		FObjectPropertyNode* RootNode = ComplexRootNode->AsObjectNode();
+		if(RootNode)
+		{
+			RootNode->SetObjectPackageOverrides(InMapping);
+		}
+	}
+}
+
+void SDetailsView::SetRootObjectCustomizationInstance(TSharedPtr<IDetailRootObjectCustomization> InRootObjectCustomization)
+{
+	RootObjectCustomization = InRootObjectCustomization;
+	RerunCurrentFilter();
 }
 
-bool SDetailsView::ShouldSetNewObjects( const TArray< TWeakObjectPtr< UObject > >& InObjects ) const
+void SDetailsView::ClearSearch()
+{
+	CurrentFilter.FilterStrings.Empty();
+	SearchBox->SetText(FText::GetEmpty());
+	RerunCurrentFilter();
+}
+
+bool SDetailsView::ShouldSetNewObjects(const TArray< TWeakObjectPtr< UObject > >& InObjects) const
 {
 	bool bShouldSetObjects = false;
 
@@ -385,30 +406,73 @@ bool SDetailsView::ShouldSetNewObjects( const TArray< TWeakObjectPtr< UObject >
 		// If a BSP brush was selected we need to refresh because surface could have been selected and the object set not updated
 		bShouldSetObjects = true;
 	}
-	else if( InObjects.Num() != RootPropertyNode->GetNumObjects() )
+	else if( InObjects.Num() != GetNumObjects() )
+	{
+		// If the object arrays differ in size then at least one object is different so we must reset
+		bShouldSetObjects = true;
+	}
+	else if(InObjects.Num() == 0)
 	{
-		// If the object arrys differ in size then at least one object is different so we must reset
+		// User is likely resetting details panel
 		bShouldSetObjects = true;
 	}
 	else
 	{
 		// Check to see if the objects passed in are different. If not we do not need to set anything
 		TSet< TWeakObjectPtr< UObject > > NewObjects;
-		NewObjects.Append( InObjects );
-		for ( TPropObjectIterator Itor( RootPropertyNode->ObjectIterator() ); Itor; ++Itor )
+		NewObjects.Append(InObjects);
+
+		if(DetailsViewArgs.bAllowMultipleTopLevelObjects)
 		{
-			TWeakObjectPtr<UObject> Object = *Itor;
-			if( Object.IsValid() && !NewObjects.Contains( Object ) )
+			
+			// For multiple top level node support, if the single object in each node is not found in the new object set
+			// then we need to refresh
+			for(int32 RootNodeIndex = 0; RootNodeIndex < RootPropertyNodes.Num(); ++RootNodeIndex)
 			{
-				// An existing object is not in the list of new objects to set
-				bShouldSetObjects = true;
-				break;
+				FObjectPropertyNode* RootPropertyNode = RootPropertyNodes[RootNodeIndex]->AsObjectNode();
+				
+				if(RootPropertyNode && RootPropertyNode->GetNumObjects() > 0)
+				{
+					if(!NewObjects.Contains(RootPropertyNode->GetUObject(0)))
+					{
+						bShouldSetObjects = true;
+						break;
+					}
+				}
+				else
+				{
+					bShouldSetObjects = true;
+					break;
+				}
+			}
+		}
+		else
+		{
+
+			ensure(RootPropertyNodes.Num() == 1);
+			FObjectPropertyNode* RootPropertyNode = RootPropertyNodes[0]->AsObjectNode();
+			if( RootPropertyNode )
+			{
+				for(TPropObjectIterator Itor(RootPropertyNode->ObjectIterator()); Itor; ++Itor)
+				{
+					TWeakObjectPtr<UObject> Object = *Itor;
+					if(Object.IsValid() && !NewObjects.Contains(Object))
+					{
+						// An existing object is not in the list of new objects to set
+						bShouldSetObjects = true;
+						break;
+					}
+					else if(!Object.IsValid())
+					{
+						// An existing object is invalid
+						bShouldSetObjects = true;
+						break;
+					}
+				}
 			}
-			else if( !Object.IsValid() )
+			else
 			{
-				// An existing object is invalid
 				bShouldSetObjects = true;
-				break;
 			}
 		}
 	}
@@ -421,15 +485,25 @@ bool SDetailsView::ShouldSetNewObjects( const TArray< TWeakObjectPtr< UObject >
 	return bShouldSetObjects;
 }
 
-void SDetailsView::SetObjectArrayPrivate( const TArray< TWeakObjectPtr< UObject > >& InObjects )
+int32 SDetailsView::GetNumObjects() const
 {
-	double StartTime = FPlatformTime::Seconds();
+	if(DetailsViewArgs.bAllowMultipleTopLevelObjects)
+	{
+		return RootPropertyNodes.Num();
+	}
+	else if( RootPropertyNodes.Num() > 0 && RootPropertyNodes[0]->AsObjectNode())
+	{
+		return RootPropertyNodes[0]->AsObjectNode()->GetNumObjects();
+	}
 
-	PreSetObject();
+	return 0;
+}
 
-	check( RootPropertyNode.IsValid() );
+void SDetailsView::SetObjectArrayPrivate(const TArray< TWeakObjectPtr< UObject > >& InObjects)
+{
+	double StartTime = FPlatformTime::Seconds();
 
-	RootPropertyNode->ClearObjectPackageOverrides();
+	PreSetObject(InObjects.Num());
 
 	// Selected actors for building SelectedActorInfo
 	TArray<AActor*> SelectedRawActors;
@@ -444,7 +518,16 @@ void SDetailsView::SetObjectArrayPrivate( const TArray< TWeakObjectPtr< UObject
 		{
 			bViewingClassDefaultObject &= Object->HasAnyFlags( RF_ClassDefaultObject );
 
-			RootPropertyNode->AddObject( Object.Get() );
+			if(DetailsViewArgs.bAllowMultipleTopLevelObjects)
+			{
+				check(RootPropertyNodes.Num() == InObjects.Num());
+				RootPropertyNodes[ObjectIndex]->AsObjectNode()->AddObject( Object.Get() );
+			}
+			else
+			{
+				RootPropertyNodes[0]->AsObjectNode()->AddObject( Object.Get() );
+			}
+
 			SelectedObjects.Add( Object );
 			AActor* Actor = Cast<AActor>( Object.Get() );
 			if( Actor )
@@ -484,14 +567,14 @@ void SDetailsView::SetObjectArrayPrivate( const TArray< TWeakObjectPtr< UObject
 	// Or call the delegate for handling when the title changed
 	FString Title;
 
-	if( !RootPropertyNode->GetObjectBaseClass() )
+	if( GetNumObjects() == 0 )
 	{
 		Title = NSLOCTEXT("PropertyView", "NothingSelectedTitle", "Nothing selected").ToString();
 	}
-	else if( RootPropertyNode->GetNumObjects() == 1 )
+	else if( GetNumObjects() == 1 && RootPropertyNodes[0]->AsObjectNode()->GetNumObjects() > 0)
 	{
 		// if the object is the default metaobject for a UClass, use the UClass's name instead
-		const UObject* Object = RootPropertyNode->ObjectConstIterator()->Get();
+		UObject* Object = RootPropertyNodes[0]->AsObjectNode()->GetUObject(0);
 		FString ObjectName = Object->GetName();
 		if ( Object->GetClass()->GetDefaultObject() == Object )
 		{
@@ -501,7 +584,7 @@ void SDetailsView::SetObjectArrayPrivate( const TArray< TWeakObjectPtr< UObject
 		{
 			// Is this an actor?  If so, it might have a friendly name to display
 			const AActor* Actor = Cast<const  AActor >( Object );
-			if( Actor != NULL )
+			if( Actor != nullptr)
 			{
 				// Use the friendly label for this actor
 				ObjectName = Actor->GetActorLabel();
@@ -510,8 +593,13 @@ void SDetailsView::SetObjectArrayPrivate( const TArray< TWeakObjectPtr< UObject
 
 		Title = ObjectName;
 	}
+	else if(DetailsViewArgs.bAllowMultipleTopLevelObjects)
+	{
+		Title = FString::Printf(*NSLOCTEXT("PropertyView", "MultipleToLevelObjectsSelected", "%i selected").ToString(), GetNumObjects());
+	}
 	else
 	{
+		FObjectPropertyNode* RootPropertyNode = RootPropertyNodes[0]->AsObjectNode();
 		Title = FString::Printf( *NSLOCTEXT("PropertyView", "MultipleSelected", "%s (%i selected)").ToString(), *RootPropertyNode->GetObjectBaseClass()->GetName(), RootPropertyNode->GetNumObjects() );
 	}
 
@@ -526,7 +614,10 @@ void SDetailsView::ReplaceObjects( const TMap<UObject*, UObject*>& OldToNewObjec
 	bool bObjectsReplaced = false;
 
 	TArray< FObjectPropertyNode* > ObjectNodes;
-	PropertyEditorHelpers::CollectObjectNodes( RootPropertyNode, ObjectNodes );
+	for(TSharedPtr<FComplexPropertyNode>& RootNode : RootPropertyNodes)
+	{
+		PropertyEditorHelpers::CollectObjectNodes(RootNode, ObjectNodes );
+	}
 
 	for( int32 ObjectNodeIndex = 0; ObjectNodeIndex < ObjectNodes.Num(); ++ObjectNodeIndex )
 	{
@@ -539,13 +630,13 @@ void SDetailsView::ReplaceObjects( const TMap<UObject*, UObject*>& OldToNewObjec
 			if( Replacement && Replacement->GetClass() == Itor->Get()->GetClass() )
 			{
 				bObjectsReplaced = true;
-				if( CurrentNode == RootPropertyNode.Get() )
+				if( CurrentNode->IsRootNode() )
 				{
 					// Note: only root objects count for the new object list. Sub-Objects (i.e components count as needing to be replaced but they don't belong in the top level object list
 					NewObjectList.Add( Replacement );
 				}
 			}
-			else if( CurrentNode == RootPropertyNode.Get() )
+			else if( CurrentNode->IsRootNode() )
 			{
 				// Note: only root objects count for the new object list. Sub-Objects (i.e components count as needing to be replaced but they don't belong in the top level object list
 				NewObjectList.Add( Itor->Get() );
@@ -566,18 +657,22 @@ void SDetailsView::RemoveDeletedObjects( const TArray<UObject*>& DeletedObjects
 	TArray< TWeakObjectPtr< UObject > > NewObjectList;
 	bool bObjectsRemoved = false;
 
-	// Scan all objects and look for objects which need to be replaced
-	for ( TPropObjectIterator Itor( RootPropertyNode->ObjectIterator() ); Itor; ++Itor )
+	for(const TSharedPtr<FComplexPropertyNode>& ComplexRootNode : RootPropertyNodes)
 	{
-		if( DeletedObjects.Contains( Itor->Get() ) )
-		{
-			// An object we had needs to be removed
-			bObjectsRemoved = true;
-		}
-		else
+		FObjectPropertyNode* RootPropertyNode = ComplexRootNode->AsObjectNode();
+		// Scan all objects and look for objects which need to be replaced
+		for ( TPropObjectIterator Itor( RootPropertyNode->ObjectIterator() ); Itor; ++Itor )
 		{
-			// If the deleted object list does not contain the current object, its ok to keep it in the list
-			NewObjectList.Add( Itor->Get() );
+			if( DeletedObjects.Contains( Itor->Get() ) )
+			{
+				// An object we had needs to be removed
+				bObjectsRemoved = true;
+			}
+			else
+			{
+				// If the deleted object list does not contain the current object, its ok to keep it in the list
+				NewObjectList.Add( Itor->Get() );
+			}
 		}
 	}
 
@@ -589,12 +684,18 @@ void SDetailsView::RemoveDeletedObjects( const TArray<UObject*>& DeletedObjects
 }
 
 /** Called before during SetObjectArray before we change the objects being observed */
-void SDetailsView::PreSetObject()
+void SDetailsView::PreSetObject(int32 InNewNumObjects)
 {
 	// Save existing expanded items first
-	if( GetRootNode().IsValid() )
+	for(TSharedPtr<FComplexPropertyNode>& RootNode : RootPropertyNodes)
 	{
-		SaveExpandedItems( GetRootNode().ToSharedRef() );
+		SaveExpandedItems(RootNode.ToSharedRef());
+
+		RootNodesPendingKill.Add(RootNode);
+		FObjectPropertyNode* RootObjectNode = RootNode->AsObjectNode();
+		RootObjectNode->RemoveAllObjects();
+		RootObjectNode->ClearCachedReadAddresses(true);
+		RootObjectNode->ClearObjectPackageOverrides();
 	}
 
 	for( auto ExternalRootNode : ExternalRootPropertyNodes )
@@ -612,15 +713,20 @@ void SDetailsView::PreSetObject()
 	}
 
 	ExternalRootPropertyNodes.Empty();
+	RootPropertyNodes.Empty(InNewNumObjects);
 
-	if(RootPropertyNode.IsValid())
+	if(DetailsViewArgs.bAllowMultipleTopLevelObjects)
 	{
-		RootNodesPendingKill.Add(RootPropertyNode);
-		RootPropertyNode->RemoveAllObjects();
-		RootPropertyNode->ClearCachedReadAddresses(true);
+		for(int32 NewRootIndex = 0; NewRootIndex < InNewNumObjects; ++NewRootIndex)
+		{
+			RootPropertyNodes.Add(MakeShareable(new FObjectPropertyNode));
+		}
 	}
+	else
+	{
+		RootPropertyNodes.Add(MakeShareable(new FObjectPropertyNode));
+	}	
 
-	RootPropertyNode = MakeShareable(new FObjectPropertyNode);
 
 	SelectedActors.Empty();
 	SelectedObjects.Empty();
@@ -631,11 +737,11 @@ void SDetailsView::PreSetObject()
 void SDetailsView::PostSetObject()
 {
 	DestroyColorPicker();
-	ColorPropertyNode = NULL;
+	ColorPropertyNode = nullptr;
 
 	FPropertyNodeInitParams InitParams;
-	InitParams.ParentNode = NULL;
-	InitParams.Property = NULL;
+	InitParams.ParentNode = nullptr;
+	InitParams.Property = nullptr;
 	InitParams.ArrayOffset = 0;
 	InitParams.ArrayIndex = INDEX_NONE;
 	InitParams.bAllowChildren = true;
@@ -656,18 +762,17 @@ void SDetailsView::PostSetObject()
 		check(false);
 	}
 
-	RootPropertyNode->InitNode( InitParams );
+	for( TSharedPtr<FComplexPropertyNode>& ComplexRootNode : RootPropertyNodes )
+	{
+		FObjectPropertyNode* RootPropertyNode = ComplexRootNode->AsObjectNode();
 
-	bool bInitiallySeen = true;
-	bool bParentAllowsVisible = true;
-	// Restore existing expanded items
+		RootPropertyNode->InitNode( InitParams );
 
-	if( GetRootNode().IsValid() )
-	{
-		RestoreExpandedItems( GetRootNode().ToSharedRef() );
+		// Restore existing expanded items
+		RestoreExpandedItems(ComplexRootNode.ToSharedRef());
 	}
 
-	UpdatePropertyMap();
+	UpdatePropertyMaps();
 
 	for( auto ExternalRootNode : ExternalRootPropertyNodes )
 	{
@@ -685,30 +790,35 @@ void SDetailsView::SetOnObjectArrayChanged(FOnObjectArrayChanged OnObjectArrayCh
 	OnObjectArrayChanged = OnObjectArrayChangedDelegate;
 }
 
+PRAGMA_DISABLE_DEPRECATION_WARNINGS
 const UClass* SDetailsView::GetBaseClass() const
 {
-	if( RootPropertyNode.IsValid() )
+	if( RootPropertyNodes.Num() == 1 )
+	{
+		return RootPropertyNodes[0]->AsObjectNode()->GetObjectBaseClass();
+	}
+	else
 	{
-		return RootPropertyNode->GetObjectBaseClass();
+		ensureMsgf(0,TEXT("GetBaseClass is not supported on multi-object details panels"));
 	}
 
-	return NULL;
+	return nullptr;
 }
 
 UClass* SDetailsView::GetBaseClass()
 {
-	if( RootPropertyNode.IsValid() )
+	if(RootPropertyNodes.Num() == 1)
 	{
-		return RootPropertyNode->GetObjectBaseClass();
+		return RootPropertyNodes[0]->AsObjectNode()->GetObjectBaseClass();
+	}
+	else
+	{
+		ensureMsgf(0, TEXT("GetBaseClass is not supported on multi-object details panels"));
 	}
 
-	return NULL;
-}
-
-UStruct* SDetailsView::GetBaseStruct() const
-{
-	return RootPropertyNode.IsValid() ? RootPropertyNode->GetObjectBaseClass() : NULL;
+	return nullptr;
 }
+PRAGMA_ENABLE_DEPRECATION_WARNINGS
 
 void SDetailsView::RegisterInstancedCustomPropertyLayout( UStruct* Class, FOnGetDetailCustomizationInstance DetailLayoutDelegate )
 {
@@ -725,14 +835,14 @@ void SDetailsView::AddExternalRootPropertyNode( TSharedRef<FPropertyNode> Extern
 	ExternalRootPropertyNodes.Add( ExternalRootNode );
 }
 
-bool SDetailsView::IsCategoryHiddenByClass( FName CategoryName ) const
+bool SDetailsView::IsCategoryHiddenByClass( const TSharedPtr<FComplexPropertyNode>& InRootNode, FName CategoryName ) const
 {
-	return RootPropertyNode->GetHiddenCategories().Contains( CategoryName );
+	return InRootNode->AsObjectNode()->GetHiddenCategories().Contains( CategoryName );
 }
 
 bool SDetailsView::IsConnected() const
 {
-	return RootPropertyNode.IsValid() && (RootPropertyNode->GetNumObjects() > 0);
+	return RootPropertyNodes.Num() > 0;
 }
 
 const FSlateBrush* SDetailsView::OnGetLockButtonImageResource() const
diff --git a/Engine/Source/Editor/PropertyEditor/Private/SDetailsView.h b/Engine/Source/Editor/PropertyEditor/Private/SDetailsView.h
--- a/Engine/Source/Editor/PropertyEditor/Private/SDetailsView.h
+++ b/Engine/Source/Editor/PropertyEditor/Private/SDetailsView.h
@@ -4,6 +4,7 @@
 
 #include "SDetailsViewBase.h"
 
+PRAGMA_DISABLE_DEPRECATION_WARNINGS
 class SDetailsView : public SDetailsViewBase
 {
 	friend class FPropertyDetailsUtilities;
@@ -24,8 +25,6 @@ public:
 	/** Move the scrolling offset (by item), but do not refresh the tree*/
 	void MoveScrollOffset(int32 DeltaOffset) override;
 
-	/* Return true if the favorite category is currently displayed */
-	bool GetCategoryInfo(FName CategoryName, int32 &SimplePropertiesNum, int32 &AdvancePropertiesNum) override;
 
 	/**
 	 * Constructs the property view widgets                   
@@ -38,6 +37,8 @@ public:
 	virtual void SetObject( UObject* InObject, bool bForceRefresh = false ) override;
 	virtual void RemoveInvalidObjects() override;
 	virtual void SetObjectPackageOverrides(const TMap<TWeakObjectPtr<UObject>, TWeakObjectPtr<UPackage>>& InMapping) override;
+	virtual void SetRootObjectCustomizationInstance(TSharedPtr<IDetailRootObjectCustomization> InRootObjectCustomization) override;
+	virtual void ClearSearch() override;
 
 	/**
 	 * Replaces objects being observed by the view with new objects
@@ -80,9 +81,9 @@ public:
 	}
 
 	/** Gets the base class being viewed */
+	/** These methods is deprecated.  When it is removed, also remove the PRAGMA_DISABLE_DEPRECATION_WARNINGS/PRAGMA_ENABLE_DEPRECATION_WARNINGS at the top and bottom of this file */
 	const UClass* GetBaseClass() const override;
 	UClass* GetBaseClass() override;
-	UStruct* GetBaseStruct() const override;
 
 	/**
 	 * Adds an external property root node to the list of root nodes that the details new needs to manage
@@ -92,21 +93,34 @@ public:
 	void AddExternalRootPropertyNode( TSharedRef<FPropertyNode> ExternalRootNode ) override;
 	
 	/**
-	 * @return True if a category is hidden by any of the uobject classes currently in view by this details panel
+	 * Whether or not a category is hidden by a given root object
+	 * @param InRootNode	The root node that for the objects we are customizing
+	 * @param CategoryName	The name of the category to check
+	 * @return true if a category is hidden, false otherwise
 	 */
-	bool IsCategoryHiddenByClass(FName CategoryName) const override;
+	bool IsCategoryHiddenByClass(const TSharedPtr<FComplexPropertyNode>& InRootNode, FName CategoryName) const override;
 
 	virtual bool IsConnected() const override;
 
-	virtual TSharedPtr<FComplexPropertyNode> GetRootNode() override
+	virtual FRootPropertyNodeList& GetRootNodes() override
 	{
-		return RootPropertyNode;
+		return RootPropertyNodes;
 	}
 
 	virtual bool DontUpdateValueWhileEditing() const override
 	{ 
 		return false; 
 	}
+
+	virtual bool ContainsMultipleTopLevelObjects() const override
+	{
+		return DetailsViewArgs.bAllowMultipleTopLevelObjects && GetNumObjects() > 1;
+	}
+
+	virtual TSharedPtr<IDetailRootObjectCustomization> GetRootObjectCustomization() const override
+	{
+		return RootObjectCustomization;
+	}
 private:
 	void RegisterInstancedCustomPropertyLayout( UStruct* Class, FOnGetDetailCustomizationInstance DetailLayoutDelegate ) override;
 	void UnregisterInstancedCustomPropertyLayout( UStruct* Class ) override;
@@ -123,8 +137,13 @@ private:
 	 */
 	bool ShouldSetNewObjects( const TArray< TWeakObjectPtr< UObject > >& InObjects ) const;
 
+	/**
+	 * Returns the number of objects being edited by this details panel.
+	 */
+	int32 GetNumObjects() const;
+
 	/** Called before during SetObjectArray before we change the objects being observed */
-	void PreSetObject();
+	void PreSetObject(int32 InNewNumObjects);
 
 	/** Called at the end of SetObjectArray after we change the objects being observed */
 	void PostSetObject();
@@ -156,10 +175,13 @@ private:
 	 * It will only contain actors from when it was locked 
 	 */
 	TArray< TWeakObjectPtr<AActor> > SelectedActors;
-	/** The root property node of the property tree for a specific set of UObjects */
-	TSharedPtr<FObjectPropertyNode> RootPropertyNode;
+	/** The root property nodes of the property tree for a specific set of UObjects */
+	TArray<TSharedPtr<FComplexPropertyNode>> RootPropertyNodes;
 	/** Callback to send when the property view changes */
 	FOnObjectArrayChanged OnObjectArrayChanged;
+	/** Customization instance used when there are multiple top level objects in this view */
+	TSharedPtr<IDetailRootObjectCustomization> RootObjectCustomization;
 	/** True if at least one viewed object is a CDO (blueprint editing) */
 	bool bViewingClassDefaultObject;
 };
+PRAGMA_ENABLE_DEPRECATION_WARNINGS
diff --git a/Engine/Source/Editor/PropertyEditor/Private/SDetailsViewBase.cpp b/Engine/Source/Editor/PropertyEditor/Private/SDetailsViewBase.cpp
--- a/Engine/Source/Editor/PropertyEditor/Private/SDetailsViewBase.cpp
+++ b/Engine/Source/Editor/PropertyEditor/Private/SDetailsViewBase.cpp
@@ -18,6 +18,7 @@
 #include "ObjectEditorUtils.h"
 #include "SColorPicker.h"
 #include "DetailPropertyRow.h"
+#include "SSearchBox.h"
 
 
 SDetailsViewBase::~SDetailsViewBase()
@@ -186,7 +187,15 @@
 
 EVisibility SDetailsViewBase::GetTreeVisibility() const
 {
-	return DetailLayout.IsValid() && DetailLayout->HasDetails() ? EVisibility::Visible : EVisibility::Collapsed;
+	for(const FDetailLayoutData& Data : DetailLayouts)
+	{
+		if(Data.DetailLayout.IsValid() && Data.DetailLayout->HasDetails())
+		{
+			return EVisibility::Visible;
+		}
+	}
+
+	return EVisibility::Collapsed;
 }
 
 /** Returns the image used for the icon on the filter button */
@@ -279,6 +288,328 @@
 	}
 }
 
+/**
+* Determines whether or not a property should be visible in the default generated detail layout
+*
+* @param PropertyNode	The property node to check
+* @param ParentNode	The parent property node to check
+* @return true if the property should be visible
+*/
+static bool IsVisibleStandaloneProperty(const FPropertyNode& PropertyNode, const FPropertyNode& ParentNode)
+{
+	const UProperty* Property = PropertyNode.GetProperty();
+	const UArrayProperty* ParentArrayProperty = Cast<const UArrayProperty>(ParentNode.GetProperty());
+	bool bIsVisibleStandalone = false;
+	if(Property)
+	{
+		if(Property->IsA(UObjectPropertyBase::StaticClass()))
+		{
+			// Do not add this child node to the current map if its a single object property in a category (serves no purpose for UI)
+			bIsVisibleStandalone = !ParentArrayProperty && (PropertyNode.GetNumChildNodes() == 0 || PropertyNode.GetNumChildNodes() > 1);
+		}
+		else if(Property->IsA(UArrayProperty::StaticClass()) || (Property->ArrayDim > 1 && PropertyNode.GetArrayIndex() == INDEX_NONE))
+		{
+			// Base array properties are always visible
+			bIsVisibleStandalone = true;
+		}
+		else
+		{
+			bIsVisibleStandalone = true;
+		}
+
+	}
+
+	return bIsVisibleStandalone;
+}
+
+void SDetailsViewBase::UpdatePropertyMaps()
+{
+	RootTreeNodes.Empty();
+
+	for(FDetailLayoutData& LayoutData : DetailLayouts)
+	{
+		// Check uniqueness.  It is critical that detail layouts can be destroyed
+		// We need to be able to create a new detail layout and properly clean up the old one in the process
+		check(!LayoutData.DetailLayout.IsValid() || LayoutData.DetailLayout.IsUnique());
+
+		// All the current customization instances need to be deleted when it is safe
+		CustomizationClassInstancesPendingDelete.Append(LayoutData.CustomizationClassInstances);
+	}
+
+	FRootPropertyNodeList& RootPropertyNodes = GetRootNodes();
+	
+	DetailLayouts.Empty(RootPropertyNodes.Num());
+
+	// There should be one detail layout for each root node
+	DetailLayouts.AddDefaulted(RootPropertyNodes.Num());
+
+	for(int32 RootNodeIndex = 0; RootNodeIndex < RootPropertyNodes.Num(); ++RootNodeIndex)
+	{
+		FDetailLayoutData& LayoutData = DetailLayouts[RootNodeIndex];
+		UpdateSinglePropertyMap(RootPropertyNodes[RootNodeIndex], LayoutData);
+	}
+}
+
+void SDetailsViewBase::UpdateSinglePropertyMap(TSharedPtr<FComplexPropertyNode>& InRootPropertyNode, FDetailLayoutData& LayoutData)
+{
+	// Reset everything
+	LayoutData.ClassToPropertyMap.Empty();
+
+	TSharedPtr<FDetailLayoutBuilderImpl> DetailLayout = MakeShareable(new FDetailLayoutBuilderImpl(InRootPropertyNode, LayoutData.ClassToPropertyMap, PropertyUtilities.ToSharedRef(), SharedThis(this)));
+	LayoutData.DetailLayout = DetailLayout;
+
+	TSharedPtr<FComplexPropertyNode> RootPropertyNode = InRootPropertyNode;
+	check(RootPropertyNode.IsValid());
+
+	bool const bEnableFavoriteSystem = GIsRequestingExit ? false : (GetDefault<UEditorExperimentalSettings>()->bEnableFavoriteSystem && DetailsViewArgs.bAllowFavoriteSystem);
+
+	// Currently object property nodes do not provide any useful information other than being a container for its children.  We do not draw anything for them.
+	// When we encounter object property nodes, add their children instead of adding them to the tree.
+	UpdateSinglePropertyMapRecursive(*RootPropertyNode, LayoutData, NAME_None, RootPropertyNode.Get(), bEnableFavoriteSystem, false);
+
+	CustomUpdatePropertyMap(LayoutData.DetailLayout);
+
+	// Ask for custom detail layouts, unless disabled. One reason for disabling custom layouts is that the custom layouts
+	// inhibit our ability to find a single property's tree node. This is problematic for the diff and merge tools, that need
+	// to display and highlight each changed property for the user. We could whitelist 'good' customizations here if 
+	// we can make them work with the diff/merge tools.
+	if(!bDisableCustomDetailLayouts)
+	{
+		QueryCustomDetailLayout(LayoutData);
+	}
+
+	LayoutData.DetailLayout->GenerateDetailLayout();
+}
+
+void SDetailsViewBase::UpdateSinglePropertyMapRecursive(FPropertyNode& InNode, FDetailLayoutData& LayoutData, FName CurCategory, FComplexPropertyNode* CurObjectNode, bool bEnableFavoriteSystem, bool bUpdateFavoriteSystemOnly)
+{
+	FDetailLayoutBuilderImpl& DetailLayout = *LayoutData.DetailLayout;
+
+	UProperty* ParentProperty = InNode.GetProperty();
+	UStructProperty* ParentStructProp = Cast<UStructProperty>(ParentProperty);
+	for(int32 ChildIndex = 0; ChildIndex < InNode.GetNumChildNodes(); ++ChildIndex)
+	{
+		//Use the original value for each child
+		bool LocalUpdateFavoriteSystemOnly = bUpdateFavoriteSystemOnly;
+
+		TSharedPtr<FPropertyNode> ChildNodePtr = InNode.GetChildNode(ChildIndex);
+		FPropertyNode& ChildNode = *ChildNodePtr;
+		UProperty* Property = ChildNode.GetProperty();
+
+		{
+			FObjectPropertyNode* ObjNode = ChildNode.AsObjectNode();
+			FCategoryPropertyNode* CategoryNode = ChildNode.AsCategoryNode();
+			if(ObjNode)
+			{
+				// Currently object property nodes do not provide any useful information other than being a container for its children.  We do not draw anything for them.
+				// When we encounter object property nodes, add their children instead of adding them to the tree.
+				UpdateSinglePropertyMapRecursive(ChildNode, LayoutData, CurCategory, ObjNode, bEnableFavoriteSystem, LocalUpdateFavoriteSystemOnly);
+			}
+			else if(CategoryNode)
+			{
+				if(!LocalUpdateFavoriteSystemOnly)
+				{
+					FName InstanceName = NAME_None;
+					FName CategoryName = CurCategory;
+					FString CategoryDelimiterString;
+					CategoryDelimiterString.AppendChar(FPropertyNodeConstants::CategoryDelimiterChar);
+					if(CurCategory != NAME_None && CategoryNode->GetCategoryName().ToString().Contains(CategoryDelimiterString))
+					{
+						// This property is child of another property so add it to the parent detail category
+						FDetailCategoryImpl& CategoryImpl = DetailLayout.DefaultCategory(CategoryName);
+						CategoryImpl.AddPropertyNode(ChildNodePtr.ToSharedRef(), InstanceName);
+					}
+				}
+
+				// For category nodes, we just set the current category and recurse through the children
+				UpdateSinglePropertyMapRecursive(ChildNode, LayoutData, CategoryNode->GetCategoryName(), CurObjectNode, bEnableFavoriteSystem, LocalUpdateFavoriteSystemOnly);
+			}
+			else
+			{
+				// Whether or not the property can be visible in the default detail layout
+				bool bVisibleByDefault = IsVisibleStandaloneProperty(ChildNode, InNode);
+
+				// Whether or not the property is a struct
+				UStructProperty* StructProperty = Cast<UStructProperty>(Property);
+
+				bool bIsStruct = StructProperty != NULL;
+
+				static FName ShowOnlyInners("ShowOnlyInnerProperties");
+
+				bool bIsChildOfCustomizedStruct = false;
+				bool bIsCustomizedStruct = false;
+
+				const UStruct* Struct = StructProperty ? StructProperty->Struct : NULL;
+				const UStruct* ParentStruct = ParentStructProp ? ParentStructProp->Struct : NULL;
+				if(Struct || ParentStruct)
+				{
+					FPropertyEditorModule& ParentPlugin = FModuleManager::GetModuleChecked<FPropertyEditorModule>("PropertyEditor");
+					if(Struct)
+					{
+						bIsCustomizedStruct = ParentPlugin.IsCustomizedStruct(Struct, SharedThis(this));
+					}
+
+					if(ParentStruct)
+					{
+						bIsChildOfCustomizedStruct = ParentPlugin.IsCustomizedStruct(ParentStruct, SharedThis(this));
+					}
+				}
+
+				// Whether or not to push out struct properties to their own categories or show them inside an expandable struct 
+				bool bPushOutStructProps = bIsStruct && !bIsCustomizedStruct && !ParentStructProp && Property->HasMetaData(ShowOnlyInners);
+
+				// Is the property edit inline new 
+				const bool bIsEditInlineNew = SPropertyEditorEditInline::Supports(&ChildNode, ChildNode.GetArrayIndex());
+
+				// Is this a property of an array
+				bool bIsChildOfArray = PropertyEditorHelpers::IsChildOfArray(ChildNode);
+
+				// Edit inline new properties should be visible by default
+				bVisibleByDefault |= bIsEditInlineNew;
+
+				// Children of arrays are not visible directly,
+				bVisibleByDefault &= !bIsChildOfArray;
+
+				FPropertyAndParent PropertyAndParent(*Property, ParentProperty);
+				const bool bIsUserVisible = IsPropertyVisible(PropertyAndParent);
+
+				// Inners of customized in structs should not be taken into consideration for customizing.  They are not designed to be individually customized when their parent is already customized
+				if(!bIsChildOfCustomizedStruct && !LocalUpdateFavoriteSystemOnly)
+				{
+					// Add any object classes with properties so we can ask them for custom property layouts later
+					LayoutData.ClassesWithProperties.Add(Property->GetOwnerStruct());
+				}
+
+				// If there is no outer object then the class is the object root and there is only one instance
+				FName InstanceName = NAME_None;
+				if(CurObjectNode && CurObjectNode->GetParentNode())
+				{
+					InstanceName = CurObjectNode->GetParentNode()->GetProperty()->GetFName();
+				}
+				else if(ParentStructProp)
+				{
+					InstanceName = ParentStructProp->GetFName();
+				}
+
+				// Do not add children of customized in struct properties or arrays
+				if(!bIsChildOfCustomizedStruct && !bIsChildOfArray && !LocalUpdateFavoriteSystemOnly)
+				{
+					// Get the class property map
+					FClassInstanceToPropertyMap& ClassInstanceMap = LayoutData.ClassToPropertyMap.FindOrAdd(Property->GetOwnerStruct()->GetFName());
+
+					FPropertyNodeMap& PropertyNodeMap = ClassInstanceMap.FindOrAdd(InstanceName);
+
+					if(!PropertyNodeMap.ParentProperty)
+					{
+						PropertyNodeMap.ParentProperty = CurObjectNode;
+					}
+					else
+					{
+						ensure(PropertyNodeMap.ParentProperty == CurObjectNode);
+					}
+
+					checkSlow(!PropertyNodeMap.Contains(Property->GetFName()));
+
+					PropertyNodeMap.Add(Property->GetFName(), ChildNodePtr);
+				}
+				bool bCanDisplayFavorite = false;
+				if(bVisibleByDefault && bIsUserVisible && !bPushOutStructProps)
+				{
+					FName CategoryName = CurCategory;
+					// For properties inside a struct, add them to their own category unless they just take the name of the parent struct.  
+					// In that case push them to the parent category
+					FName PropertyCatagoryName = FObjectEditorUtils::GetCategoryFName(Property);
+					if(!ParentStructProp || (PropertyCatagoryName != ParentStructProp->Struct->GetFName()))
+					{
+						CategoryName = PropertyCatagoryName;
+					}
+
+					if(!LocalUpdateFavoriteSystemOnly)
+					{
+						if(IsPropertyReadOnly(PropertyAndParent))
+						{
+							ChildNode.SetNodeFlags(EPropertyNodeFlags::IsReadOnly, true);
+						}
+
+						// Add a property to the default category
+						FDetailCategoryImpl& CategoryImpl = DetailLayout.DefaultCategory(CategoryName);
+						CategoryImpl.AddPropertyNode(ChildNodePtr.ToSharedRef(), InstanceName);
+					}
+
+					bCanDisplayFavorite = true;
+					if(bEnableFavoriteSystem)
+					{
+						if(bIsCustomizedStruct)
+						{
+							bCanDisplayFavorite = false;
+							//CustomizedStruct child are not categorize since they are under an object but we have to put them in favorite category if the user want to favorite them
+							LocalUpdateFavoriteSystemOnly = true;
+						}
+						else if(ChildNodePtr->IsFavorite())
+						{
+							//Find or create the favorite category, we have to duplicate favorite property row under this category
+							FString CategoryFavoritesName = TEXT("Favorites");
+							FName CatFavName = *CategoryFavoritesName;
+							FDetailCategoryImpl& CategoryFavImpl = DetailLayout.DefaultCategory(CatFavName);
+							CategoryFavImpl.SetSortOrder(0);
+							CategoryFavImpl.SetCategoryAsSpecialFavorite();
+
+							//Add the property to the favorite
+							FObjectPropertyNode *RootObjectParent = ChildNodePtr->FindRootObjectItemParent();
+							FName RootInstanceName = NAME_None;
+							if(RootObjectParent != nullptr)
+							{
+								RootInstanceName = RootObjectParent->GetObjectBaseClass()->GetFName();
+							}
+
+							if(LocalUpdateFavoriteSystemOnly)
+							{
+								if(IsPropertyReadOnly(PropertyAndParent))
+								{
+									ChildNode.SetNodeFlags(EPropertyNodeFlags::IsReadOnly, true);
+								}
+								else
+								{
+									//If the parent has a condition that is not met, make the child as readonly
+									FDetailLayoutCustomization ParentTmpCustomization;
+									ParentTmpCustomization.PropertyRow = MakeShareable(new FDetailPropertyRow(InNode.AsShared(), CategoryFavImpl.AsShared()));
+									if(ParentTmpCustomization.PropertyRow->GetPropertyEditor()->IsPropertyEditingEnabled() == false)
+									{
+										ChildNode.SetNodeFlags(EPropertyNodeFlags::IsReadOnly, true);
+									}
+								}
+							}
+
+							//Duplicate the row
+							CategoryFavImpl.AddPropertyNode(ChildNodePtr.ToSharedRef(), RootInstanceName);
+						}
+
+						if(bIsStruct)
+						{
+							LocalUpdateFavoriteSystemOnly = true;
+						}
+					}
+				}
+				ChildNodePtr->SetCanDisplayFavorite(bCanDisplayFavorite);
+
+				bool bRecurseIntoChildren =
+					!bIsChildOfCustomizedStruct // Don't recurse into built in struct children, we already know what they are and how to display them
+					&&  !bIsCustomizedStruct // Don't recurse into customized structs
+					&&	!bIsChildOfArray // Do not recurse into arrays, the children are drawn by the array property parent
+					&&	!bIsEditInlineNew // Edit inline new children are not supported for customization yet
+					&&	bIsUserVisible // Properties must be allowed to be visible by a user if they are not then their children are not visible either
+					&& (!bIsStruct || bPushOutStructProps); //  Only recurse into struct properties if they are going to be displayed as standalone properties in categories instead of inside an expandable area inside a category
+
+				if(bRecurseIntoChildren || LocalUpdateFavoriteSystemOnly)
+				{
+					// Built in struct properties or children of arras 
+					UpdateSinglePropertyMapRecursive(ChildNode, LayoutData, CurCategory, CurObjectNode, bEnableFavoriteSystem, LocalUpdateFavoriteSystemOnly);
+				}
+			}
+		}
+	}
+}
+
 void SDetailsViewBase::OnColorPickerWindowClosed(const TSharedRef<SWindow>& Window)
 {
 	const TSharedPtr< FPropertyNode > PinnedColorPropertyNode = ColorPropertyNode.Pin();
@@ -345,6 +676,11 @@
 	GenericLayoutDelegate = OnGetGenericDetails;
 }
 
+void SDetailsViewBase::RefreshRootObjectVisibility()
+{
+	RerunCurrentFilter();
+}
+
 TSharedPtr<FAssetThumbnailPool> SDetailsViewBase::GetThumbnailPool() const
 {
 	if (!ThumbnailPool.IsValid())
@@ -490,9 +826,9 @@
 	UpdateFilteredDetails();
 }
 
-void SDetailsViewBase::QueryLayoutForClass(FDetailLayoutBuilderImpl& CustomDetailLayout, UStruct* Class)
+void SDetailsViewBase::QueryLayoutForClass(FDetailLayoutData& LayoutData, UStruct* Class)
 {
-	CustomDetailLayout.SetCurrentCustomizationClass(CastChecked<UClass>(Class), NAME_None);
+	LayoutData.DetailLayout->SetCurrentCustomizationClass(CastChecked<UClass>(Class), NAME_None);
 
 	FPropertyEditorModule& ParentPlugin = FModuleManager::GetModuleChecked<FPropertyEditorModule>("PropertyEditor");
 	FCustomDetailLayoutNameMap& GlobalCustomLayoutNameMap = ParentPlugin.ClassNameToDetailLayoutNameMap;
@@ -512,26 +848,23 @@
 		TSharedRef<IDetailCustomization> CustomizationInstance = Callback->DetailLayoutDelegate.Execute();
 
 		// Ask for details immediately
-		CustomizationInstance->CustomizeDetails(CustomDetailLayout);
+		CustomizationInstance->CustomizeDetails(*LayoutData.DetailLayout);
 
 		// Save the instance from destruction until we refresh
-		CustomizationClassInstances.Add(CustomizationInstance);
+		LayoutData.CustomizationClassInstances.Add(CustomizationInstance);
 	}
 }
 
-void SDetailsViewBase::QueryCustomDetailLayout(FDetailLayoutBuilderImpl& CustomDetailLayout)
+void SDetailsViewBase::QueryCustomDetailLayout(FDetailLayoutData& LayoutData)
 {
 	FPropertyEditorModule& ParentPlugin = FModuleManager::GetModuleChecked<FPropertyEditorModule>("PropertyEditor");
 
 	// Get the registered classes that customize details
 	FCustomDetailLayoutNameMap& GlobalCustomLayoutNameMap = ParentPlugin.ClassNameToDetailLayoutNameMap;
 
-	UStruct* BaseStruct = GetBaseStruct();
+	UStruct* BaseStruct = LayoutData.DetailLayout->GetRootNode()->GetBaseStructure();
 
-	// All the current customization instances need to be deleted when it is safe
-	CustomizationClassInstancesPendingDelete = CustomizationClassInstances;
+	LayoutData.CustomizationClassInstances.Empty();
 
-	CustomizationClassInstances.Empty();
-
 	//Ask for generic details not specific to an object being viewed 
 	if (GenericLayoutDelegate.IsBound())
@@ -540,10 +873,10 @@
 		TSharedRef<IDetailCustomization> CustomizationInstance = GenericLayoutDelegate.Execute();
 
 		// Ask for details immediately
-		CustomizationInstance->CustomizeDetails(CustomDetailLayout);
+		CustomizationInstance->CustomizeDetails(*LayoutData.DetailLayout);
 
 		// Save the instance from destruction until we refresh
-		CustomizationClassInstances.Add(CustomizationInstance);
+		LayoutData.CustomizationClassInstances.Add(CustomizationInstance);
 	}
 
 
@@ -558,7 +891,7 @@
 
 	TMap< TWeakObjectPtr<UStruct>, FDetailLayoutCallback*> FinalCallbackMap;
 
-	for (auto ClassIt = ClassesWithProperties.CreateConstIterator(); ClassIt; ++ClassIt)
+	for (auto ClassIt = LayoutData.ClassesWithProperties.CreateConstIterator(); ClassIt; ++ClassIt)
 	{
 		// Check the instanced map first
 		FDetailLayoutCallback* Callback = InstancedClassToDetailLayoutMap.Find(*ClassIt);
@@ -591,11 +924,11 @@
 			{
 				UStruct* Class = WeakClass.Get();
 
-				FClassInstanceToPropertyMap& InstancedPropertyMap = ClassToPropertyMap.FindChecked(Class->GetFName());
+				FClassInstanceToPropertyMap& InstancedPropertyMap = LayoutData.ClassToPropertyMap.FindChecked(Class->GetFName());
 				for (FClassInstanceToPropertyMap::TIterator InstanceIt(InstancedPropertyMap); InstanceIt; ++InstanceIt)
 				{
 					FName Key = InstanceIt.Key();
-					CustomDetailLayout.SetCurrentCustomizationClass(Class, Key);
+					LayoutData.DetailLayout->SetCurrentCustomizationClass(Class, Key);
 
 					const FOnGetDetailCustomizationInstance& DetailDelegate = LayoutIt.Value()->DetailLayoutDelegate;
 
@@ -607,10 +940,10 @@
 						TSharedRef<IDetailCustomization> CustomizationInstance = DetailDelegate.Execute();
 
 						// Ask for details immediately
-						CustomizationInstance->CustomizeDetails(CustomDetailLayout);
+						CustomizationInstance->CustomizeDetails(*LayoutData.DetailLayout);
 
 						// Save the instance from destruction until we refresh
-						CustomizationClassInstances.Add(CustomizationInstance);
+						LayoutData.CustomizationClassInstances.Add(CustomizationInstance);
 					}
 				}
 			}
@@ -622,16 +955,16 @@
 	if (BaseStruct && !QueriedClasses.Contains(BaseStruct))
 	{
 		ParentClassesToQuery.Add(BaseStruct);
-		ClassesWithProperties.Add(BaseStruct);
+		LayoutData.ClassesWithProperties.Add(BaseStruct);
 	}
 
 	// Find base classes of queried classes that were not queried and add them to the query list
 	// this supports cases where a parent class has no properties but still wants to add customization
-	for (auto QueriedClassIt = ClassesWithProperties.CreateConstIterator(); QueriedClassIt; ++QueriedClassIt)
+	for (auto QueriedClassIt = LayoutData.ClassesWithProperties.CreateConstIterator(); QueriedClassIt; ++QueriedClassIt)
 	{
 		UStruct* ParentStruct = (*QueriedClassIt)->GetSuperStruct();
 
-		while (ParentStruct && ParentStruct->IsA(UClass::StaticClass()) && !QueriedClasses.Contains(ParentStruct) && !ClassesWithProperties.Contains(ParentStruct))
+		while (ParentStruct && ParentStruct->IsA(UClass::StaticClass()) && !QueriedClasses.Contains(ParentStruct) && !LayoutData.ClassesWithProperties.Contains(ParentStruct))
 		{
 			ParentClassesToQuery.Add(ParentStruct);
 			ParentStruct = ParentStruct->GetSuperStruct();
@@ -644,7 +977,7 @@
 	{
 		if (Cast<UClass>(*ParentIt))
 		{
-			QueryLayoutForClass(CustomDetailLayout, *ParentIt);
+			QueryLayoutForClass(LayoutData, *ParentIt);
 		}
 	}
 }
@@ -689,26 +1022,33 @@
 			PendingKillNode.Reset();
 		}
 	}
-	RootNodesPendingKill.Reset();
+	RootNodesPendingKill.Empty();
 
 	// Empty all the customization instances that need to be deleted
 	CustomizationClassInstancesPendingDelete.Empty();
 
-	auto RootPropertyNode = GetRootNode();
-	check(RootPropertyNode.IsValid());
+	FRootPropertyNodeList& RootPropertyNodes = GetRootNodes();
 
-	// Purge any objects that are marked pending kill from the object list
-	if (auto ObjectRoot = RootPropertyNode->AsObjectNode())
+	for(TSharedPtr<FComplexPropertyNode>& RootPropertyNode : RootPropertyNodes)
 	{
-		ObjectRoot->PurgeKilledObjects();
+		check(RootPropertyNode.IsValid());
+
+		// Purge any objects that are marked pending kill from the object list
+		if(auto ObjectRoot = RootPropertyNode->AsObjectNode())
+		{
+			ObjectRoot->PurgeKilledObjects();
+		}
+
+		if(DeferredActions.Num() > 0)
+		{		
+			// Any deferred actions are likely to cause the node  tree to be at least partially rebuilt
+			// Save the expansion state of existing nodes so we can expand them later
+			SaveExpandedItems(RootPropertyNode.ToSharedRef());
+		}
 	}
 
 	if (DeferredActions.Num() > 0)
 	{
-		// Any deferred actions are likely to cause the node  tree to be at least partially rebuilt
-		// Save the expansion state of existing nodes so we can expand them later
-		SaveExpandedItems( RootPropertyNode.ToSharedRef() );
-
 		// Execute any deferred actions
 		for (int32 ActionIndex = 0; ActionIndex < DeferredActions.Num(); ++ActionIndex)
 		{
@@ -718,40 +1058,53 @@
 	}
 
 	TSharedPtr<FComplexPropertyNode> LastRootPendingKill;
-	if ( RootNodesPendingKill.Num() > 0 )
+	if (RootNodesPendingKill.Num() > 0 )
 	{
 		LastRootPendingKill = RootNodesPendingKill.Last();
 	}
 
-	if ( RootPropertyNode == LastRootPendingKill )
-	{ 
-		// Reaquire the root property node.  It may have been changed by the deferred actions if something like a blueprint editor forcefully resets a details panel during a posteditchange
-		RootPropertyNode = GetRootNode();
-
-		RestoreExpandedItems( RootPropertyNode.ToSharedRef() );
-	}
-
-
 	bool bValidateExternalNodes = true;
-	FPropertyNode::DataValidationResult Result = RootPropertyNode->EnsureDataIsValid();
-	if (Result == FPropertyNode::PropertiesChanged || Result == FPropertyNode::EditInlineNewValueChanged)
-	{
-		RestoreExpandedItems( RootPropertyNode.ToSharedRef() );
-		UpdatePropertyMap();
-		UpdateFilteredDetails();
-	}
-	else if (Result == FPropertyNode::ArraySizeChanged)
-	{
-		RestoreExpandedItems( RootPropertyNode.ToSharedRef() );
-		UpdateFilteredDetails();
-	}
-	else if (Result == FPropertyNode::ObjectInvalid)
-	{
+
+	int32 FoundIndex = RootPropertyNodes.Find(LastRootPendingKill);
+	if (FoundIndex != INDEX_NONE)
+	{ 
+		// Reaquire the root property nodes.  It may have been changed by the deferred actions if something like a blueprint editor forcefully resets a details panel during a posteditchange
 		ForceRefresh();
 
 		// All objects are being reset, no need to validate external nodes
 		bValidateExternalNodes = false;
 	}
+	else
+	{
+		for(TSharedPtr<FComplexPropertyNode>& RootPropertyNode : RootPropertyNodes)
+		{
+			if(RootPropertyNode == LastRootPendingKill)
+			{
+				RestoreExpandedItems(RootPropertyNode.ToSharedRef());
+			}
 
+			EPropertyDataValidationResult Result = RootPropertyNode->EnsureDataIsValid();
+			if(Result == EPropertyDataValidationResult::PropertiesChanged || Result == EPropertyDataValidationResult::EditInlineNewValueChanged)
+			{
+				RestoreExpandedItems(RootPropertyNode.ToSharedRef());
+				UpdatePropertyMaps();
+				UpdateFilteredDetails();
+			}
+			else if(Result == EPropertyDataValidationResult::ArraySizeChanged)
+			{
+				RestoreExpandedItems(RootPropertyNode.ToSharedRef());
+				UpdateFilteredDetails();
+			}
+			else if(Result == EPropertyDataValidationResult::ObjectInvalid)
+			{
+				ForceRefresh();
+				break;
+
+				// All objects are being reset, no need to validate external nodes
+				bValidateExternalNodes = false;
+			}
+		}
+	}
+
 	if (bValidateExternalNodes)
 	{
@@ -761,17 +1114,17 @@
 
 			if (PropertyNode.IsValid())
 			{
-				Result = PropertyNode->EnsureDataIsValid();
-				if (Result == FPropertyNode::PropertiesChanged || Result == FPropertyNode::EditInlineNewValueChanged)
+				EPropertyDataValidationResult Result = PropertyNode->EnsureDataIsValid();
+				if (Result == EPropertyDataValidationResult::PropertiesChanged || Result == EPropertyDataValidationResult::EditInlineNewValueChanged)
 				{
 					RestoreExpandedItems(PropertyNode.ToSharedRef());
-					UpdatePropertyMap();
+					UpdatePropertyMaps();
 					UpdateFilteredDetails();
 					// Note this will invalidate all the external root nodes so there is no need to continue
 					ExternalRootPropertyNodes.Empty();
 					break;
 				}
-				else if (Result == FPropertyNode::ArraySizeChanged)
+				else if (Result == EPropertyDataValidationResult::ArraySizeChanged)
 				{
 					RestoreExpandedItems(PropertyNode.ToSharedRef());
 					UpdateFilteredDetails();
@@ -786,9 +1139,12 @@
 		}
 	}
 
-	if (DetailLayout.IsValid())
+	for(FDetailLayoutData& LayoutData : DetailLayouts)
 	{
-		DetailLayout->Tick(InDeltaTime);
+		if(LayoutData.DetailLayout.IsValid())
+		{
+			LayoutData.DetailLayout->Tick(InDeltaTime);
+		}
 	}
 
 	if (!ColorPropertyNode.IsValid() && bHasOpenColorPicker)
@@ -908,7 +1264,7 @@
 		}
 	}
 
-	if (DetailLayout.IsValid() && BestBaseStruct)
+	if (DetailLayouts.Num() > 0 && BestBaseStruct)
 	{
 		bool bShouldSave = !ExpandedCustomItemsString.IsEmpty();
 		if (!bShouldSave)
@@ -957,331 +1313,68 @@
 
 void SDetailsViewBase::UpdateFilteredDetails()
 {
-	auto RootPropertyNode = GetRootNode();
-	if (RootPropertyNode.IsValid())
-	{
-		RootPropertyNode->FilterNodes(CurrentFilter.FilterStrings);
-		RootPropertyNode->ProcessSeenFlags(true);
+	RootTreeNodes.Reset();
 
-		for (int32 NodeIndex = 0; NodeIndex < ExternalRootPropertyNodes.Num(); ++NodeIndex)
-		{
-			TSharedPtr<FPropertyNode> PropertyNode = ExternalRootPropertyNodes[NodeIndex].Pin();
-
-			if (PropertyNode.IsValid())
-			{
-				PropertyNode->FilterNodes(CurrentFilter.FilterStrings);
-				PropertyNode->ProcessSeenFlags(true);
-			}
-		}
-
-		if (DetailLayout.IsValid())
-		{
-			DetailLayout->FilterDetailLayout(CurrentFilter);
-		}
-
-		RootTreeNodes = DetailLayout->GetRootTreeNodes();
-	}
-
-	RefreshTree();
-}
-
-/**
-* Determines whether or not a property should be visible in the default generated detail layout
-*
-* @param PropertyNode	The property node to check
-* @param ParentNode	The parent property node to check
-* @return true if the property should be visible
-*/
-static bool IsVisibleStandaloneProperty(const FPropertyNode& PropertyNode, const FPropertyNode& ParentNode)
-{
-	const UProperty* Property = PropertyNode.GetProperty();
-	const UArrayProperty* ParentArrayProperty = Cast<const UArrayProperty>(ParentNode.GetProperty());
-	bool bIsVisibleStandalone = false;
-	if (Property)
-	{
-		if (Property->IsA(UObjectPropertyBase::StaticClass()))
-		{
-			// Do not add this child node to the current map if its a single object property in a category (serves no purpose for UI)
-			bIsVisibleStandalone = !ParentArrayProperty && (PropertyNode.GetNumChildNodes() == 0 || PropertyNode.GetNumChildNodes() > 1);
-		}
-		else if (Property->IsA(UArrayProperty::StaticClass()) || (Property->ArrayDim > 1 && PropertyNode.GetArrayIndex() == INDEX_NONE))
-		{
-			// Base array properties are always visible
-			bIsVisibleStandalone = true;
-		}
-		else
-		{
-			bIsVisibleStandalone = true;
-		}
-
-	}
-
-	return bIsVisibleStandalone;
-}
-
-void SDetailsViewBase::UpdatePropertyMapRecursive(FPropertyNode& InNode, FDetailLayoutBuilderImpl& InDetailLayout, FName CurCategory, FComplexPropertyNode* CurObjectNode, bool bEnableFavoriteSystem, bool bUpdateFavoriteSystemOnly)
-{
-	UProperty* ParentProperty = InNode.GetProperty();
-	UStructProperty* ParentStructProp = Cast<UStructProperty>(ParentProperty);
-	for (int32 ChildIndex = 0; ChildIndex < InNode.GetNumChildNodes(); ++ChildIndex)
-	{
-		//Use the original value for each child
-		bool LocalUpdateFavoriteSystemOnly = bUpdateFavoriteSystemOnly;
-
-		TSharedPtr<FPropertyNode> ChildNodePtr = InNode.GetChildNode(ChildIndex);
-		FPropertyNode& ChildNode = *ChildNodePtr;
-		UProperty* Property = ChildNode.GetProperty();
-
-		{
-			FObjectPropertyNode* ObjNode = ChildNode.AsObjectNode();
-			FCategoryPropertyNode* CategoryNode = ChildNode.AsCategoryNode();
-			if (ObjNode)
-			{
-				// Currently object property nodes do not provide any useful information other than being a container for its children.  We do not draw anything for them.
-				// When we encounter object property nodes, add their children instead of adding them to the tree.
-				UpdatePropertyMapRecursive(ChildNode, InDetailLayout, CurCategory, ObjNode, bEnableFavoriteSystem, LocalUpdateFavoriteSystemOnly);
-			}
-			else if (CategoryNode)
-			{
-				if (!LocalUpdateFavoriteSystemOnly)
-				{
-					FName InstanceName = NAME_None;
-					FName CategoryName = CurCategory;
-					FString CategoryDelimiterString;
-					CategoryDelimiterString.AppendChar(FPropertyNodeConstants::CategoryDelimiterChar);
-					if (CurCategory != NAME_None && CategoryNode->GetCategoryName().ToString().Contains(CategoryDelimiterString))
-					{
-						// This property is child of another property so add it to the parent detail category
-						FDetailCategoryImpl& CategoryImpl = InDetailLayout.DefaultCategory(CategoryName);
-						CategoryImpl.AddPropertyNode(ChildNodePtr.ToSharedRef(), InstanceName);
-					}
-				}
-
-				// For category nodes, we just set the current category and recurse through the children
-				UpdatePropertyMapRecursive(ChildNode, InDetailLayout, CategoryNode->GetCategoryName(), CurObjectNode, bEnableFavoriteSystem, LocalUpdateFavoriteSystemOnly);
-			}
-			else
-			{
-				// Whether or not the property can be visible in the default detail layout
-				bool bVisibleByDefault = IsVisibleStandaloneProperty(ChildNode, InNode);
-
-				// Whether or not the property is a struct
-				UStructProperty* StructProperty = Cast<UStructProperty>(Property);
-
-				bool bIsStruct = StructProperty != NULL;
-
-				static FName ShowOnlyInners("ShowOnlyInnerProperties");
-
-				bool bIsChildOfCustomizedStruct = false;
-				bool bIsCustomizedStruct = false;
-
-				const UStruct* Struct = StructProperty ? StructProperty->Struct : NULL;
-				const UStruct* ParentStruct = ParentStructProp ? ParentStructProp->Struct : NULL;
-				if (Struct || ParentStruct)
-				{
-					FPropertyEditorModule& ParentPlugin = FModuleManager::GetModuleChecked<FPropertyEditorModule>("PropertyEditor");
-					if (Struct)
-					{
-						bIsCustomizedStruct = ParentPlugin.IsCustomizedStruct(Struct, SharedThis( this ) );
-					}
-
-					if (ParentStruct)
-					{
-						bIsChildOfCustomizedStruct = ParentPlugin.IsCustomizedStruct(ParentStruct, SharedThis( this ) );
-					}
-				}
-
-				// Whether or not to push out struct properties to their own categories or show them inside an expandable struct 
-				bool bPushOutStructProps = bIsStruct && !bIsCustomizedStruct && !ParentStructProp && Property->HasMetaData(ShowOnlyInners);
-
-				// Is the property edit inline new 
-				const bool bIsEditInlineNew = SPropertyEditorEditInline::Supports(&ChildNode, ChildNode.GetArrayIndex());
-
-				// Is this a property of an array
-				bool bIsChildOfArray = PropertyEditorHelpers::IsChildOfArray(ChildNode);
-
-				// Edit inline new properties should be visible by default
-				bVisibleByDefault |= bIsEditInlineNew;
-
-				// Children of arrays are not visible directly,
-				bVisibleByDefault &= !bIsChildOfArray;
-
-				FPropertyAndParent PropertyAndParent(*Property, ParentProperty);
-				const bool bIsUserVisible = IsPropertyVisible(PropertyAndParent);
-
-				// Inners of customized in structs should not be taken into consideration for customizing.  They are not designed to be individually customized when their parent is already customized
-				if (!bIsChildOfCustomizedStruct && !LocalUpdateFavoriteSystemOnly)
-				{
-					// Add any object classes with properties so we can ask them for custom property layouts later
-					ClassesWithProperties.Add(Property->GetOwnerStruct());
-				}
-
-				// If there is no outer object then the class is the object root and there is only one instance
-				FName InstanceName = NAME_None;
-				if (CurObjectNode && CurObjectNode->GetParentNode())
-				{
-					InstanceName = CurObjectNode->GetParentNode()->GetProperty()->GetFName();
-				}
-				else if (ParentStructProp)
-				{
-					InstanceName = ParentStructProp->GetFName();
-				}
-
-				// Do not add children of customized in struct properties or arrays
-				if (!bIsChildOfCustomizedStruct && !bIsChildOfArray && !LocalUpdateFavoriteSystemOnly)
-				{
-					// Get the class property map
-					FClassInstanceToPropertyMap& ClassInstanceMap = ClassToPropertyMap.FindOrAdd(Property->GetOwnerStruct()->GetFName());
-
-					FPropertyNodeMap& PropertyNodeMap = ClassInstanceMap.FindOrAdd(InstanceName);
-
-					if (!PropertyNodeMap.ParentProperty)
-					{
-						PropertyNodeMap.ParentProperty = CurObjectNode;
-					}
-					else
-					{
-						ensure(PropertyNodeMap.ParentProperty == CurObjectNode);
-					}
-
-					checkSlow(!PropertyNodeMap.Contains(Property->GetFName()));
-
-					PropertyNodeMap.Add(Property->GetFName(), ChildNodePtr);
-				}
-				bool bCanDisplayFavorite = false;
-				if (bVisibleByDefault && bIsUserVisible && !bPushOutStructProps)
-				{
-					FName CategoryName = CurCategory;
-					// For properties inside a struct, add them to their own category unless they just take the name of the parent struct.  
-					// In that case push them to the parent category
-					FName PropertyCatagoryName = FObjectEditorUtils::GetCategoryFName(Property);
-					if (!ParentStructProp || (PropertyCatagoryName != ParentStructProp->Struct->GetFName()))
-					{
-						CategoryName = PropertyCatagoryName;
-					}
-
-					if (!LocalUpdateFavoriteSystemOnly)
-					{
-						if (IsPropertyReadOnly(PropertyAndParent))
-						{
-							ChildNode.SetNodeFlags(EPropertyNodeFlags::IsReadOnly, true);
-						}
-
-						// Add a property to the default category
-						FDetailCategoryImpl& CategoryImpl = InDetailLayout.DefaultCategory(CategoryName);
-						CategoryImpl.AddPropertyNode(ChildNodePtr.ToSharedRef(), InstanceName);
-					}
-					
-					bCanDisplayFavorite = true;
-					if (bEnableFavoriteSystem)
-					{
-						if (bIsCustomizedStruct)
-						{
-							bCanDisplayFavorite = false;
-							//CustomizedStruct child are not categorize since they are under an object but we have to put them in favorite category if the user want to favorite them
-							LocalUpdateFavoriteSystemOnly = true;
-						}
-						else if (ChildNodePtr->IsFavorite())
-						{
-							//Find or create the favorite category, we have to duplicate favorite property row under this category
-							FString CategoryFavoritesName = TEXT("Favorites");
-							FName CatFavName = *CategoryFavoritesName;
-							FDetailCategoryImpl& CategoryFavImpl = InDetailLayout.DefaultCategory(CatFavName);
-							CategoryFavImpl.SetSortOrder(0);
-							CategoryFavImpl.SetCategoryAsSpecialFavorite();
-
-							//Add the property to the favorite
-							FObjectPropertyNode *RootObjectParent = ChildNodePtr->FindRootObjectItemParent();
-							FName RootInstanceName = NAME_None;
-							if (RootObjectParent != nullptr)
-							{
-								RootInstanceName = RootObjectParent->GetObjectBaseClass()->GetFName();
-							}
-
-							if (LocalUpdateFavoriteSystemOnly)
-							{
-								if (IsPropertyReadOnly(PropertyAndParent))
-								{
-									ChildNode.SetNodeFlags(EPropertyNodeFlags::IsReadOnly, true);
-								}
-								else
-								{
-									//If the parent has a condition that is not met, make the child as readonly
-									FDetailLayoutCustomization ParentTmpCustomization;
-									ParentTmpCustomization.PropertyRow = MakeShareable(new FDetailPropertyRow(InNode.AsShared(), CategoryFavImpl.AsShared()));
-									if(ParentTmpCustomization.PropertyRow->GetPropertyEditor()->IsPropertyEditingEnabled() == false)
-									{
-										ChildNode.SetNodeFlags(EPropertyNodeFlags::IsReadOnly, true);
-									}
-								}
-							}
-							
-							//Duplicate the row
-							CategoryFavImpl.AddPropertyNode(ChildNodePtr.ToSharedRef(), RootInstanceName);
-						}
-
-						if (bIsStruct)
-						{
-							LocalUpdateFavoriteSystemOnly = true;
-						}
-					}
-				}
-				ChildNodePtr->SetCanDisplayFavorite(bCanDisplayFavorite);
-
-				bool bRecurseIntoChildren =
-					!bIsChildOfCustomizedStruct // Don't recurse into built in struct children, we already know what they are and how to display them
-					&&  !bIsCustomizedStruct // Don't recurse into customized structs
-					&&	!bIsChildOfArray // Do not recurse into arrays, the children are drawn by the array property parent
-					&&	!bIsEditInlineNew // Edit inline new children are not supported for customization yet
-					&&	bIsUserVisible // Properties must be allowed to be visible by a user if they are not then their children are not visible either
-					&& (!bIsStruct || bPushOutStructProps); //  Only recurse into struct properties if they are going to be displayed as standalone properties in categories instead of inside an expandable area inside a category
-
-				if (bRecurseIntoChildren || LocalUpdateFavoriteSystemOnly)
-				{
-					// Built in struct properties or children of arras 
-					UpdatePropertyMapRecursive(ChildNode, InDetailLayout, CurCategory, CurObjectNode, bEnableFavoriteSystem, LocalUpdateFavoriteSystemOnly);
-				}
-			}
-		}
-	}
-}
-
-void SDetailsViewBase::UpdatePropertyMap()
-{
-	// Reset everything
-	ClassToPropertyMap.Empty();
-	ClassesWithProperties.Empty();
-
-	// We need to be able to create a new detail layout and properly clean up the old one in the process
-	check(!DetailLayout.IsValid() || DetailLayout.IsUnique());
-
-	RootTreeNodes.Empty();
-
-	DetailLayout = MakeShareable(new FDetailLayoutBuilderImpl(ClassToPropertyMap, PropertyUtilities.ToSharedRef(), SharedThis(this)));
-
-
-	auto RootPropertyNode = GetRootNode();
-	check(RootPropertyNode.IsValid());
-
-	bool const bEnableFavoriteSystem = GIsRequestingExit ? false : (GetDefault<UEditorExperimentalSettings>()->bEnableFavoriteSystem && DetailsViewArgs.bAllowFavoriteSystem);
-
-	// Currently object property nodes do not provide any useful information other than being a container for its children.  We do not draw anything for them.
-	// When we encounter object property nodes, add their children instead of adding them to the tree.
-	UpdatePropertyMapRecursive(*RootPropertyNode, *DetailLayout, NAME_None, RootPropertyNode.Get(), bEnableFavoriteSystem, false);
-
-	CustomUpdatePropertyMap();
-
-	// Ask for custom detail layouts, unless disabled. One reason for disabling custom layouts is that the custom layouts
-	// inhibit our ability to find a single property's tree node. This is problematic for the diff and merge tools, that need
-	// to display and highlight each changed property for the user. We could whitelist 'good' customizations here if 
-	// we can make them work with the diff/merge tools.
-	if( !bDisableCustomDetailLayouts )
-	{
-		QueryCustomDetailLayout(*DetailLayout);
-	}
+	FDetailNodeList InitialRootNodeList;
 	
-	DetailLayout->GenerateDetailLayout();
+	NumVisbleTopLevelObjectNodes = 0;
+	FRootPropertyNodeList& RootPropertyNodes = GetRootNodes();
+
+	for(int32 RootNodeIndex = 0; RootNodeIndex < RootPropertyNodes.Num(); ++RootNodeIndex)
+	{
+		TSharedPtr<FComplexPropertyNode>& RootPropertyNode = RootPropertyNodes[RootNodeIndex];
+		if(RootPropertyNode.IsValid())
+		{
+			RootPropertyNode->FilterNodes(CurrentFilter.FilterStrings);
+			RootPropertyNode->ProcessSeenFlags(true);
+
+			for(int32 NodeIndex = 0; NodeIndex < ExternalRootPropertyNodes.Num(); ++NodeIndex)
+			{
+				TSharedPtr<FPropertyNode> PropertyNode = ExternalRootPropertyNodes[NodeIndex].Pin();
+
+				if(PropertyNode.IsValid())
+				{
+					PropertyNode->FilterNodes(CurrentFilter.FilterStrings);
+					PropertyNode->ProcessSeenFlags(true);
+				}
+			}
+
+			TSharedPtr<FDetailLayoutBuilderImpl>& DetailLayout = DetailLayouts[RootNodeIndex].DetailLayout;
+			if(DetailLayout.IsValid())
+			{
+				DetailLayout->FilterDetailLayout(CurrentFilter);
+			}
+
+			FDetailNodeList& LayoutRoots = DetailLayout->GetRootTreeNodes();
+			if(LayoutRoots.Num() > 0)
+			{
+				// A top level object nodes has a non-filtered away root so add one to the total number we have
+				++NumVisbleTopLevelObjectNodes;
+
+				InitialRootNodeList.Append(LayoutRoots);
+			}
+		
+		}
+	}
+
+
+	// for multiple top level object we need to do a secondary pass on top level object nodes after we have determined if there is any nodes visible at all.  If there are then we ask the details panel if it wants to show childen
+	for(TSharedRef<class IDetailTreeNode> RootNode : InitialRootNodeList)
+	{
+		if(RootNode->ShouldShowOnlyChildren())
+		{
+			RootNode->GetChildren(RootTreeNodes);
+		}
+		else
+		{
+			RootTreeNodes.Add(RootNode);
+		}
+	}
+
+	RefreshTree();
 }
 
+
+
 void SDetailsViewBase::RegisterInstancedCustomPropertyLayoutInternal(UStruct* Class, FOnGetDetailCustomizationInstance DetailLayoutDelegate)
 {
diff --git a/Engine/Source/Editor/PropertyEditor/Private/SDetailsViewBase.h b/Engine/Source/Editor/PropertyEditor/Private/SDetailsViewBase.h
--- a/Engine/Source/Editor/PropertyEditor/Private/SDetailsViewBase.h
+++ b/Engine/Source/Editor/PropertyEditor/Private/SDetailsViewBase.h
@@ -47,6 +47,24 @@ typedef TMap<FName, FClassInstanceToPropertyMap> FClassToPropertyMap;
 
 typedef STreeView< TSharedRef<class IDetailTreeNode> > SDetailTree;
 
+typedef TArray<TSharedPtr<FComplexPropertyNode>> FRootPropertyNodeList;
+
+struct FDetailLayoutData
+{
+	TSharedPtr<FDetailLayoutBuilderImpl> DetailLayout;
+
+	FClassToPropertyMap ClassToPropertyMap;
+
+	/** A  unique classes being viewed */
+	TSet< TWeakObjectPtr<UStruct> > ClassesWithProperties;
+
+	/** Customization class instances currently active in this view */
+	TArray< TSharedPtr<IDetailCustomization> > CustomizationClassInstances;
+
+};
+
+typedef TArray<FDetailLayoutData> FDetailLayoutList;
+
 /** Represents a filter which controls the visibility of items in the details view */
 struct FDetailFilter
 {
@@ -90,9 +108,12 @@ public:
 		, bIsLocked(false)
 		, bHasOpenColorPicker(false)
 		, bDisableCustomDetailLayouts( false )
+		, NumVisbleTopLevelObjectNodes(0)
 	{
 	}
 
+	virtual ~SDetailsViewBase();
+
 	/**
 	* @return true of the details view can be updated from editor selection
 	*/
@@ -101,6 +122,17 @@ public:
 		return DetailsViewArgs.bUpdatesFromSelection;
 	}
 
+	virtual bool HasActiveSearch() const override
+	{
+		return CurrentFilter.FilterStrings.Num() > 0;
+	}
+
+	virtual int32 GetNumVisibleTopLevelObjects() const override
+	{
+		return NumVisbleTopLevelObjectNodes;
+	}
+
+
 	/** @return The identifier for this details view, or NAME_None is this view is anonymous */
 	virtual FName GetIdentifier() const override
 	{
@@ -120,54 +152,42 @@ public:
 	virtual void ShowAllAdvancedProperties() override;
 	virtual void SetOnDisplayedPropertiesChanged(FOnDisplayedPropertiesChanged InOnDisplayedPropertiesChangedDelegate) override;
 	virtual void SetDisableCustomDetailLayouts( bool bInDisableCustomDetailLayouts ) override { bDisableCustomDetailLayouts = bInDisableCustomDetailLayouts; }
-	virtual void RerunCurrentFilter() override;
+	virtual void SetIsPropertyVisibleDelegate(FIsPropertyVisible InIsPropertyVisible) override;
+	virtual void SetIsPropertyReadOnlyDelegate(FIsPropertyReadOnly InIsPropertyReadOnly) override;
+	virtual void SetIsPropertyEditingEnabledDelegate(FIsPropertyEditingEnabled IsPropertyEditingEnabled) override;
+	virtual bool IsPropertyEditingEnabled() const override;
+	virtual void SetKeyframeHandler(TSharedPtr<class IDetailKeyframeHandler> InKeyframeHandler) override;
+	virtual TSharedPtr<IDetailKeyframeHandler> GetKeyframeHandler() override;
+	virtual void SetExtensionHandler(TSharedPtr<class IDetailPropertyExtensionHandler> InExtensionHandler) override;
+	virtual TSharedPtr<IDetailPropertyExtensionHandler> GetExtensionHandler() override;
+	virtual bool IsPropertyVisible(const struct FPropertyAndParent& PropertyAndParent) const override;
+	virtual bool IsPropertyReadOnly(const struct FPropertyAndParent& PropertyAndParent) const override;
+	virtual void SetGenericLayoutDetailsDelegate(FOnGetDetailCustomizationInstance OnGetGenericDetails) override;
+	virtual bool IsLocked() const override { return bIsLocked; }
+	virtual void RefreshRootObjectVisibility() override;
+	virtual FOnFinishedChangingProperties& OnFinishedChangingProperties() override { return OnFinishedChangingPropertiesDelegate; }
 
-	virtual FOnFinishedChangingProperties& OnFinishedChangingProperties() override
-	{ 
-		return OnFinishedChangingPropertiesDelegate; 
-	}
+	/** IDetailsViewPrivate interface */
+	virtual void RerunCurrentFilter() override;
+	void SetNodeExpansionState(TSharedRef<IDetailTreeNode> InTreeNode, bool bIsItemExpanded, bool bRecursive) override;
+	void SaveCustomExpansionState(const FString& NodePath, bool bIsExpanded) override;
+	bool GetCustomSavedExpansionState(const FString& NodePath) const override;
+	virtual void NotifyFinishedChangingProperties(const FPropertyChangedEvent& PropertyChangedEvent) override;
+	void RefreshTree() override;
+	TSharedPtr<class FAssetThumbnailPool> GetThumbnailPool() const override;
+	TSharedPtr<IPropertyUtilities> GetPropertyUtilities() override;
+	void CreateColorPickerWindow(const TSharedRef< class FPropertyEditor >& PropertyEditor, bool bUseAlpha) override;
+	virtual FNotifyHook* GetNotifyHook() const override { return DetailsViewArgs.NotifyHook; }
 
 	virtual bool IsConnected() const = 0;
 
+	virtual FRootPropertyNodeList& GetRootNodes() = 0;
+
 	/**
 	 * Called when the open color picker window associated with this details view is closed
 	 */
 	void OnColorPickerWindowClosed(const TSharedRef<SWindow>& Window);
 
-	/**
-	 * Returns true if the details view is locked and cant have its observed objects changed 
-	 */
-	virtual bool IsLocked() const override
-	{
-		return bIsLocked;
-	}
-
-	/**
-	 * Sets a delegate which is called to determine whether a specific property should be visible
-	 */
-	virtual void SetIsPropertyVisibleDelegate(FIsPropertyVisible InIsPropertyVisible) override;
-
-	/**
-	 * Sets a delegate to call to determine if a specific property should be read-only in this instance of the details view
-	 */ 
-	virtual void SetIsPropertyReadOnlyDelegate( FIsPropertyReadOnly InIsPropertyReadOnly ) override;
-
-	/**
-	 * Sets a delegate to call to determine if the properties  editing is enabled
-	 */ 
-	virtual void SetIsPropertyEditingEnabledDelegate(FIsPropertyEditingEnabled IsPropertyEditingEnabled) override;
-
-	/**
-	 * @return true if property editing is enabled (based on the FIsPropertyEditingEnabled delegate)
-	 */
-	virtual bool IsPropertyEditingEnabled() const override;
-
-	virtual void SetKeyframeHandler( TSharedPtr<class IDetailKeyframeHandler> InKeyframeHandler ) override;
-	virtual TSharedPtr<IDetailKeyframeHandler> GetKeyframeHandler() override;
-
-	virtual void SetExtensionHandler(TSharedPtr<class IDetailPropertyExtensionHandler> InExtensionHandler) override;
-	virtual TSharedPtr<IDetailPropertyExtensionHandler> GetExtensionHandler() override;
-
 	/**
 	 * Requests that an item in the tree be expanded or collapsed
 	 *
@@ -176,15 +196,6 @@ public:
 	 */
 	void RequestItemExpanded(TSharedRef<IDetailTreeNode> TreeNode, bool bExpand) override;
 
-	/**
-	 * Sets the expansion state for a node and optionally all of its children
-	 *
-	 * @param InTreeNode		The node to change expansion state on
-	 * @param bIsItemExpanded	The new expansion state
-	 * @param bRecursive		Whether or not to apply the expansion change to any children
-	 */
-	void SetNodeExpansionState(TSharedRef<IDetailTreeNode> InTreeNode, bool bIsItemExpanded, bool bRecursive) override;
-
 	/**
 	 * Sets the expansion state all root nodes and optionally all of their children
 	 *
@@ -196,81 +207,23 @@ public:
 	/**
 	 * Queries a layout for a specific class
 	 */
-	void QueryLayoutForClass(FDetailLayoutBuilderImpl& CustomDetailLayout, UStruct* Class);
+	void QueryLayoutForClass(FDetailLayoutData& LayoutData, UStruct* Class);
 
 	/**
 	 * Calls a delegate for each registered class that has properties visible to get any custom detail layouts
 	 */
-	void QueryCustomDetailLayout(class FDetailLayoutBuilderImpl& CustomDetailLayout);
-
-	/**
-	 * Refreshes the detail's treeview
-	 */
-	void RefreshTree() override;
-
-	/**
-	 * Saves the expansion state of a tree node
-	 *
-	 * @param NodePath	The path to the detail node to save
-	 * @param bIsExpanded	true if the node is expanded, false otherwise
-	 */
-	void SaveCustomExpansionState(const FString& NodePath, bool bIsExpanded) override;
-
-	/**
-	 * Gets the saved expansion state of a tree node in this category	
-	 *
-	 * @param NodePath	The path to the detail node to get
-	 * @return true if the node should be expanded, false otherwise
-	 */	
-	bool GetCustomSavedExpansionState(const FString& NodePath) const override;
-
-	/**
-	 * Called when properties have finished changing (after PostEditChange is called)
-	 */
-	virtual void NotifyFinishedChangingProperties(const FPropertyChangedEvent& PropertyChangedEvent) override;
+	void QueryCustomDetailLayout(FDetailLayoutData& LayoutData);
 
 	/** Column width accessibility */
 	float OnGetLeftColumnWidth() const { return 1.0f - ColumnWidth; }
 	float OnGetRightColumnWidth() const { return ColumnWidth; }
 	void OnSetColumnWidth(float InWidth) { ColumnWidth = InWidth; }
 
-	/**
-	 * @return True if the property is visible
-	 */
-	virtual bool IsPropertyVisible( const struct FPropertyAndParent& PropertyAndParent ) const override;
-
-	/**
-	 * @return True if the property is visible
-	 */
-	virtual bool IsPropertyReadOnly( const struct FPropertyAndParent& PropertyAndParent ) const override;
-
-	/**
-	* Sets a delegate which is regardless of the objects being viewed to lay out generic details not specific to any object
-	*/
-	virtual void SetGenericLayoutDetailsDelegate(FOnGetDetailCustomizationInstance OnGetGenericDetails) override;
-
 	/**
 	 * Adds an action to execute next tick
 	 */
 	virtual void EnqueueDeferredAction(FSimpleDelegate& DeferredAction) override;
 
-	/**
-	 * @return The thumbnail pool that should be used for thumbnails being rendered in this view
-	 */
-	TSharedPtr<class FAssetThumbnailPool> GetThumbnailPool() const override;
-
-	/**
-	 * Returns the property utilities for this view
-	 */
-	TSharedPtr<IPropertyUtilities> GetPropertyUtilities() override;
-
-	/** Returns the notify hook to use when properties change */
-	virtual FNotifyHook* GetNotifyHook() const override
-	{ 
-		return DetailsViewArgs.NotifyHook; 
-	}
-
-	virtual ~SDetailsViewBase();
 
 	// SWidget interface
 	virtual bool SupportsKeyboardFocus() const override;
@@ -286,15 +239,6 @@ public:
 	*/
 	void RestoreExpandedItems(TSharedRef<FPropertyNode> StartNode);
 
-	virtual TSharedPtr<FComplexPropertyNode> GetRootNode() = 0;
-
-	/**
-	 * Creates the color picker window for this property view.
-	 *
-	 * @param PropertyEditor				The slate property node to edit.
-	 * @param bUseAlpha			Whether or not alpha is supported
-	 */
-	void CreateColorPickerWindow(const TSharedRef< class FPropertyEditor >& PropertyEditor, bool bUseAlpha) override;
 
 protected:
 	/**
@@ -303,8 +247,9 @@ protected:
 	void SetColorPropertyFromColorPicker(FLinearColor NewColor);
 
 	/** Updates the property map for access when customizing the details view.  Generates default layout for properties */
-	void UpdatePropertyMap();
-	virtual void CustomUpdatePropertyMap() {}
+	void UpdatePropertyMaps();
+	void UpdateSinglePropertyMap(TSharedPtr<FComplexPropertyNode>& InRootPropertyNode, FDetailLayoutData& LayoutData);
+
 	/** 
 	 * Recursively updates children of property nodes. Generates default layout for properties 
 	 * 
@@ -312,8 +257,9 @@ protected:
 	 * @param The detail layout builder that will be used for customization of this property map
 	 * @param CurCategory The current category name
 	 */
-	void UpdatePropertyMapRecursive( FPropertyNode& InNode, FDetailLayoutBuilderImpl& DetailLayout, FName CurCategory, FComplexPropertyNode* CurObjectNode, bool bEnableFavoriteSystem, bool bUpdateFavoriteSystemOnly);
+	void UpdateSinglePropertyMapRecursive( FPropertyNode& InNode, FDetailLayoutData& LayoutData, FName CurCategory, FComplexPropertyNode* CurObjectNode, bool bEnableFavoriteSystem, bool bUpdateFavoriteSystemOnly);
 
+	virtual void CustomUpdatePropertyMap(TSharedPtr<FDetailLayoutBuilderImpl>& InDetailLayout) {}
 
 	/** Called to get the visibility of the tree view */
 	EVisibility GetTreeVisibility() const;
@@ -411,20 +357,14 @@ protected:
 protected:
 	/** The user defined args for the details view */
 	FDetailsViewArgs DetailsViewArgs;
-	/** A mapping of class (or struct) to properties in that struct (only top level, non-deeply nested properties appear in this map) */
-	FClassToPropertyMap ClassToPropertyMap;
-	/** A mapping unique classes being viewed to their variable names (for multiple of the same type being viewed)*/
-	TSet< TWeakObjectPtr<UStruct> > ClassesWithProperties;
 	/** A mapping of classes to detail layout delegates, called when querying for custom detail layouts in this instance of the details view only*/
 	FCustomDetailLayoutMap InstancedClassToDetailLayoutMap;
-	/** The current detail layout based on selection */
-	TSharedPtr<class FDetailLayoutBuilderImpl> DetailLayout;
+	/** The current detail layout based on objects in this details panel.  There is one layout for each top level object node.*/
+	FDetailLayoutList DetailLayouts;
 	/** Row for searching and view options */
 	TSharedPtr<SHorizontalBox>  FilterRow;
 	/** Search box */
-	TSharedPtr<SWidget> SearchBox;
-	/** Customization class instances currently active in this view */
-	TArray< TSharedPtr<IDetailCustomization> > CustomizationClassInstances;
+	TSharedPtr<SSearchBox> SearchBox;
 	/** Customization instances that need to be destroyed when safe to do so */
 	TArray< TSharedPtr<IDetailCustomization> > CustomizationClassInstancesPendingDelete;
 	/** Map of nodes that are requesting an automatic expansion/collapse due to being filtered */
@@ -434,7 +374,7 @@ protected:
 	/** Tree view */
 	TSharedPtr<SDetailTree> DetailTree;
 	/** Root tree nodes visible in the tree */
-	TArray< TSharedRef<IDetailTreeNode> > RootTreeNodes;
+	FDetailNodeList RootTreeNodes;
 	/** Delegate executed to determine if a property should be visible */
 	FIsPropertyVisible IsPropertyVisibleDelegate;
 	/** Delegate executed to determine if a property should be read-only */
@@ -452,7 +392,7 @@ protected:
 	/** True if this property view is currently locked (I.E The objects being observed are not changed automatically due to user selection)*/
 	bool bIsLocked;
 	/** The property node that the color picker is currently editing. */
-	TWeakPtr< FPropertyNode > ColorPropertyNode;
+	TWeakPtr<FPropertyNode> ColorPropertyNode;
 	/** Whether or not this instance of the details view opened a color picker and it is not closed yet */
 	bool bHasOpenColorPicker;
 	/** Settings for this view */
@@ -469,7 +409,7 @@ protected:
 	TArray<FSimpleDelegate> DeferredActions;
 
 	/** Root tree nodes that needs to be destroyed when safe */
-	TArray<TSharedPtr<FComplexPropertyNode>> RootNodesPendingKill;
+	FRootPropertyNodeList RootNodesPendingKill;
 
 	/** The handler for the keyframe UI, determines if the key framing UI should be displayed. */
 	TSharedPtr<IDetailKeyframeHandler> KeyframeHandler;
@@ -488,4 +428,6 @@ protected:
 
 	/** True if we want to skip generation of custom layouts for displayed object */
 	bool bDisableCustomDetailLayouts;
+
+	int32 NumVisbleTopLevelObjectNodes;
 };
diff --git a/Engine/Source/Editor/PropertyEditor/Private/SPropertyTreeViewImpl.cpp b/Engine/Source/Editor/PropertyEditor/Private/SPropertyTreeViewImpl.cpp
--- a/Engine/Source/Editor/PropertyEditor/Private/SPropertyTreeViewImpl.cpp
+++ b/Engine/Source/Editor/PropertyEditor/Private/SPropertyTreeViewImpl.cpp
@@ -1001,13 +1001,13 @@ void SPropertyTreeViewImpl::Tick( const FGeometry& AllottedGeometry, const doubl
 	}
 	DeferredActions.Empty();
 
-	FPropertyNode::DataValidationResult Result = RootPropertyNode->EnsureDataIsValid();
-	if( Result == FPropertyNode::PropertiesChanged || Result == FPropertyNode::ArraySizeChanged || Result == FPropertyNode::EditInlineNewValueChanged )
+	EPropertyDataValidationResult Result = RootPropertyNode->EnsureDataIsValid();
+	if( Result == EPropertyDataValidationResult::PropertiesChanged || Result == EPropertyDataValidationResult::ArraySizeChanged || Result == EPropertyDataValidationResult::EditInlineNewValueChanged )
 	{
 		// Make sure our new property windows are properly filtered.  
 		FilterView( CurrentFilterText );
 	}
-	else if( Result == FPropertyNode::ObjectInvalid && !bNodeTreeExternallyManaged )
+	else if( Result == EPropertyDataValidationResult::ObjectInvalid && !bNodeTreeExternallyManaged )
 	{
 		TArray< TWeakObjectPtr< UObject > > ResetArray;
 		for ( TPropObjectIterator Itor( RootPropertyNode->ObjectIterator() ) ; Itor ; ++Itor )
@@ -1020,12 +1020,6 @@ void SPropertyTreeViewImpl::Tick( const FGeometry& AllottedGeometry, const doubl
 			}
 		}
 
-		//@todo Slate property window: Support window flags
-		/*
-		EPropertyWindowFlags::Type OldFlags = HasFlags( EPropertyWindowFlags::ShouldShowCategories | 
-		EPropertyWindowFlags::Sorted | 
-		EPropertyWindowFlags::ShouldShowNonEditable);*/
-
 		SetObjectArray(ResetArray);
 	}
 	
@@ -1052,7 +1046,7 @@ TSharedRef<ITableRow> SPropertyTreeViewImpl::CreatePropertyEditor( TSharedPtr<FP
 {
 	FCategoryPropertyNode* CategoryNode = InPropertyNode->AsCategoryNode();
 
-	if (CategoryNode != NULL)
+	if (CategoryNode != nullptr)
 	{
 		// This is a category node; it does not need columns.
 		// Just use a simple setup.
diff --git a/Engine/Source/Editor/PropertyEditor/Private/SStructureDetailsView.cpp b/Engine/Source/Editor/PropertyEditor/Private/SStructureDetailsView.cpp
--- a/Engine/Source/Editor/PropertyEditor/Private/SStructureDetailsView.cpp
+++ b/Engine/Source/Editor/PropertyEditor/Private/SStructureDetailsView.cpp
@@ -40,7 +40,9 @@ void SStructureDetailsView::Construct(const FArguments& InArgs)
 	CustomName = InArgs._CustomName;
 
 	// Create the root property now
-	RootNode = MakeShareable(new FStructurePropertyNode);
+	// Only one root node in a structure details view
+	RootNodes.Empty(1);
+	RootNodes.Add(MakeShareable(new FStructurePropertyNode));
 		
 	PropertyUtilities = MakeShareable( new FPropertyDetailsUtilities( *this ) );
 	
@@ -178,15 +180,20 @@ void SStructureDetailsView::Construct(const FArguments& InArgs)
 
 void SStructureDetailsView::SetStructureData(TSharedPtr<FStructOnScope> InStructData)
 {
+	TSharedPtr<FComplexPropertyNode> RootNode = GetRootNode();
 	//PRE SET
-	SaveExpandedItems( RootNode.ToSharedRef() );
-	RootNode->SetStructure(NULL);
+	SaveExpandedItems(RootNode.ToSharedRef() );
+	RootNode->AsStructureNode()->SetStructure(nullptr);
 	RootNodesPendingKill.Add(RootNode);
+
+	RootNodes.Empty(1);
+
 	RootNode = MakeShareable(new FStructurePropertyNode);
+	RootNodes.Add(RootNode);
 
 	//SET
 	StructData = InStructData;
-	RootNode->SetStructure(StructData);
+	RootNode->AsStructureNode()->SetStructure(StructData);
 	if (!StructData.IsValid())
 	{
 		bIsLocked = false;
@@ -210,7 +217,7 @@ void SStructureDetailsView::SetStructureData(TSharedPtr<FStructOnScope> InStruct
 
 	RestoreExpandedItems(RootNode.ToSharedRef());
 
-	UpdatePropertyMap();
+	UpdatePropertyMaps();
 
 	UpdateFilteredDetails();
 }
@@ -220,6 +227,14 @@ void SStructureDetailsView::ForceRefresh()
 	SetStructureData(StructData);
 }
 
+void SStructureDetailsView::ClearSearch()
+{
+	CurrentFilter.FilterStrings.Empty();
+	SearchBox->SetText(FText::GetEmpty());
+	RerunCurrentFilter();
+}
+
+
 const TArray< TWeakObjectPtr<UObject> >& SStructureDetailsView::GetSelectedObjects() const
 {
 	static const TArray< TWeakObjectPtr<UObject> > DummyRef;
@@ -240,22 +255,34 @@ const FSelectedActorInfo& SStructureDetailsView::GetSelectedActorInfo() const
 
 bool SStructureDetailsView::IsConnected() const
 {
-	return StructData.IsValid() && StructData->IsValid() && RootNode.IsValid() && RootNode->IsValid();
+	const FStructurePropertyNode* RootNode = GetRootNode().IsValid() ? GetRootNode()->AsStructureNode() : nullptr;
+	return StructData.IsValid() && StructData->IsValid() && RootNode && RootNode->HasValidStructData();
+}
+
+FRootPropertyNodeList& SStructureDetailsView::GetRootNodes()
+{
+	return RootNodes;
 }
 
 TSharedPtr<class FComplexPropertyNode> SStructureDetailsView::GetRootNode()
 {
-	return RootNode;
+	return RootNodes[0];
+}
+
+const TSharedPtr<class FComplexPropertyNode> SStructureDetailsView::GetRootNode() const
+{
+	return RootNodes[0];
 }
 
-void SStructureDetailsView::CustomUpdatePropertyMap()
+void SStructureDetailsView::CustomUpdatePropertyMap(TSharedPtr<FDetailLayoutBuilderImpl>& InDetailLayout)
 {
-	DetailLayout->DefaultCategory(NAME_None).SetDisplayName(NAME_None, CustomName);
+	InDetailLayout->DefaultCategory(NAME_None).SetDisplayName(NAME_None, CustomName);
 }
 
 EVisibility SStructureDetailsView::GetPropertyEditingVisibility() const
 {
-	return StructData.IsValid() && StructData->IsValid() && RootNode.IsValid() && RootNode->IsValid() ? EVisibility::Visible : EVisibility::Collapsed;
+	const FStructurePropertyNode* RootNode = GetRootNode().IsValid() ? GetRootNode()->AsStructureNode() : nullptr;
+	return StructData.IsValid() && StructData->IsValid() && RootNode && RootNode->HasValidStructData() ? EVisibility::Visible : EVisibility::Collapsed;
 }
 
 void SStructureDetailsView::RegisterInstancedCustomPropertyLayout(UStruct* Class, FOnGetDetailCustomizationInstance DetailLayoutDelegate)
diff --git a/Engine/Source/Editor/PropertyEditor/Private/SStructureDetailsView.h b/Engine/Source/Editor/PropertyEditor/Private/SStructureDetailsView.h
--- a/Engine/Source/Editor/PropertyEditor/Private/SStructureDetailsView.h
+++ b/Engine/Source/Editor/PropertyEditor/Private/SStructureDetailsView.h
@@ -4,11 +4,10 @@
 #include "SDetailsViewBase.h"
 #include "IStructureDetailsView.h"
 
+PRAGMA_DISABLE_DEPRECATION_WARNINGS
 class SStructureDetailsView : public SDetailsViewBase, public IStructureDetailsView
 {
-	TSharedPtr<class FStructOnScope> StructData;
-	TSharedPtr<class FStructurePropertyNode> RootNode;
-	FText CustomName;
+
 public:
 
 	SLATE_BEGIN_ARGS(SStructureDetailsView)
@@ -36,6 +35,11 @@ public:
 		return true;
 	}
 
+	virtual bool ContainsMultipleTopLevelObjects() const override
+	{
+		return false;
+	}
+
 public:
 
 	// IStructureDetailsView interface
@@ -59,19 +63,21 @@ public:
 
 public:
 
-	// IDetailsViewPrivate interface
+	// IDetailsView interface
 
+	/** This method is deprecated.  When it is removed, also remove the PRAGMA_DISABLE_DEPRECATION_WARNINGS/PRAGMA_ENABLE_DEPRECATION_WARNINGS at the top and bottom of this file */
 	virtual const UClass* GetBaseClass() const override
 	{
 		return NULL;
 	}
 
+	/** This method is deprecated.  When it is removed, also remove the PRAGMA_DISABLE_DEPRECATION_WARNINGS/PRAGMA_ENABLE_DEPRECATION_WARNINGS at the top and bottom of this file */
 	virtual UClass* GetBaseClass() override
 	{
 		return NULL;
 	}
 
-	virtual bool IsCategoryHiddenByClass(FName CategoryName) const override
+	virtual bool IsCategoryHiddenByClass(const TSharedPtr<FComplexPropertyNode>& InRootNode, FName CategoryName) const override
 	{
 		return false;
 	}
@@ -79,16 +85,10 @@ public:
 	virtual void ForceRefresh() override;
 	virtual void MoveScrollOffset(int32 DeltaOffset) override {}
 	virtual void AddExternalRootPropertyNode(TSharedRef<FPropertyNode> ExternalRootNode) override {}
-
+	virtual void ClearSearch() override;
 public:
 
 	// IDetailsView interface
-
-	virtual UStruct* GetBaseStruct() const override
-	{
-		return GetBaseScriptStruct();
-	}
-
 	virtual const TArray< TWeakObjectPtr<UObject> >& GetSelectedObjects() const override;
 	virtual const TArray< TWeakObjectPtr<AActor> >& GetSelectedActors() const override;
 	virtual const struct FSelectedActorInfo& GetSelectedActorInfo() const override;
@@ -106,12 +106,23 @@ public:
 	virtual void SetObject(UObject* InObject, bool bForceRefresh = false) override{}
 	virtual void RemoveInvalidObjects() override {}
 	virtual void SetObjectPackageOverrides(const TMap<TWeakObjectPtr<UObject>, TWeakObjectPtr<UPackage>>& InMapping) override {}
+	virtual void SetRootObjectCustomizationInstance(TSharedPtr<IDetailRootObjectCustomization> InRootObjectCustomization) override {}
+	virtual TSharedPtr<class IDetailRootObjectCustomization> GetRootObjectCustomization() const override { return nullptr; }
 
-	virtual TSharedPtr<class FComplexPropertyNode> GetRootNode() override;
+	/* This is required by the base class but there is only ever one root node in a structure details view */
+	virtual FRootPropertyNodeList& GetRootNodes() override;
 
+	TSharedPtr<class FComplexPropertyNode> GetRootNode();
+	const TSharedPtr<class FComplexPropertyNode> GetRootNode() const;
 protected:
 
-	virtual void CustomUpdatePropertyMap() override;
+	virtual void CustomUpdatePropertyMap(TSharedPtr<FDetailLayoutBuilderImpl>& InDetailLayout) override;
 
 	EVisibility GetPropertyEditingVisibility() const;
+
+private:
+	TSharedPtr<class FStructOnScope> StructData;
+	FRootPropertyNodeList RootNodes;
+	FText CustomName;
 };
+PRAGMA_ENABLE_DEPRECATION_WARNINGS
diff --git a/Engine/Source/Editor/PropertyEditor/Private/StructurePropertyNode.h b/Engine/Source/Editor/PropertyEditor/Private/StructurePropertyNode.h
--- a/Engine/Source/Editor/PropertyEditor/Private/StructurePropertyNode.h
+++ b/Engine/Source/Editor/PropertyEditor/Private/StructurePropertyNode.h
@@ -12,11 +12,14 @@ public:
 	FStructurePropertyNode() : FComplexPropertyNode() {}
 	virtual ~FStructurePropertyNode() {}
 
+	virtual FStructurePropertyNode* AsStructureNode() override { return this; }
+	virtual const FStructurePropertyNode* AsStructureNode() const override { return this; }
+
 	/** FPropertyNode Interface */
 	virtual uint8* GetValueBaseAddress(uint8* Base) override
 	{
 		ensure(true);
-		return IsValid() ? StructData->GetStructMemory() : NULL;
+		return HasValidStructData() ? StructData->GetStructMemory() : NULL;
 	}
 
 	void SetStructure(TSharedPtr<FStructOnScope> InStructData)
@@ -26,14 +29,14 @@ public:
 		StructData = InStructData;
 	}
 
-	bool IsValid() const
+	bool HasValidStructData() const
 	{
 		return StructData.IsValid() && StructData->IsValid();
 	}
 
 	bool GetReadAddressUncached(FPropertyNode& InPropertyNode, FReadAddressListData& OutAddresses) const override
 	{
-		if (!IsValid())
+		if (!HasValidStructData())
 		{
 			return false;
 		}
@@ -52,37 +55,45 @@ public:
 
 	bool GetReadAddressUncached(FPropertyNode& InPropertyNode,
 		bool InRequiresSingleSelection,
-		FReadAddressListData& OutAddresses,
+		FReadAddressListData* OutAddresses,
 		bool bComparePropertyContents,
 		bool bObjectForceCompare,
 		bool bArrayPropertiesCanDifferInSize) const override
 	{
-		return GetReadAddressUncached(InPropertyNode, OutAddresses);
+		if(OutAddresses)
+		{
+			return GetReadAddressUncached(InPropertyNode, *OutAddresses);
+		}
+		else
+		{
+			FReadAddressListData Unused;
+			return GetReadAddressUncached(InPropertyNode, Unused);
+		}
 	}
 
 	UPackage* GetOwnerPackage() const
 	{
-		return IsValid() ? StructData->GetPackage() : nullptr;
+		return HasValidStructData() ? StructData->GetPackage() : nullptr;
 	}
 
 	/** FComplexPropertyNode Interface */
 	virtual const UStruct* GetBaseStructure() const override
 	{ 
-		return IsValid() ? StructData->GetStruct() : NULL; 
+		return HasValidStructData() ? StructData->GetStruct() : NULL;
 	}
 	virtual UStruct* GetBaseStructure() override
 	{
-		const UStruct* Struct = IsValid() ? StructData->GetStruct() : NULL;
+		const UStruct* Struct = HasValidStructData() ? StructData->GetStruct() : NULL;
 		return const_cast<UStruct*>(Struct);
 	}
 	virtual int32 GetInstancesNum() const override
 	{ 
-		return IsValid() ? 1 : 0; 
+		return HasValidStructData() ? 1 : 0;
 	}
 	virtual uint8* GetMemoryOfInstance(int32 Index) override
 	{ 
 		check(0 == Index);
-		return IsValid() ? StructData->GetStructMemory() : NULL;
+		return HasValidStructData() ? StructData->GetStructMemory() : NULL;
 	}
 	virtual TWeakObjectPtr<UObject> GetInstanceAsUObject(int32 Index) override
 	{
diff --git a/Engine/Source/Editor/PropertyEditor/Public/DetailCategoryBuilder.h b/Engine/Source/Editor/PropertyEditor/Public/DetailCategoryBuilder.h
--- a/Engine/Source/Editor/PropertyEditor/Public/DetailCategoryBuilder.h
+++ b/Engine/Source/Editor/PropertyEditor/Public/DetailCategoryBuilder.h
@@ -63,7 +63,7 @@ public:
 	 * @param Location				The location within the category where the property is shown
 	 * See IDetailCategoryBuilder::GetProperty for clarification of parameters
 	 */
-	virtual IDetailPropertyRow& AddProperty( FName PropertyPath, UClass* ClassOutermost = NULL, FName InstanceName = NAME_None, EPropertyLocation::Type Location = EPropertyLocation::Default ) = 0;
+	virtual IDetailPropertyRow& AddProperty( FName PropertyPath, UClass* ClassOutermost = nullptr, FName InstanceName = NAME_None, EPropertyLocation::Type Location = EPropertyLocation::Default ) = 0;
 
 	/**
 	 * Adds a property, shown in the default way to the category 
diff --git a/Engine/Source/Editor/PropertyEditor/Public/DetailLayoutBuilder.h b/Engine/Source/Editor/PropertyEditor/Public/DetailLayoutBuilder.h
--- a/Engine/Source/Editor/PropertyEditor/Public/DetailLayoutBuilder.h
+++ b/Engine/Source/Editor/PropertyEditor/Public/DetailLayoutBuilder.h
@@ -51,6 +51,11 @@ public:
 	 */
 	virtual const class IDetailsView& GetDetailsView() const = 0;
 
+	/**
+	 * @return The base class of the objects being customized in this detail layout
+	 */
+	virtual UClass* GetBaseClass() const = 0;
+
 	/**
 	 * Gets the current object(s) being customized by this builder
 	 *
@@ -107,6 +112,14 @@ public:
 	 */
 	virtual TSharedRef<IPropertyHandle> GetProperty( const FName PropertyPath, const UClass* ClassOutermost = NULL, FName InstanceName = NAME_None )  = 0;
 
+	/**
+	 * Gets the top level property, for showing the warning for experimental or early access class
+	 * 
+	 * @return the top level property name
+	 */
+
+	virtual FName GetTopLevelProperty() = 0;
+
 	/**
 	 * Hides a property from view 
 	 *
diff --git a/Engine/Source/Editor/PropertyEditor/Public/IDetailRootObjectCustomization.h b/Engine/Source/Editor/PropertyEditor/Public/IDetailRootObjectCustomization.h
new file mode 100644
--- /dev/null
+++ b/Engine/Source/Editor/PropertyEditor/Public/IDetailRootObjectCustomization.h
@@ -0,0 +1,31 @@
+// Copyright 1998-2016 Epic Games, Inc. All Rights Reserved.
+#pragma once
+
+
+/** 
+ * Interface for any class that lays out details for a specific class
+ */
+class IDetailRootObjectCustomization : public TSharedFromThis<IDetailRootObjectCustomization>
+{
+public:
+	virtual ~IDetailRootObjectCustomization() {}
+
+	/** 
+	 * Called when the details panel wants to display an object header widget for a given object
+	 *
+	 * @param	InRootObject	The object whose header is being customized
+ 	 */
+	virtual TSharedPtr<SWidget> CustomizeObjectHeader(const UObject* InRootObject) = 0;
+
+	/**
+	 * Whether or not the object and all of its children should be visible in the details panel
+	 */
+	virtual bool IsObjectVisible(const UObject* InRootObject) const = 0;
+
+	/**
+	 * Whether or object should have a header displayed or just show the children directly
+	 *
+	 * @return true if the this customization should be displayed or false to just show the children directly
+	 */
+	virtual bool ShouldDisplayHeader(const UObject* InRootObject) const = 0;
+};
