Fixed bug where details panel customizations would not show up.  This was causing critical information to be missing.  Change duplicated from Dev-Editor

diff --git a/Engine/Source/Editor/PropertyEditor/Private/DetailLayoutBuilderImpl.cpp b/Engine/Source/Editor/PropertyEditor/Private/DetailLayoutBuilderImpl.cpp
--- a/Engine/Source/Editor/PropertyEditor/Private/DetailLayoutBuilderImpl.cpp
+++ b/Engine/Source/Editor/PropertyEditor/Private/DetailLayoutBuilderImpl.cpp
@@ -479,7 +479,35 @@ void FDetailLayoutBuilderImpl::GetObjectsBeingCustomized( TArray< TWeakObjectPtr
 	OutObjects.Empty();
 
 	FObjectPropertyNode* RootObjectNode = RootNode.IsValid() ? RootNode.Pin()->AsObjectNode() : nullptr;
-	if(RootObjectNode)
+	// The class to find properties in defaults to the class currently being customized
+	FName ClassName = CurrentCustomizationClass ? CurrentCustomizationClass->GetFName() : NAME_None;
+
+	if(ClassName != NAME_None && CurrentCustomizationVariableName != NAME_None)
+	{
+		// If this fails there are no properties associated with the class name provided
+		FClassInstanceToPropertyMap* ClassInstanceToPropertyMapPtr = PropertyMap.Find(ClassName);
+
+		if(ClassInstanceToPropertyMapPtr)
+		{
+			FClassInstanceToPropertyMap& ClassInstanceToPropertyMap = *ClassInstanceToPropertyMapPtr;
+
+			FPropertyNodeMap* PropertyNodeMapPtr = ClassInstanceToPropertyMap.Find(CurrentCustomizationVariableName);
+
+			if(PropertyNodeMapPtr)
+			{
+				FPropertyNodeMap& PropertyNodeMap = *PropertyNodeMapPtr;
+				FObjectPropertyNode* ParentObjectProperty = PropertyNodeMap.ParentProperty ? PropertyNodeMap.ParentProperty->AsObjectNode() : NULL;
+				if(ParentObjectProperty)
+				{
+					for(int32 ObjectIndex = 0; ObjectIndex < ParentObjectProperty->GetNumObjects(); ++ObjectIndex)
+					{
+						OutObjects.Add(ParentObjectProperty->GetUObject(ObjectIndex));
+					}
+				}
+			}
+		}
+	}
+	else if(RootObjectNode)
 	{
 		for (int32 ObjectIndex = 0; ObjectIndex < RootObjectNode->GetNumObjects(); ++ObjectIndex)
 		{
