Additional framework IWYU non-unity fixes

diff --git a/Engine/Source/Editor/AudioEditor/Classes/Factories/SoundSubmixFactory.h b/Engine/Source/Editor/AudioEditor/Classes/Factories/SoundSubmixFactory.h
--- a/Engine/Source/Editor/AudioEditor/Classes/Factories/SoundSubmixFactory.h
+++ b/Engine/Source/Editor/AudioEditor/Classes/Factories/SoundSubmixFactory.h
@@ -5,6 +5,7 @@
 //~=============================================================================
 
 #pragma once
+#include "Factories/Factory.h"
 #include "SoundSubmixFactory.generated.h"
 
 UCLASS(MinimalAPI, hidecategories=Object)
diff --git a/Engine/Source/Editor/AudioEditor/Private/SoundSubmixGraph.cpp b/Engine/Source/Editor/AudioEditor/Private/SoundSubmixGraph.cpp
--- a/Engine/Source/Editor/AudioEditor/Private/SoundSubmixGraph.cpp
+++ b/Engine/Source/Editor/AudioEditor/Private/SoundSubmixGraph.cpp
@@ -7,6 +7,7 @@
 #include "Sound/SoundSubmix.h"
 #include "SoundSubmixGraph/SoundSubmixGraphNode.h"
 #include "GraphEditor.h"
+#include "UObject/Package.h"
 
 class FSoundSubmixAudioEditor : public ISoundSubmixAudioEditor
 {
diff --git a/Engine/Source/Editor/AudioEditor/Public/AudioEditorModule.h b/Engine/Source/Editor/AudioEditor/Public/AudioEditorModule.h
--- a/Engine/Source/Editor/AudioEditor/Public/AudioEditorModule.h
+++ b/Engine/Source/Editor/AudioEditor/Public/AudioEditorModule.h
@@ -12,6 +12,7 @@
 class UDialogueWave;
 class USoundClass;
 class USoundCue;
+class USoundSubmix;
 class USoundNode;
 class USoundWave;
 struct FDialogueContextMapping;
diff --git a/Engine/Source/Editor/ComponentVisualizers/Public/AttenuatedComponentVisualizer.h b/Engine/Source/Editor/ComponentVisualizers/Public/AttenuatedComponentVisualizer.h
--- a/Engine/Source/Editor/ComponentVisualizers/Public/AttenuatedComponentVisualizer.h
+++ b/Engine/Source/Editor/ComponentVisualizers/Public/AttenuatedComponentVisualizer.h
@@ -3,8 +3,7 @@
 #pragma once
 
 #include "ComponentVisualizer.h"
-
-struct FEngineShowFlags;
+#include "ShowFlags.h"
 
 template<class T>
 class COMPONENTVISUALIZERS_API TAttenuatedComponentVisualizer : public FComponentVisualizer
diff --git a/Engine/Source/Editor/PropertyEditor/Private/UserInterface/PropertyEditor/SPropertyEditorEditInline.cpp b/Engine/Source/Editor/PropertyEditor/Private/UserInterface/PropertyEditor/SPropertyEditorEditInline.cpp
--- a/Engine/Source/Editor/PropertyEditor/Private/UserInterface/PropertyEditor/SPropertyEditorEditInline.cpp
+++ b/Engine/Source/Editor/PropertyEditor/Private/UserInterface/PropertyEditor/SPropertyEditorEditInline.cpp
@@ -10,6 +10,7 @@
 #include "ClassViewerFilter.h"
 #include "Styling/SlateIconFinder.h"
 #include "ConstructorHelpers.h"
+#include "Editor/EditorEngine.h"
 
 class FPropertyEditorInlineClassFilter : public IClassViewerFilter
 {
