Copying //UE4/Orion-Staging to //UE4/Main (Source: //Orion/Dev-General @ 3483207)

Change 3483207 on 2017/06/09 by [REDACTED_PERSON]

	Batch Animation Compression fixes.
	- Fixed incorrect 'MemorySavingsFromPrevious' resulting in picking suboptimal compressors.
	- Fixed uncompressed size calculation not taking into account scale component.
	- Fixed animations with 'bDoNotOverrideCompression' causing crashes because they were not recompressed.
	- Animation with 'bDoNotOverrideCompression' that use the automatic compressions are not skipped by the automatic batch compression.
	- Added 'CompressCommandletVersion' to DDC key, so we can force recompression on all animations easily.

	Repopulated DDC with all animations.

	#!codereview martin.wilson
	#!rb lina.halper
	#!tests loaded editor, ran a quick game.

Change 3483107 on 2017/06/09 by robomerge

	#!ROBOMERGE-AUTHOR: andrew.grant
	Non-shipping changes -
	 Added GPU health check if we are waiting for > 2 secs on the rendering thread
	 Changed param for GPU health checking from aftermath to gpucrashdebugging

	#!tests compiled
	#!rb arne

	#!ROBOMERGE-SOURCE: CL 3483100 in //Orion/Release-40.4/... via CL 3483101 via CL 3483103
	#!ROBOMERGE-BOT: ORION (Main -> Dev-Cinematics)

Change 3483106 on 2017/06/09 by robomerge

	#!ROBOMERGE-AUTHOR: andrew.grant
	Non-shipping changes -
	 Added GPU health check if we are waiting for > 2 secs on the rendering thread
	 Changed param for GPU health checking from aftermath to gpucrashdebugging

	#!tests compiled
	#!rb arne

	#!ROBOMERGE-SOURCE: CL 3483100 in //Orion/Release-40.4/... via CL 3483101 via CL 3483103
	#!ROBOMERGE-BOT: ORION (Main -> Dev-Balance)

Change 3483105 on 2017/06/09 by robomerge

	#!ROBOMERGE-AUTHOR: andrew.grant
	Non-shipping changes -
	 Added GPU health check if we are waiting for > 2 secs on the rendering thread
	 Changed param for GPU health checking from aftermath to gpucrashdebugging

	#!tests compiled
	#!rb arne

	#!ROBOMERGE-SOURCE: CL 3483100 in //Orion/Release-40.4/... via CL 3483101 via CL 3483103
	#!ROBOMERGE-BOT: ORION (Main -> Dev-REGS)

Change 3483104 on 2017/06/09 by robomerge

	#!ROBOMERGE-AUTHOR: andrew.grant
	Non-shipping changes -
	 Added GPU health check if we are waiting for > 2 secs on the rendering thread
	 Changed param for GPU health checking from aftermath to gpucrashdebugging

	#!tests compiled
	#!rb arne

	#!ROBOMERGE-SOURCE: CL 3483100 in //Orion/Release-40.4/... via CL 3483101 via CL 3483103
	#!ROBOMERGE-BOT: ORION (Main -> Dev-General)

Change 3483103 on 2017/06/09 by robomerge

	#!ROBOMERGE-AUTHOR: andrew.grant
	Non-shipping changes -
	 Added GPU health check if we are waiting for > 2 secs on the rendering thread
	 Changed param for GPU health checking from aftermath to gpucrashdebugging

	#!tests compiled
	#!rb arne

	#!ROBOMERGE-SOURCE: CL 3483100 in //Orion/Release-40.4/... via CL 3483101
	#!ROBOMERGE-BOT: ORION (Release-40.5 -> Main)

Change 3483101 on 2017/06/09 by robomerge

	#!ROBOMERGE-AUTHOR: andrew.grant
	Non-shipping changes -
	 Added GPU health check if we are waiting for > 2 secs on the rendering thread
	 Changed param for GPU health checking from aftermath to gpucrashdebugging

	#!tests compiled
	#!rb arne

	#!ROBOMERGE-SOURCE: CL 3483100 in //Orion/Release-40.4/...
	#!ROBOMERGE-BOT: ORION (Release-40.4 -> Release-40.5)

Change 3483100 on 2017/06/09 by [REDACTED_PERSON]

	Non-shipping changes -
	 Added GPU health check if we are waiting for > 2 secs on the rendering thread
	 Changed param for GPU health checking from aftermath to gpucrashdebugging

	#!tests compiled
	#!rb arne

Change 3482985 on 2017/06/09 by robomerge

	#!ROBOMERGE-AUTHOR: daniel.lamb
	Fixed up the allocated small pool memory stat.
	#!rb Andrew.Grant
	#!test Paragon startup
	#!lockdown Andrew.Grant

	#!ROBOMERGE-SOURCE: CL 3482444 in //Orion/Release-40.4/... via CL 3482448 via CL 3482449
	#!ROBOMERGE-BOT: ORION (Main -> Dev-Cinematics)

Change 3482984 on 2017/06/09 by robomerge

	#!ROBOMERGE-AUTHOR: daniel.lamb
	Fixed up the allocated small pool memory stat.
	#!rb Andrew.Grant
	#!test Paragon startup
	#!lockdown Andrew.Grant

	#!ROBOMERGE-SOURCE: CL 3482444 in //Orion/Release-40.4/... via CL 3482448 via CL 3482449
	#!ROBOMERGE-BOT: ORION (Main -> Dev-Balance)

Change 3482983 on 2017/06/09 by robomerge

	#!ROBOMERGE-AUTHOR: daniel.lamb
	Fixed up the allocated small pool memory stat.
	#!rb Andrew.Grant
	#!test Paragon startup
	#!lockdown Andrew.Grant

	#!ROBOMERGE-SOURCE: CL 3482444 in //Orion/Release-40.4/... via CL 3482448 via CL 3482449
	#!ROBOMERGE-BOT: ORION (Main -> Dev-REGS)

Change 3482982 on 2017/06/09 by robomerge

	#!ROBOMERGE-AUTHOR: daniel.lamb
	Fixed up the allocated small pool memory stat.
	#!rb Andrew.Grant
	#!test Paragon startup
	#!lockdown Andrew.Grant

	#!ROBOMERGE-SOURCE: CL 3482444 in //Orion/Release-40.4/... via CL 3482448 via CL 3482449
	#!ROBOMERGE-BOT: ORION (Main -> Dev-UI)

Change 3482981 on 2017/06/09 by robomerge

	#!ROBOMERGE-AUTHOR: daniel.lamb
	Fixed up the allocated small pool memory stat.
	#!rb Andrew.Grant
	#!test Paragon startup
	#!lockdown Andrew.Grant

	#!ROBOMERGE-SOURCE: CL 3482444 in //Orion/Release-40.4/... via CL 3482448 via CL 3482449
	#!ROBOMERGE-BOT: ORION (Main -> Dev-General)

Change 3482612 on 2017/06/09 by [REDACTED_PERSON]

	Niagara - Fix various wiring issues.
	+ Reverting dynamic inputs no longer leaves the graph disconnected.
	+ Reverting dynamic inputs no longer leaves the controls in the stack.
	+ Adding multiple dynamic inputs to the same module now wires them correctly.
	+ Adding dynamic inputs when there is already an override read now wires correctly.
	+ Moving modules with dynamic inputs up and down and removing them now works correctly.

	#!tests Everything above.
	#!rb none

	#!codereview Olaf.Piesche,Simon.Tovey,Shaun.Kime

Change 3482449 on 2017/06/09 by robomerge

	#!ROBOMERGE-AUTHOR: daniel.lamb
	Fixed up the allocated small pool memory stat.
	#!rb Andrew.Grant
	#!test Paragon startup
	#!lockdown Andrew.Grant

	#!ROBOMERGE-SOURCE: CL 3482444 in //Orion/Release-40.4/... via CL 3482448
	#!ROBOMERGE-BOT: ORION (Release-40.5 -> Main)

Change 3482448 on 2017/06/09 by robomerge

	#!ROBOMERGE-AUTHOR: daniel.lamb
	Fixed up the allocated small pool memory stat.
	#!rb Andrew.Grant
	#!test Paragon startup
	#!lockdown Andrew.Grant

	#!ROBOMERGE-SOURCE: CL 3482444 in //Orion/Release-40.4/...
	#!ROBOMERGE-BOT: ORION (Release-40.4 -> Release-40.5)

Change 3482444 on 2017/06/09 by [REDACTED_PERSON]

	Fixed up the allocated small pool memory stat.
	#!rb Andrew.Grant
	#!test Paragon startup
	#!lockdown Andrew.Grant

Change 3482261 on 2017/06/09 by [REDACTED_PERSON]

	Made Get/Set nodes available at all times.
	Tweaked the right-click menu on parameter map base to allow for particle namespaced custom variables and also limiting based on script context.

	#!rb none
	#!tests n/a

Change 3482147 on 2017/06/09 by [REDACTED_PERSON]

	Fixing crash when updating the vertex data and the vertex attributes are no longer part of the data set.

	#!rb none
	#!tests opened existing files

Change 3482076 on 2017/06/09 by [REDACTED_PERSON]

	Resave to prevent the constant recompiling of DefaultParticle

Change 3481302 on 2017/06/08 by [REDACTED_PERSON]

	Adding a FunctionCall derived node type that allows you to set any namespaced pin by name and type.

	#!rb none
	#!tests created emitter with values in spawn and update
	#!codereview frank.fella

Change 3480830 on 2017/06/08 by [REDACTED_PERSON]

	First batch of recompressed animations.

	#!codereview jay.hosfelt, dwayne.martin
	#!lockdown Andrew.Bains

Change 3480524 on 2017/06/08 by [REDACTED_PERSON]

	Fixed CompressAnimations Commandlet to work with new DDC refactor.

	#!codereview martin.wilson
	#!rb lina.halper
	#!tests Paragon full animation recompression.

Change 3480278 on 2017/06/08 by robomerge

	#!ROBOMERGE-AUTHOR: andrew.grant
	Additional logging for OR-38938
	#!rb Ryan.Gerleve
	#!tests compiled

	#!ROBOMERGE-SOURCE: CL 3479906 in //Orion/Release-40.4/... via CL 3479909 via CL 3479910
	#!ROBOMERGE-BOT: ORION (Main -> Dev-Cinematics)

Change 3480277 on 2017/06/08 by robomerge

	#!ROBOMERGE-AUTHOR: andrew.grant
	Additional logging for OR-38938
	#!rb Ryan.Gerleve
	#!tests compiled

	#!ROBOMERGE-SOURCE: CL 3479906 in //Orion/Release-40.4/... via CL 3479909 via CL 3479910
	#!ROBOMERGE-BOT: ORION (Main -> Dev-Balance)

Change 3480276 on 2017/06/08 by robomerge

	#!ROBOMERGE-AUTHOR: andrew.grant
	Additional logging for OR-38938
	#!rb Ryan.Gerleve
	#!tests compiled

	#!ROBOMERGE-SOURCE: CL 3479906 in //Orion/Release-40.4/... via CL 3479909 via CL 3479910
	#!ROBOMERGE-BOT: ORION (Main -> Dev-REGS)

Change 3480273 on 2017/06/08 by robomerge

	#!ROBOMERGE-AUTHOR: andrew.grant
	Additional logging for OR-38938
	#!rb Ryan.Gerleve
	#!tests compiled

	#!ROBOMERGE-SOURCE: CL 3479906 in //Orion/Release-40.4/... via CL 3479909 via CL 3479910
	#!ROBOMERGE-BOT: ORION (Main -> Dev-UI)

Change 3480270 on 2017/06/08 by robomerge

	#!ROBOMERGE-AUTHOR: andrew.grant
	Additional logging for OR-38938
	#!rb Ryan.Gerleve
	#!tests compiled

	#!ROBOMERGE-SOURCE: CL 3479906 in //Orion/Release-40.4/... via CL 3479909 via CL 3479910
	#!ROBOMERGE-BOT: ORION (Main -> Dev-General)

Change 3480090 on 2017/06/08 by robomerge

	#!ROBOMERGE-AUTHOR: andrew.grant
	[NULL MERGE]
	Locked 40.3 builds to 3472726
	#!tests #!rb none

	#!ROBOMERGE-SOURCE: CL 3479202 in //Orion/Release-40.3/... via CL 3479203 via CL 3479204 via CL 3479205
	#!ROBOMERGE-BOT: ORION (Main -> Dev-Cinematics)

Change 3480089 on 2017/06/08 by robomerge

	#!ROBOMERGE-AUTHOR: andrew.grant
	[NULL MERGE]
	Locked 40.3 builds to 3472726
	#!tests #!rb none

	#!ROBOMERGE-SOURCE: CL 3479202 in //Orion/Release-40.3/... via CL 3479203 via CL 3479204 via CL 3479205
	#!ROBOMERGE-BOT: ORION (Main -> Dev-Balance)

Change 3480088 on 2017/06/08 by robomerge

	#!ROBOMERGE-AUTHOR: andrew.grant
	[NULL MERGE]
	Locked 40.3 builds to 3472726
	#!tests #!rb none

	#!ROBOMERGE-SOURCE: CL 3479202 in //Orion/Release-40.3/... via CL 3479203 via CL 3479204 via CL 3479205
	#!ROBOMERGE-BOT: ORION (Main -> Dev-REGS)

Change 3480087 on 2017/06/08 by robomerge

	#!ROBOMERGE-AUTHOR: andrew.grant
	[NULL MERGE]
	Locked 40.3 builds to 3472726
	#!tests #!rb none

	#!ROBOMERGE-SOURCE: CL 3479202 in //Orion/Release-40.3/... via CL 3479203 via CL 3479204 via CL 3479205
	#!ROBOMERGE-BOT: ORION (Main -> Dev-UI)

Change 3480086 on 2017/06/08 by robomerge

	#!ROBOMERGE-AUTHOR: andrew.grant
	[NULL MERGE]
	Locked 40.3 builds to 3472726
	#!tests #!rb none

	#!ROBOMERGE-SOURCE: CL 3479202 in //Orion/Release-40.3/... via CL 3479203 via CL 3479204 via CL 3479205
	#!ROBOMERGE-BOT: ORION (Main -> Dev-General)

Change 3480085 on 2017/06/08 by robomerge

	#!ROBOMERGE-AUTHOR: daniel.lamb
	Added stats to MallocBinned2.
	#!rb Andrew.Grant
	#!test Paragon PS4
	#!lockdown Andrew.Grant

	#!ROBOMERGE-SOURCE: CL 3479159 in //Orion/Release-40.4/... via CL 3479160 via CL 3479161
	#!ROBOMERGE-BOT: ORION (Main -> Dev-Cinematics)

Change 3480084 on 2017/06/08 by robomerge

	#!ROBOMERGE-AUTHOR: daniel.lamb
	Added stats to MallocBinned2.
	#!rb Andrew.Grant
	#!test Paragon PS4
	#!lockdown Andrew.Grant

	#!ROBOMERGE-SOURCE: CL 3479159 in //Orion/Release-40.4/... via CL 3479160 via CL 3479161
	#!ROBOMERGE-BOT: ORION (Main -> Dev-Balance)

Change 3480083 on 2017/06/08 by robomerge

	#!ROBOMERGE-AUTHOR: daniel.lamb
	Added stats to MallocBinned2.
	#!rb Andrew.Grant
	#!test Paragon PS4
	#!lockdown Andrew.Grant

	#!ROBOMERGE-SOURCE: CL 3479159 in //Orion/Release-40.4/... via CL 3479160 via CL 3479161
	#!ROBOMERGE-BOT: ORION (Main -> Dev-REGS)

Change 3480082 on 2017/06/08 by robomerge

	#!ROBOMERGE-AUTHOR: daniel.lamb
	Added stats to MallocBinned2.
	#!rb Andrew.Grant
	#!test Paragon PS4
	#!lockdown Andrew.Grant

	#!ROBOMERGE-SOURCE: CL 3479159 in //Orion/Release-40.4/... via CL 3479160 via CL 3479161
	#!ROBOMERGE-BOT: ORION (Main -> Dev-UI)

Change 3480081 on 2017/06/08 by robomerge

	#!ROBOMERGE-AUTHOR: daniel.lamb
	Added stats to MallocBinned2.
	#!rb Andrew.Grant
	#!test Paragon PS4
	#!lockdown Andrew.Grant

	#!ROBOMERGE-SOURCE: CL 3479159 in //Orion/Release-40.4/... via CL 3479160 via CL 3479161
	#!ROBOMERGE-BOT: ORION (Main -> Dev-General)

Change 3480073 on 2017/06/08 by robomerge

	#!ROBOMERGE-AUTHOR: jeff.williams
	Removing implicit requirements to display Badges

	Badge requirements are not considered while culling nodes from the build graph. This allowed implicit dependencies resolved before culling to invalidate badges afterwards. Only explicitly declared dependencies are now used to validate badges.

	#!rb none
	#!tests compile, validated export output

	#!ROBOMERGE-SOURCE: CL 3479012 in //Orion/Main/...
	#!ROBOMERGE-BOT: ORION (Main -> Dev-Cinematics)

Change 3480072 on 2017/06/08 by robomerge

	#!ROBOMERGE-AUTHOR: jeff.williams
	Removing implicit requirements to display Badges

	Badge requirements are not considered while culling nodes from the build graph. This allowed implicit dependencies resolved before culling to invalidate badges afterwards. Only explicitly declared dependencies are now used to validate badges.

	#!rb none
	#!tests compile, validated export output

	#!ROBOMERGE-SOURCE: CL 3479012 in //Orion/Main/...
	#!ROBOMERGE-BOT: ORION (Main -> Dev-Balance)

Change 3480071 on 2017/06/08 by robomerge

	#!ROBOMERGE-AUTHOR: jeff.williams
	Removing implicit requirements to display Badges

	Badge requirements are not considered while culling nodes from the build graph. This allowed implicit dependencies resolved before culling to invalidate badges afterwards. Only explicitly declared dependencies are now used to validate badges.

	#!rb none
	#!tests compile, validated export output

	#!ROBOMERGE-SOURCE: CL 3479012 in //Orion/Main/...
	#!ROBOMERGE-BOT: ORION (Main -> Dev-REGS)

Change 3480070 on 2017/06/08 by robomerge

	#!ROBOMERGE-AUTHOR: jeff.williams
	Removing implicit requirements to display Badges

	Badge requirements are not considered while culling nodes from the build graph. This allowed implicit dependencies resolved before culling to invalidate badges afterwards. Only explicitly declared dependencies are now used to validate badges.

	#!rb none
	#!tests compile, validated export output

	#!ROBOMERGE-SOURCE: CL 3479012 in //Orion/Main/...
	#!ROBOMERGE-BOT: ORION (Main -> Dev-UI)

Change 3480069 on 2017/06/08 by robomerge

	#!ROBOMERGE-AUTHOR: jeff.williams
	Removing implicit requirements to display Badges

	Badge requirements are not considered while culling nodes from the build graph. This allowed implicit dependencies resolved before culling to invalidate badges afterwards. Only explicitly declared dependencies are now used to validate badges.

	#!rb none
	#!tests compile, validated export output

	#!ROBOMERGE-SOURCE: CL 3479012 in //Orion/Main/...
	#!ROBOMERGE-BOT: ORION (Main -> Dev-General)

Change 3479910 on 2017/06/08 by robomerge

	#!ROBOMERGE-AUTHOR: andrew.grant
	Additional logging for OR-38938
	#!rb Ryan.Gerleve
	#!tests compiled

	#!ROBOMERGE-SOURCE: CL 3479906 in //Orion/Release-40.4/... via CL 3479909
	#!ROBOMERGE-BOT: ORION (Release-40.5 -> Main)

Change 3479909 on 2017/06/08 by robomerge

	#!ROBOMERGE-AUTHOR: andrew.grant
	Additional logging for OR-38938
	#!rb Ryan.Gerleve
	#!tests compiled

	#!ROBOMERGE-SOURCE: CL 3479906 in //Orion/Release-40.4/...
	#!ROBOMERGE-BOT: ORION (Release-40.4 -> Release-40.5)

Change 3479906 on 2017/06/08 by [REDACTED_PERSON]

	Additional logging for OR-38938
	#!rb Ryan.Gerleve
	#!tests compiled

Change 3479800 on 2017/06/08 by [REDACTED_PERSON]

	EditCondition UProperty metadata works on UStruct properties as well (including data table row structs)
	- Submitting on behalf of Jamie Dale (thanks Jamie!)

	#!rb Jamie.Dale
	#!tests EditCondition works for both UClass and UStruct properties

Change 3479765 on 2017/06/08 by [REDACTED_PERSON]

	Allow overriding of collections per component from BP and a functional test map for it.

	#!rb none
	#!tests test map works
	#!codereview Olaf.Piesche, Frank.Fella, Shaun.Kime

Change 3479205 on 2017/06/07 by robomerge

	#!ROBOMERGE-AUTHOR: andrew.grant
	[NULL MERGE]
	Locked 40.3 builds to 3472726
	#!tests #!rb none

	#!ROBOMERGE-SOURCE: CL 3479202 in //Orion/Release-40.3/... via CL 3479203 via CL 3479204
	#!ROBOMERGE-BOT: ORION (Release-40.5 -> Main)

Change 3479204 on 2017/06/07 by robomerge

	#!ROBOMERGE-AUTHOR: andrew.grant
	[NULL MERGE]
	Locked 40.3 builds to 3472726
	#!tests #!rb none

	#!ROBOMERGE-SOURCE: CL 3479202 in //Orion/Release-40.3/... via CL 3479203
	#!ROBOMERGE-BOT: ORION (Release-40.4 -> Release-40.5)

Change 3479203 on 2017/06/07 by robomerge

	#!ROBOMERGE-AUTHOR: andrew.grant
	[NULL MERGE]
	Locked 40.3 builds to 3472726
	#!tests #!rb none

	#!ROBOMERGE-SOURCE: CL 3479202 in //Orion/Release-40.3/...
	#!ROBOMERGE-BOT: ORION (Release-40.3 -> Release-40.4)

Change 3479202 on 2017/06/07 by [REDACTED_PERSON]

	Locked 40.3 builds to 3472726
	#!ROBOMERGE: !40.4
	#!tests #!rb none

Change 3479161 on 2017/06/07 by robomerge

	#!ROBOMERGE-AUTHOR: daniel.lamb
	Added stats to MallocBinned2.
	#!rb Andrew.Grant
	#!test Paragon PS4
	#!lockdown Andrew.Grant

	#!ROBOMERGE-SOURCE: CL 3479159 in //Orion/Release-40.4/... via CL 3479160
	#!ROBOMERGE-BOT: ORION (Release-40.5 -> Main)

Change 3479160 on 2017/06/07 by robomerge

	#!ROBOMERGE-AUTHOR: daniel.lamb
	Added stats to MallocBinned2.
	#!rb Andrew.Grant
	#!test Paragon PS4
	#!lockdown Andrew.Grant

	#!ROBOMERGE-SOURCE: CL 3479159 in //Orion/Release-40.4/...
	#!ROBOMERGE-BOT: ORION (Release-40.4 -> Release-40.5)

Change 3479159 on 2017/06/07 by [REDACTED_PERSON]

	Added stats to MallocBinned2.
	#!rb Andrew.Grant
	#!test Paragon PS4
	#!codereview Gil.Gribb
	#!lockdown Andrew.Grant

Change 3479012 on 2017/06/07 by [REDACTED_PERSON]

	Removing implicit requirements to display Badges

	Badge requirements are not considered while culling nodes from the build graph. This allowed implicit dependencies resolved before culling to invalidate badges afterwards. Only explicitly declared dependencies are now used to validate badges.

	#!rb none
	#!tests compile, validated export output

Change 3478991 on 2017/06/07 by [REDACTED_PERSON]

	Added auto-compile to emitters. It is an emitter-wide value, toggled by the dropdown next to the compile button.

	#!rb none
	#!tests made multiple edits to an emitter

Change 3478976 on 2017/06/07 by [REDACTED_PERSON]

	Sequencer: Fix burnin when there are warmup frames. The current time used for the burnin is offset from the playback range's start time. When using warmup frames, the start time will include the warmup time so it needs to be factored out when setting the actual current time for the frame.

	#!jira UE-45737
	#!rb none
	#!codereview andrew.rodham
	#!tests none

Change 3478426 on 2017/06/07 by [REDACTED_PERSON]

	Expose some ability system stuff to blueprints:
	-Query for AGE Handle based on GE Query
	-Methods for accessing AGE start/end/duration values

	Test asset for bill for example

	#!rb none
	#!tests pie
	#!review-3478427 Jon.Lietz, [REDACTED_HANDLE]

Change 3478424 on 2017/06/07 by [REDACTED_PERSON]

	Prevent creating invalid 'VBCompactPoseData', resulting in crashes in Animation Modifiers.
	(Fix for licensee crash).

	#!rb lina.halper
	#!codereview martin.wilson
	#!tests Ice sync marker automator from Athomas.

Change 3478151 on 2017/06/07 by [REDACTED_PERSON]

	spot edigrate GameplayTagQuery customization fix for crash when editing query on bp defaults.
	#!rb none
	#!tests compile

Change 3477983 on 2017/06/07 by robomerge

	#!ROBOMERGE-AUTHOR: alexis.matte
	[NULL MERGE]
	Fix morph target import

	#!jira OR-38471
	#!rb none
	#!tests none
	#!lockdown Andrew.Grant

	#!ROBOMERGE-SOURCE: CL 3477453 in //Orion/Release-40.4/... via CL 3477925 via CL 3477941
	#!ROBOMERGE-BOT: ORION (Main -> Dev-Cinematics)

Change 3477982 on 2017/06/07 by robomerge

	#!ROBOMERGE-AUTHOR: alexis.matte
	[NULL MERGE]
	Fix morph target import

	#!jira OR-38471
	#!rb none
	#!tests none
	#!lockdown Andrew.Grant

	#!ROBOMERGE-SOURCE: CL 3477453 in //Orion/Release-40.4/... via CL 3477925 via CL 3477941
	#!ROBOMERGE-BOT: ORION (Main -> Dev-Balance)

Change 3477981 on 2017/06/07 by robomerge

	#!ROBOMERGE-AUTHOR: alexis.matte
	[NULL MERGE]
	Fix morph target import

	#!jira OR-38471
	#!rb none
	#!tests none
	#!lockdown Andrew.Grant

	#!ROBOMERGE-SOURCE: CL 3477453 in //Orion/Release-40.4/... via CL 3477925 via CL 3477941
	#!ROBOMERGE-BOT: ORION (Main -> Dev-REGS)

Change 3477980 on 2017/06/07 by robomerge

	#!ROBOMERGE-AUTHOR: alexis.matte
	[NULL MERGE]
	Fix morph target import

	#!jira OR-38471
	#!rb none
	#!tests none
	#!lockdown Andrew.Grant

	#!ROBOMERGE-SOURCE: CL 3477453 in //Orion/Release-40.4/... via CL 3477925 via CL 3477941
	#!ROBOMERGE-BOT: ORION (Main -> Dev-UI)

Change 3477979 on 2017/06/07 by robomerge

	#!ROBOMERGE-AUTHOR: alexis.matte
	[NULL MERGE]
	Fix morph target import

	#!jira OR-38471
	#!rb none
	#!tests none
	#!lockdown Andrew.Grant

	#!ROBOMERGE-SOURCE: CL 3477453 in //Orion/Release-40.4/... via CL 3477925 via CL 3477941
	#!ROBOMERGE-BOT: ORION (Main -> Dev-General)

Change 3477941 on 2017/06/07 by robomerge

	#!ROBOMERGE-AUTHOR: alexis.matte
	[NULL MERGE]
	Fix morph target import

	#!jira OR-38471
	#!rb none
	#!tests none
	#!lockdown Andrew.Grant

	#!ROBOMERGE-SOURCE: CL 3477453 in //Orion/Release-40.4/... via CL 3477925
	#!ROBOMERGE-BOT: ORION (Release-40.5 -> Main)

Change 3477925 on 2017/06/07 by robomerge

	#!ROBOMERGE-AUTHOR: alexis.matte
	Fix morph target import

	#!jira OR-38471
	#!rb none
	#!tests none
	#!lockdown Andrew.Grant

	#!ROBOMERGE-SOURCE: CL 3477453 in //Orion/Release-40.4/...
	#!ROBOMERGE-BOT: ORION (Release-40.4 -> Release-40.5)
	#!ROBOMERGE[ORION]: !Main

Change 3477774 on 2017/06/07 by [REDACTED_PERSON]

	implement a dev-editor cl 3470188
	Fix the material isolate for cloth or hair
	#!jira UE-38985
	#!rb none
	#!tests none

Change 3477722 on 2017/06/07 by [REDACTED_PERSON]

	Re-enabling D-Pad navigation support in card shop.

	Exposed OnNavigation to UserWidget in the form of NativeOnNavigation, leveraged this new feature to have the classes I care about (HandEntry / CardShopEquipSlot)

	Split out BaseButton_Group's "SelectNextButton" process into "GetButton" and "Select Button" so I could use the GetButton when doing navigation.

	#!rb matt.schembari
	#!tests Compile DebugGameEditor Win64 / Shipping Client PS4

Change 3477610 on 2017/06/07 by [REDACTED_PERSON]

	Fixing up emitter nodes in system graph when deleted

	#!rb none
	#!tests added/removed multiple emitters

Change 3477528 on 2017/06/07 by [REDACTED_PERSON]

	? Fixed up issue with interface function binding from the removal of variable IDs.
	? Fixed issue where system parameters were garbage on the first tick of a system.
	? Bypassed issue with component lifetime. When destroying systems in some cases the component is pending kill so it's weak ptr returns null.
	We need to investigate this further.

	#!rb none
	#!tests stuff works
	#!codereview Olaf.Piesche, Frank.Fella, Shaun.Kime

Change 3477453 on 2017/06/07 by [REDACTED_PERSON]

	Fix morph target import

	#!jira OR-38471
	#!rb none
	#!tests none
	#!ROBOMERGE: !Main
	#!lockdown Andrew.Grant

Change 3477182 on 2017/06/07 by [REDACTED_PERSON]

	Niagara - Rename files from class renames in last check-in.

	#!tests Compiled.
	#!rb none

Change 3477171 on 2017/06/06 by [REDACTED_PERSON]

	Niagara - Can now add dynamic inputs directly in the stack.

	#!tests Added dynamic inputs directly from the stack.
	#!rb none

	#!codereview Olaf.Piesche,Simon.Tovey,Shaun.Kime

Change 3477115 on 2017/06/06 by [REDACTED_PERSON]

	Merging //Orion/Main to Release-40.5 (//Orion/Release-40.5) [REDACTED_HANDLE]

	#!rb none
	#!tests none

Change 3477098 on 2017/06/06 by [REDACTED_PERSON]

	Initial branch of files from Release-40.4 (//Orion/Release-40.4) to Release-40.5 (//Orion/Release-40.5)

Change 3476585 on 2017/06/06 by [REDACTED_PERSON]

	EQS touches to hopefully address the elusive EQS NaN in live build #!Orion

	#!test golden path
	#!rb none

Change 3476342 on 2017/06/06 by [REDACTED_PERSON]

	FCSPose<PoseType>::ConvertToLocalPoses Allow root bone to be modified. Minor optimization: Take out root bone check from loop.

	#!rb lina.halper
	#!tests Ghost PIE

Change 3476336 on 2017/06/06 by [REDACTED_PERSON]

	First pass at trying to prevent Wyeth's crash in the EmitterInstance destructor.

	#!rb none
	#!tests tried iterating with multiple changes between emitters/systems
	#!codereview simon.tovey, frank.fella, olaf.piesche

Change 3476160 on 2017/06/06 by [REDACTED_PERSON]

	Removing ID's from FNiagaraVariables. Reworking existing code to properly handle this.

	#!rb none
	#!codereview simon.tovey, frank.fella, olaf.piesche
	#!tests recompiled and ran existing emitters, created system, iterated between system and emitter

Change 3476157 on 2017/06/06 by [REDACTED_PERSON]

	Fixing code dependency

	#!rb none
	#!tests n/a

Change 3476155 on 2017/06/06 by [REDACTED_PERSON]

	Added ability to get Emitter alias from parameter map

	#!tests n/a
	#!rb none

Change 3476152 on 2017/06/06 by [REDACTED_PERSON]

	Fixing comment so that system tooltip was meaningful from creation menu

	#!rb none
	#!tests  n/a

Change 3476148 on 2017/06/06 by [REDACTED_PERSON]

	Removing gamethread checks as we use a parallel for to update emitter instances, causing this to always fail with multiple emitters in a system.

	#!rb none
	#!codereview simon.tovey, olaf.piesche
	#!tests added multiple emitters and didn't crash

Change 3475898 on 2017/06/06 by [REDACTED_PERSON]

	Manual recreation of CL#!3465092 #!UE4

	By LukaszF: "fixed navigation area modifiers created from shape components: sphere and capsule"

	#!test golden path
	#!rb none

Change 3475817 on 2017/06/06 by robomerge

	#!ROBOMERGE-AUTHOR: andrew.grant
	Now with correctly unshelved CL - made Aftermath a command line option

	#!tests compiled, verified initialziation is command line driven
	#!rb none

	#!ROBOMERGE-SOURCE: CL 3475810 in //Orion/Release-40.4/... via CL 3475812
	#!ROBOMERGE-BOT: ORION (Main -> Dev-Cinematics)

Change 3475816 on 2017/06/06 by robomerge

	#!ROBOMERGE-AUTHOR: andrew.grant
	Now with correctly unshelved CL - made Aftermath a command line option

	#!tests compiled, verified initialziation is command line driven
	#!rb none

	#!ROBOMERGE-SOURCE: CL 3475810 in //Orion/Release-40.4/... via CL 3475812
	#!ROBOMERGE-BOT: ORION (Main -> Dev-Balance)

Change 3475815 on 2017/06/06 by robomerge

	#!ROBOMERGE-AUTHOR: andrew.grant
	Now with correctly unshelved CL - made Aftermath a command line option

	#!tests compiled, verified initialziation is command line driven
	#!rb none

	#!ROBOMERGE-SOURCE: CL 3475810 in //Orion/Release-40.4/... via CL 3475812
	#!ROBOMERGE-BOT: ORION (Main -> Dev-REGS)

Change 3475814 on 2017/06/06 by robomerge

	#!ROBOMERGE-AUTHOR: andrew.grant
	Now with correctly unshelved CL - made Aftermath a command line option

	#!tests compiled, verified initialziation is command line driven
	#!rb none

	#!ROBOMERGE-SOURCE: CL 3475810 in //Orion/Release-40.4/... via CL 3475812
	#!ROBOMERGE-BOT: ORION (Main -> Dev-UI)

Change 3475813 on 2017/06/06 by robomerge

	#!ROBOMERGE-AUTHOR: andrew.grant
	Now with correctly unshelved CL - made Aftermath a command line option

	#!tests compiled, verified initialziation is command line driven
	#!rb none

	#!ROBOMERGE-SOURCE: CL 3475810 in //Orion/Release-40.4/... via CL 3475812
	#!ROBOMERGE-BOT: ORION (Main -> Dev-General)

Change 3475812 on 2017/06/06 by robomerge

	#!ROBOMERGE-AUTHOR: andrew.grant
	Now with correctly unshelved CL - made Aftermath a command line option

	#!tests compiled, verified initialziation is command line driven
	#!rb none

	#!ROBOMERGE-SOURCE: CL 3475810 in //Orion/Release-40.4/...
	#!ROBOMERGE-BOT: ORION (Release-40.4 -> Main)

Change 3475810 on 2017/06/06 by [REDACTED_PERSON]

	Now with correctly unshelved CL - made Aftermath a command line option

	#!tests compiled, verified initialziation is command line driven
	#!rb none

Change 3475792 on 2017/06/06 by [REDACTED_PERSON]

	item cooldowns

	- added in native ability class (UOrionSourceItemAbility) that will be repsonsible for item keyword cooldowns and cost.
	- Moved Application, trigger and activation/deactivation of itemkeywords out of the deck instance and into UOrionSourceItemAbility.
	- added in support for cultivate card trait
	- added in to the engine FAbilityEndedData that will pass through delegates what ability ended the spec handle and if it was cancelled or not
	- added 2 delegates for when abilities end, one inside UAbilitySystemComponent::NotifyAbilityEnded() the other in UGameplayAbility::EndAbility() they bost pass through a const  FAbilityEndedData&

	#!rb david.ratti
	#!tests buy and play cards

Change 3475760 on 2017/06/06 by robomerge

	#!ROBOMERGE-AUTHOR: andrew.grant
	Made aftermath iniitialization off by default and controlled by the -aftermath command line option
	Logs are now warnings if aftermath is requested but can't be initialized

	#!tests verified command line test works
	#!rb none
	[REDACTED_HANDLE], [REDACTED_HANDLE]

	#!ROBOMERGE-SOURCE: CL 3475753 in //Orion/Release-40.4/... via CL 3475755
	#!ROBOMERGE-BOT: ORION (Main -> Dev-Cinematics)

Change 3475759 on 2017/06/06 by robomerge

	#!ROBOMERGE-AUTHOR: andrew.grant
	Made aftermath iniitialization off by default and controlled by the -aftermath command line option
	Logs are now warnings if aftermath is requested but can't be initialized

	#!tests verified command line test works
	#!rb none
	[REDACTED_HANDLE], [REDACTED_HANDLE]

	#!ROBOMERGE-SOURCE: CL 3475753 in //Orion/Release-40.4/... via CL 3475755
	#!ROBOMERGE-BOT: ORION (Main -> Dev-Balance)

Change 3475758 on 2017/06/06 by robomerge

	#!ROBOMERGE-AUTHOR: andrew.grant
	Made aftermath iniitialization off by default and controlled by the -aftermath command line option
	Logs are now warnings if aftermath is requested but can't be initialized

	#!tests verified command line test works
	#!rb none
	[REDACTED_HANDLE], [REDACTED_HANDLE]

	#!ROBOMERGE-SOURCE: CL 3475753 in //Orion/Release-40.4/... via CL 3475755
	#!ROBOMERGE-BOT: ORION (Main -> Dev-REGS)

Change 3475757 on 2017/06/06 by robomerge

	#!ROBOMERGE-AUTHOR: andrew.grant
	Made aftermath iniitialization off by default and controlled by the -aftermath command line option
	Logs are now warnings if aftermath is requested but can't be initialized

	#!tests verified command line test works
	#!rb none
	[REDACTED_HANDLE], [REDACTED_HANDLE]

	#!ROBOMERGE-SOURCE: CL 3475753 in //Orion/Release-40.4/... via CL 3475755
	#!ROBOMERGE-BOT: ORION (Main -> Dev-UI)

Change 3475756 on 2017/06/06 by robomerge

	#!ROBOMERGE-AUTHOR: andrew.grant
	Made aftermath iniitialization off by default and controlled by the -aftermath command line option
	Logs are now warnings if aftermath is requested but can't be initialized

	#!tests verified command line test works
	#!rb none
	[REDACTED_HANDLE], [REDACTED_HANDLE]

	#!ROBOMERGE-SOURCE: CL 3475753 in //Orion/Release-40.4/... via CL 3475755
	#!ROBOMERGE-BOT: ORION (Main -> Dev-General)

Change 3475755 on 2017/06/06 by robomerge

	#!ROBOMERGE-AUTHOR: andrew.grant
	Made aftermath iniitialization off by default and controlled by the -aftermath command line option
	Logs are now warnings if aftermath is requested but can't be initialized

	#!tests verified command line test works
	#!rb none
	[REDACTED_HANDLE], [REDACTED_HANDLE]

	#!ROBOMERGE-SOURCE: CL 3475753 in //Orion/Release-40.4/...
	#!ROBOMERGE-BOT: ORION (Release-40.4 -> Main)

Change 3475753 on 2017/06/06 by [REDACTED_PERSON]

	Made aftermath iniitialization off by default and controlled by the -aftermath command line option
	Logs are now warnings if aftermath is requested but can't be initialized

	#!tests verified command line test works
	#!rb none
	#!review-3475754 [REDACTED_HANDLE], [REDACTED_HANDLE]

Change 3475491 on 2017/06/06 by [REDACTED_PERSON]

	Feeding parameter collection values into simulaitons.

	? Setup binding from parameter collections to simulation exec contexts. Data is fed in now.
	? Modified names of collection parameter such that they're always uniquely associated with a particular collection. In case two sets use the same name for example.
	Required some name conversion between the internals and the UI.
	? Modified node to not link to params by ID as they will be removed shortly.
	? NiagaraWorldManager now ticking to push parameter data from global collections.
	? Added BP function library call to grab the global collection instance for a collection and BP getters and setters for instances.
	? Components also can override the global instance though this isn't hooked up to anything as yet. I imagine this will be handy for creating override volumes in the world and having components interpolate between those similar to post process volumes.

	Minor/unrelated
	? Fixed crash on exit. Changed system instance in component to be Unique ptr and always access via component to more direcly control lifetime.
	? Crash fix when getting matrices from parameter map. TypeEditorUtilities was null.
	? Fixed bug in GetTypeDefaultValue()
	? Fixed property tagging on FNiagaraStatScope

	#!tests emitters work. Data is fed in.
	#!rb none
	#!codereview Olaf.Piesche, Shaun.Kime, Frank.Fella

Change 3474483 on 2017/06/05 by [REDACTED_PERSON]

	Added new BlendBoneByChannel AnimNode to blend two poses, per bone, per channel. For example blend only translation from Pelvis.

	#!rb none
	#!test Ghost
	#!codereview lina.halper

Change 3474099 on 2017/06/05 by [REDACTED_PERSON]

	Copy/paste material should copy paste only the material instance
	#!rb none
	#!test none

Change 3474073 on 2017/06/05 by [REDACTED_PERSON]

	Added estimated timing for reatltime updates.
	#!rb Trivial
	#!test Launch build paragon.

Change 3474066 on 2017/06/05 by [REDACTED_PERSON]

	Increased heartbeat frequency for realtime cooking.
	#!rb Trivial
	#!test Realtime cooking

Change 3473623 on 2017/06/05 by [REDACTED_PERSON]

	Using notimeouts on client and server when running realtime cooking, as the client is slowed down making it timeout.
	#!rb Trivial
	#!test Realtime cook paragon orion_entry.

Change 3473484 on 2017/06/05 by [REDACTED_PERSON]

	Niagara - Preliminary support for dynamic inputs.

	#!tests Dynamic inputs are shown in the stack UI and their inputs are editable.
	#!rb none

	#!codereview Olaf.Piesche,Simon.Tovey,Shaun.Kime

Change 3473481 on 2017/06/05 by [REDACTED_PERSON]

	Niagara - Highlight the connecting wire when hovering the wire itself or one of it's connected pins.

	#!tests The wire highlights.
	#!rb none

	#!codereview Olaf.Piesche,Simon.Tovey,Shaun.Kime

Change 3473480 on 2017/06/05 by [REDACTED_PERSON]

	Niagara - Notify the graph that it has changed when adding and connecting pins on a node with dynamic pins.

	#!tests The graph is now shown as modified and needing compiling when connecting or adding pins on a node with dynamic pins.

	#!rb none
	#!codereview Olaf.Piesche,Simon.Tovey,Shaun.Kime

Change 3473479 on 2017/06/05 by [REDACTED_PERSON]

	Niagara - Fix an issue where module inputs were not getting aliased correctly when there was more than one of the same node when modifying them from the stack.

	#!test The inputs now get aliased correctly.
	#!rb none

	#!codereview Olaf.Piesche,Simon.Tovey,Shaun.Kime

Change 3472889 on 2017/06/03 by [REDACTED_PERSON]

	Fixed merge error
	#!tests compiled
	#!rb none

Change 3472547 on 2017/06/02 by [REDACTED_PERSON]

	Use the correct number of instances after sim step; this makes killing particles work properly in GPU sim

	#!codereview simon.tovey
	#!rb none
	#!tests GPUTest emitter and OrbitalMotion test emitter

Change 3472452 on 2017/06/02 by [REDACTED_PERSON]

	More GPU spawn fixes; no more garbage particles in buffers after spawning with GPU simulation
	Bit more cleanup

	#!rb none
	#!tests GPUTest emitter
	#!codereview simon.tovey

Change 3472284 on 2017/06/02 by robomerge

	#!ROBOMERGE-AUTHOR: andrew.grant
	Gauntlet - if the specified build has a client but not a server, fallback to using the editor as a server

	[REDACTED_HANDLE]
	#!tests ran Gauntlet on build with / without server
	#!rb -

	#!ROBOMERGE-SOURCE: CL 3471727 in //Orion/Release-40.3/... via CL 3472202 via CL 3472213
	#!ROBOMERGE-BOT: ORION (Main -> Dev-Cinematics)

Change 3472283 on 2017/06/02 by robomerge

	#!ROBOMERGE-AUTHOR: andrew.grant
	Gauntlet - if the specified build has a client but not a server, fallback to using the editor as a server

	[REDACTED_HANDLE]
	#!tests ran Gauntlet on build with / without server
	#!rb -

	#!ROBOMERGE-SOURCE: CL 3471727 in //Orion/Release-40.3/... via CL 3472202 via CL 3472213
	#!ROBOMERGE-BOT: ORION (Main -> Dev-Balance)

Change 3472282 on 2017/06/02 by robomerge

	#!ROBOMERGE-AUTHOR: andrew.grant
	Gauntlet - if the specified build has a client but not a server, fallback to using the editor as a server

	[REDACTED_HANDLE]
	#!tests ran Gauntlet on build with / without server
	#!rb -

	#!ROBOMERGE-SOURCE: CL 3471727 in //Orion/Release-40.3/... via CL 3472202 via CL 3472213
	#!ROBOMERGE-BOT: ORION (Main -> Dev-REGS)

Change 3472278 on 2017/06/02 by robomerge

	#!ROBOMERGE-AUTHOR: andrew.grant
	Gauntlet - if the specified build has a client but not a server, fallback to using the editor as a server

	[REDACTED_HANDLE]
	#!tests ran Gauntlet on build with / without server
	#!rb -

	#!ROBOMERGE-SOURCE: CL 3471727 in //Orion/Release-40.3/... via CL 3472202 via CL 3472213
	#!ROBOMERGE-BOT: ORION (Main -> Dev-UI)

Change 3472275 on 2017/06/02 by robomerge

	#!ROBOMERGE-AUTHOR: andrew.grant
	Gauntlet - if the specified build has a client but not a server, fallback to using the editor as a server

	[REDACTED_HANDLE]
	#!tests ran Gauntlet on build with / without server
	#!rb -

	#!ROBOMERGE-SOURCE: CL 3471727 in //Orion/Release-40.3/... via CL 3472202 via CL 3472213
	#!ROBOMERGE-BOT: ORION (Main -> Dev-General)

Change 3472213 on 2017/06/02 by robomerge

	#!ROBOMERGE-AUTHOR: andrew.grant
	Gauntlet - if the specified build has a client but not a server, fallback to using the editor as a server

	[REDACTED_HANDLE]
	#!tests ran Gauntlet on build with / without server
	#!rb -

	#!ROBOMERGE-SOURCE: CL 3471727 in //Orion/Release-40.3/... via CL 3472202
	#!ROBOMERGE-BOT: ORION (Release-40.4 -> Main)

Change 3472202 on 2017/06/02 by robomerge

	#!ROBOMERGE-AUTHOR: andrew.grant
	Gauntlet - if the specified build has a client but not a server, fallback to using the editor as a server

	[REDACTED_HANDLE]
	#!tests ran Gauntlet on build with / without server
	#!rb -

	#!ROBOMERGE-SOURCE: CL 3471727 in //Orion/Release-40.3/...
	#!ROBOMERGE-BOT: ORION (Release-40.3 -> Release-40.4)

Change 3471976 on 2017/06/02 by robomerge

	#!ROBOMERGE-AUTHOR: nick.reid
	Gauntlet script fixes
	#!tests ran locally
	#!rb AG

	#!ROBOMERGE-SOURCE: CL 3471604 in //Orion/Release-40.4/... via CL 3471809
	#!ROBOMERGE-BOT: ORION (Main -> Dev-Cinematics)

Change 3471975 on 2017/06/02 by robomerge

	#!ROBOMERGE-AUTHOR: nick.reid
	Gauntlet script fixes
	#!tests ran locally
	#!rb AG

	#!ROBOMERGE-SOURCE: CL 3471604 in //Orion/Release-40.4/... via CL 3471809
	#!ROBOMERGE-BOT: ORION (Main -> Dev-Balance)

Change 3471974 on 2017/06/02 by robomerge

	#!ROBOMERGE-AUTHOR: nick.reid
	Gauntlet script fixes
	#!tests ran locally
	#!rb AG

	#!ROBOMERGE-SOURCE: CL 3471604 in //Orion/Release-40.4/... via CL 3471809
	#!ROBOMERGE-BOT: ORION (Main -> Dev-REGS)

Change 3471973 on 2017/06/02 by robomerge

	#!ROBOMERGE-AUTHOR: nick.reid
	Gauntlet script fixes
	#!tests ran locally
	#!rb AG

	#!ROBOMERGE-SOURCE: CL 3471604 in //Orion/Release-40.4/... via CL 3471809
	#!ROBOMERGE-BOT: ORION (Main -> Dev-UI)

Change 3471972 on 2017/06/02 by robomerge

	#!ROBOMERGE-AUTHOR: nick.reid
	Gauntlet script fixes
	#!tests ran locally
	#!rb AG

	#!ROBOMERGE-SOURCE: CL 3471604 in //Orion/Release-40.4/... via CL 3471809
	#!ROBOMERGE-BOT: ORION (Main -> Dev-General)

Change 3471966 on 2017/06/02 by [REDACTED_PERSON]

	Fixed robomerge integration
	#!tests #!rb none

Change 3471845 on 2017/06/02 by robomerge

	#!ROBOMERGE-AUTHOR: nick.reid
	AG - made local builds use editor server
	#!tests ran locally
	#!rb none

	#!ROBOMERGE-SOURCE: CL 3471566 in //Orion/Release-40.4/... via CL 3471806
	#!ROBOMERGE-BOT: ORION (Main -> Dev-Cinematics)

Change 3471844 on 2017/06/02 by robomerge

	#!ROBOMERGE-AUTHOR: nick.reid
	AG - made local builds use editor server
	#!tests ran locally
	#!rb none

	#!ROBOMERGE-SOURCE: CL 3471566 in //Orion/Release-40.4/... via CL 3471806
	#!ROBOMERGE-BOT: ORION (Main -> Dev-Balance)

Change 3471843 on 2017/06/02 by robomerge

	#!ROBOMERGE-AUTHOR: nick.reid
	AG - made local builds use editor server
	#!tests ran locally
	#!rb none

	#!ROBOMERGE-SOURCE: CL 3471566 in //Orion/Release-40.4/... via CL 3471806
	#!ROBOMERGE-BOT: ORION (Main -> Dev-REGS)

Change 3471842 on 2017/06/02 by robomerge

	#!ROBOMERGE-AUTHOR: nick.reid
	AG - made local builds use editor server
	#!tests ran locally
	#!rb none

	#!ROBOMERGE-SOURCE: CL 3471566 in //Orion/Release-40.4/... via CL 3471806
	#!ROBOMERGE-BOT: ORION (Main -> Dev-UI)

Change 3471835 on 2017/06/02 by robomerge

	#!ROBOMERGE-AUTHOR: ben.marsh
	Remove setting to copy full crash dumps to \\epicgames.net\root\Projects\Paragon\QA_CrashReports. Don't think anyone is using this.

	#!rb none

	#!ROBOMERGE-SOURCE: CL 3471379 in //Orion/Main/...
	#!ROBOMERGE-BOT: ORION (Main -> Dev-Cinematics)

Change 3471834 on 2017/06/02 by robomerge

	#!ROBOMERGE-AUTHOR: ben.marsh
	Remove setting to copy full crash dumps to \\epicgames.net\root\Projects\Paragon\QA_CrashReports. Don't think anyone is using this.

	#!rb none

	#!ROBOMERGE-SOURCE: CL 3471379 in //Orion/Main/...
	#!ROBOMERGE-BOT: ORION (Main -> Dev-Balance)

Change 3471833 on 2017/06/02 by robomerge

	#!ROBOMERGE-AUTHOR: ben.marsh
	Remove setting to copy full crash dumps to \\epicgames.net\root\Projects\Paragon\QA_CrashReports. Don't think anyone is using this.

	#!rb none

	#!ROBOMERGE-SOURCE: CL 3471379 in //Orion/Main/...
	#!ROBOMERGE-BOT: ORION (Main -> Dev-REGS)

Change 3471832 on 2017/06/02 by robomerge

	#!ROBOMERGE-AUTHOR: ben.marsh
	Remove setting to copy full crash dumps to \\epicgames.net\root\Projects\Paragon\QA_CrashReports. Don't think anyone is using this.

	#!rb none

	#!ROBOMERGE-SOURCE: CL 3471379 in //Orion/Main/...
	#!ROBOMERGE-BOT: ORION (Main -> Dev-UI)

Change 3471831 on 2017/06/02 by robomerge

	#!ROBOMERGE-AUTHOR: ben.marsh
	Remove setting to copy full crash dumps to \\epicgames.net\root\Projects\Paragon\QA_CrashReports. Don't think anyone is using this.

	#!rb none

	#!ROBOMERGE-SOURCE: CL 3471379 in //Orion/Main/...
	#!ROBOMERGE-BOT: ORION (Main -> Dev-General)

Change 3471809 on 2017/06/02 by robomerge

	#!ROBOMERGE-AUTHOR: nick.reid
	Gauntlet script fixes
	#!tests ran locally
	#!rb AG

	#!ROBOMERGE-SOURCE: CL 3471604 in //Orion/Release-40.4/...
	#!ROBOMERGE-BOT: ORION (Release-40.4 -> Main)

Change 3471806 on 2017/06/02 by robomerge

	#!ROBOMERGE-AUTHOR: nick.reid
	AG - made local builds use editor server
	#!tests ran locally
	#!rb none

	#!ROBOMERGE-SOURCE: CL 3471566 in //Orion/Release-40.4/...
	#!ROBOMERGE-BOT: ORION (Release-40.4 -> Main)

Change 3471727 on 2017/06/02 by [REDACTED_PERSON]

	Gauntlet - if the specified build has a client but not a server, fallback to using the editor as a server

	#!review-3471728 [REDACTED_HANDLE]
	#!tests ran Gauntlet on build with / without server
	#!rb -

Change 3471689 on 2017/06/02 by [REDACTED_PERSON]

	#!ue4-orion - Added virtual OnClientCorrectionReceived() to CharacterMovement.

	Stubbed implementation for Orion to be replaced/augmented for analytics.

	#!codereview Andrew.Grant
	#!rb none
	#!jira OR-37131
	#!tests Multi PIE

Change 3471654 on 2017/06/02 by [REDACTED_PERSON]

	Merging file cull from //Orion/Main to Dev-Balance (//Orion/Dev-Balance)
	#!tests #!rb na

Change 3471627 on 2017/06/02 by [REDACTED_PERSON]

	Merging file pruning from //Orion/Main to Dev-Cinematics (//Orion/Dev-Cinematics)
	#!tests #!rb na

Change 3471604 on 2017/06/02 by [REDACTED_PERSON]

	Gauntlet script fixes
	#!tests ran locally
	#!rb AG

Change 3471566 on 2017/06/02 by [REDACTED_PERSON]

	AG - made local builds use editor server
	#!tests ran locally
	#!rb none

Change 3471379 on 2017/06/02 by [REDACTED_PERSON]

	Remove setting to copy full crash dumps to \\epicgames.net\root\Projects\Paragon\QA_CrashReports. Don't think anyone is using this.

	#!rb none

Change 3471304 on 2017/06/02 by andrew.grant

	Removing some unused files to free up space across branches
	#!tests compiled locally, preflighted standard build
	#!rb none

	#!ROBOMERGE-SOURCE: CL 3470976 in //Orion/Release-40.2/... via CL 3471002 via CL 3471024 via CL 3471072
	#!ROBOMERGE-BOT: ORION (Main -> Dev-General)

	#!ROBOMERGE-SAYS: Unresolved conflicts. andrew.grant, please merge this change by hand.
	//ROBOMERGE_ORION_Dev_General/Engine/Source/ThirdParty/PhysX/Lib/Win32/VS2015/APEX_Clothing_x86.lib
	//ROBOMERGE_ORION_Dev_General/Engine/Source/ThirdParty/PhysX/Lib/Win32/VS2015/APEX_ClothingCHECKED_x86.lib
	//ROBOMERGE_ORION_Dev_General/Engine/Source/ThirdParty/PhysX/Lib/Win32/VS2015/APEX_ClothingPROFILE_x86.lib
	//ROBOMERGE_ORION_Dev_General/Engine/Source/ThirdParty/PhysX/Lib/Win32/VS2015/APEX_Destructible_x86.lib
	//ROBOMERGE_ORION_Dev_General/Engine/Source/ThirdParty/PhysX/Lib/Win32/VS2015/APEX_DestructibleCHECKED_x86.lib
	//ROBOMERGE_ORION_Dev_General/Engine/Source/ThirdParty/PhysX/Lib/Win32/VS2015/APEX_DestructiblePROFILE_x86.lib
	//ROBOMERGE_ORION_Dev_General/Engine/Source/ThirdParty/PhysX...
	#!CodeReview: andrew.grant, jason.bestimt, jeff.williams

Change 3471231 on 2017/06/02 by robomerge

	#!ROBOMERGE-AUTHOR: andrew.grant
	Removing some unused files to free up space across branches
	#!tests compiled locally, preflighted standard build
	#!rb none

	#!ROBOMERGE-SOURCE: CL 3470976 in //Orion/Release-40.2/... via CL 3471002 via CL 3471024 via CL 3471072
	#!ROBOMERGE-BOT: ORION (Main -> Dev-REGS)

Change 3471205 on 2017/06/02 by robomerge

	#!ROBOMERGE-AUTHOR: andrew.grant
	Removing some unused files to free up space across branches
	#!tests compiled locally, preflighted standard build
	#!rb none

	#!ROBOMERGE-SOURCE: CL 3470976 in //Orion/Release-40.2/... via CL 3471002 via CL 3471024 via CL 3471072
	#!ROBOMERGE-BOT: ORION (Main -> Dev-UI)

Change 3471072 on 2017/06/02 by robomerge

	#!ROBOMERGE-AUTHOR: andrew.grant
	Removing some unused files to free up space across branches
	#!tests compiled locally, preflighted standard build
	#!rb none

	#!ROBOMERGE-SOURCE: CL 3470976 in //Orion/Release-40.2/... via CL 3471002 via CL 3471024
	#!ROBOMERGE-BOT: ORION (Release-40.4 -> Main)

Change 3471024 on 2017/06/02 by robomerge

	#!ROBOMERGE-AUTHOR: andrew.grant
	Removing some unused files to free up space across branches
	#!tests compiled locally, preflighted standard build
	#!rb none

	#!ROBOMERGE-SOURCE: CL 3470976 in //Orion/Release-40.2/... via CL 3471002
	#!ROBOMERGE-BOT: ORION (Release-40.3 -> Release-40.4)

Change 3471002 on 2017/06/02 by robomerge

	#!ROBOMERGE-AUTHOR: andrew.grant
	Removing some unused files to free up space across branches
	#!tests compiled locally, preflighted standard build
	#!rb none

	#!ROBOMERGE-SOURCE: CL 3470976 in //Orion/Release-40.2/...
	#!ROBOMERGE-BOT: ORION (Release-40.2 -> Release-40.3)

Change 3470976 on 2017/06/01 by [REDACTED_PERSON]

	Removing some unused files to free up space across branches
	#!tests compiled locally, preflighted standard build
	#!rb none

Change 3470672 on 2017/06/01 by [REDACTED_PERSON]

	Added new commandline argument for gauntlet which allows seperate client commands.
	Fixed realtime cooking to pass commandline options correctly to the server and client.
	#!rb None
	#!test Realtime cooking paragon

Change 3470645 on 2017/06/01 by [REDACTED_PERSON]

	GPU sim part 2; cleanup, more bug fixing

	#!lockdown Andrew.Bains
	#!codereview simon.tovey
	#!rb none
	#!tests the usual

Change 3470636 on 2017/06/01 by [REDACTED_PERSON]

	Improved startup time of editor by reducing number of automatic cook platforms for realtime cooking.
	#!rb Trivial
	#!test Editor paragon.

Change 3470472 on 2017/06/01 by [REDACTED_PERSON]

	Checkpointing work on compiling system and emitter graph. Very simple graphs of these types work now. No harm has befallen any of the previously working graphs.

	Some constants did change and you will MANUALLY NEED TO UPDATE any graphs referencing them.

	// Engine parameters are always read-only, no matter what level you are at.
	Engine.DeltaTime
	Engine.InverseDeltaTime
	Engine.ExecutionCount
	Engine.Owner.Position
	Engine.Owner.Velocity
	Engine.Owner.XAxis
	Engine.Owner.YAxis
	Engine.Owner.ZAxis
	Engine.Owner.LocalToWorld
	Engine.Owner.WorldToLocal
	Engine.Owner.LocalToWorldTransposed
	Engine.Owner.WorldToLocalTransposed

	// System parameters are writable in System Spawn/Update scripts and read-only otherwise.
	System.Age

	// Emitter parameters are writable in System Spawn/Update & Emitter Spawn/Update scripts and read-only otherwise.
	Emitter.Age
	Emitter.SpawnRate
	Emitter.SpawnInterval
	Emitter.InterpSpawnStartDt
	Emitter.PreviousSpawnRemainder

	#!rb none
	#!tests all existing graphs
	#!code.review frank.fella, simon.tovey, olaf.piesche

Change 3469908 on 2017/06/01 by robomerge

	#!ROBOMERGE-AUTHOR: andrew.grant
	Bumped script version to grab new publishing tools
	#!tests #!rb none

	#!ROBOMERGE-SOURCE: CL 3469901 in //Orion/Release-40.3/... via CL 3469902 via CL 3469903
	#!ROBOMERGE-BOT: ORION (Main -> Dev-Cinematics)

Change 3469907 on 2017/06/01 by robomerge

	#!ROBOMERGE-AUTHOR: andrew.grant
	Bumped script version to grab new publishing tools
	#!tests #!rb none

	#!ROBOMERGE-SOURCE: CL 3469901 in //Orion/Release-40.3/... via CL 3469902 via CL 3469903
	#!ROBOMERGE-BOT: ORION (Main -> Dev-Balance)

Change 3469906 on 2017/06/01 by robomerge

	#!ROBOMERGE-AUTHOR: andrew.grant
	Bumped script version to grab new publishing tools
	#!tests #!rb none

	#!ROBOMERGE-SOURCE: CL 3469901 in //Orion/Release-40.3/... via CL 3469902 via CL 3469903
	#!ROBOMERGE-BOT: ORION (Main -> Dev-REGS)

Change 3469905 on 2017/06/01 by robomerge

	#!ROBOMERGE-AUTHOR: andrew.grant
	Bumped script version to grab new publishing tools
	#!tests #!rb none

	#!ROBOMERGE-SOURCE: CL 3469901 in //Orion/Release-40.3/... via CL 3469902 via CL 3469903
	#!ROBOMERGE-BOT: ORION (Main -> Dev-UI)

Change 3469904 on 2017/06/01 by robomerge

	#!ROBOMERGE-AUTHOR: andrew.grant
	Bumped script version to grab new publishing tools
	#!tests #!rb none

	#!ROBOMERGE-SOURCE: CL 3469901 in //Orion/Release-40.3/... via CL 3469902 via CL 3469903
	#!ROBOMERGE-BOT: ORION (Main -> Dev-General)

Change 3469903 on 2017/06/01 by robomerge

	#!ROBOMERGE-AUTHOR: andrew.grant
	Bumped script version to grab new publishing tools
	#!tests #!rb none

	#!ROBOMERGE-SOURCE: CL 3469901 in //Orion/Release-40.3/... via CL 3469902
	#!ROBOMERGE-BOT: ORION (Release-40.4 -> Main)

Change 3469902 on 2017/06/01 by robomerge

	#!ROBOMERGE-AUTHOR: andrew.grant
	Bumped script version to grab new publishing tools
	#!tests #!rb none

	#!ROBOMERGE-SOURCE: CL 3469901 in //Orion/Release-40.3/...
	#!ROBOMERGE-BOT: ORION (Release-40.3 -> Release-40.4)

Change 3469901 on 2017/06/01 by [REDACTED_PERSON]

	Bumped script version to grab new publishing tools
	#!tests #!rb none

Change 3469459 on 2017/06/01 by robomerge

	#!ROBOMERGE-AUTHOR: david.ratti
	UBT Merge from BenM:

	UBT: Fix single-file compile causing a different UHT manifest to be generated, potentially excluding hidden dependencies.

	#!rb none
	#!tests single file compile

	#!ROBOMERGE-SOURCE: CL 3469454 in //Orion/Main/...
	#!ROBOMERGE-BOT: ORION (Main -> Dev-Cinematics)

Change 3469458 on 2017/06/01 by robomerge

	#!ROBOMERGE-AUTHOR: david.ratti
	UBT Merge from BenM:

	UBT: Fix single-file compile causing a different UHT manifest to be generated, potentially excluding hidden dependencies.

	#!rb none
	#!tests single file compile

	#!ROBOMERGE-SOURCE: CL 3469454 in //Orion/Main/...
	#!ROBOMERGE-BOT: ORION (Main -> Dev-Balance)

Change 3469457 on 2017/06/01 by robomerge

	#!ROBOMERGE-AUTHOR: david.ratti
	UBT Merge from BenM:

	UBT: Fix single-file compile causing a different UHT manifest to be generated, potentially excluding hidden dependencies.

	#!rb none
	#!tests single file compile

	#!ROBOMERGE-SOURCE: CL 3469454 in //Orion/Main/...
	#!ROBOMERGE-BOT: ORION (Main -> Dev-REGS)

Change 3469455 on 2017/06/01 by robomerge

	#!ROBOMERGE-AUTHOR: david.ratti
	UBT Merge from BenM:

	UBT: Fix single-file compile causing a different UHT manifest to be generated, potentially excluding hidden dependencies.

	#!rb none
	#!tests single file compile

	#!ROBOMERGE-SOURCE: CL 3469454 in //Orion/Main/...
	#!ROBOMERGE-BOT: ORION (Main -> Dev-General)

Change 3469454 on 2017/06/01 by [REDACTED_PERSON]

	UBT Merge from BenM:

	UBT: Fix single-file compile causing a different UHT manifest to be generated, potentially excluding hidden dependencies.

	#!rb none
	#!tests single file compile

Change 3469422 on 2017/06/01 by [REDACTED_PERSON]

	Cursor - We shouldn't try to map the cursor for "None".  Also fixing the ensure to use printf formatting.

	#!fyi Matt.Schembari
	#!rb none
	#!tests ran on PS4

Change 3469368 on 2017/06/01 by [REDACTED_PERSON]

	Added support for precooked cook on the fly with realtime updates.
	Prefly for short.
	#!rb Andrew.Grant
	#!review-3468486 [REDACTED_HANDLE], [REDACTED_HANDLE]
	#!test Cook paragon, prefly paragon, shared cooked builds paragon

Change 3469261 on 2017/06/01 by [REDACTED_PERSON]

	Main thrust of this CL is to improve parameter handling for both code complexity and performance.
	Also paves the way for simple binding of parameter collections.

	- Refactored much execution work into FNiagaraScriptExecutionContext and made them persistent objects. This should be usable for system level scripts too.
	- Moved paraemter storage to use FNiagaraParameterStore. Done away with all those arrays and searching to build a final temp buffer for execution.
	- Same buffer should work for CPU and GPU.
	- Now binding directly between parameter stores to push data down into execution contexts that use it.
	- Future CL will extend systems to bind to the parameter collections they use so edits to said collection will automatically propagate down into using emtiters.
	- Changed parameter collections slightly so their instances will always have the same layout and have a copy of all the collection's data. Will remove a couple of cases where a rebind would be required at runtime.

	MISC
	- Moved stats id creation to the script itself as this data was being duplicated for every emitter.
	- Moved previous frame parameter data for interpolated spawn to the start of the parameter buffer to better fit in with other changes.
	- Various minor bug fixes.

	#!rb Shaun.Kime
	#!tests Test emitters work. Maybe a few issues with GPU sim which I'll work through with Olaf.
	#!codereview Shaun.Kime, Frank.Fella, Olaf.Piesche

Change 3469232 on 2017/06/01 by [REDACTED_PERSON]

	UBT: Fix single-file compile causing a different UHT manifest to be generated, potentially excluding hidden dependencies.

	#!rb none
	#!fyi David.Ratti
	#!tests single file compile

Change 3468842 on 2017/06/01 by robomerge

	#!ROBOMERGE-AUTHOR: mieszko.zielinski
	Changed 'ensureAlways' to 'ensure' in EnvQueryInstance.cpp #!UE4

	A temp fix for hitches in OR-39101. Looking for a root cause now.

	#!rb none
	#!test golden path
	#!jira OR-39101
	#!lockdown Andrew.Grant

	#!ROBOMERGE-SOURCE: CL 3468105 in //Orion/Release-40.3/... via CL 3468106 via CL 3468107
	#!ROBOMERGE-BOT: ORION (Main -> Dev-Cinematics)

Change 3468841 on 2017/06/01 by robomerge

	#!ROBOMERGE-AUTHOR: mieszko.zielinski
	Changed 'ensureAlways' to 'ensure' in EnvQueryInstance.cpp #!UE4

	A temp fix for hitches in OR-39101. Looking for a root cause now.

	#!rb none
	#!test golden path
	#!jira OR-39101
	#!lockdown Andrew.Grant

	#!ROBOMERGE-SOURCE: CL 3468105 in //Orion/Release-40.3/... via CL 3468106 via CL 3468107
	#!ROBOMERGE-BOT: ORION (Main -> Dev-Balance)

Change 3468840 on 2017/06/01 by robomerge

	#!ROBOMERGE-AUTHOR: mieszko.zielinski
	Changed 'ensureAlways' to 'ensure' in EnvQueryInstance.cpp #!UE4

	A temp fix for hitches in OR-39101. Looking for a root cause now.

	#!rb none
	#!test golden path
	#!jira OR-39101
	#!lockdown Andrew.Grant

	#!ROBOMERGE-SOURCE: CL 3468105 in //Orion/Release-40.3/... via CL 3468106 via CL 3468107
	#!ROBOMERGE-BOT: ORION (Main -> Dev-REGS)

Change 3468839 on 2017/06/01 by robomerge

	#!ROBOMERGE-AUTHOR: mieszko.zielinski
	Changed 'ensureAlways' to 'ensure' in EnvQueryInstance.cpp #!UE4

	A temp fix for hitches in OR-39101. Looking for a root cause now.

	#!rb none
	#!test golden path
	#!jira OR-39101
	#!lockdown Andrew.Grant

	#!ROBOMERGE-SOURCE: CL 3468105 in //Orion/Release-40.3/... via CL 3468106 via CL 3468107
	#!ROBOMERGE-BOT: ORION (Main -> Dev-UI)

Change 3468838 on 2017/06/01 by robomerge

	#!ROBOMERGE-AUTHOR: mieszko.zielinski
	Changed 'ensureAlways' to 'ensure' in EnvQueryInstance.cpp #!UE4

	A temp fix for hitches in OR-39101. Looking for a root cause now.

	#!rb none
	#!test golden path
	#!jira OR-39101
	#!lockdown Andrew.Grant

	#!ROBOMERGE-SOURCE: CL 3468105 in //Orion/Release-40.3/... via CL 3468106 via CL 3468107
	#!ROBOMERGE-BOT: ORION (Main -> Dev-General)

Change 3468797 on 2017/06/01 by robomerge

	#!ROBOMERGE-AUTHOR: andrew.grant
	[NULL MERGE]
	Locking Release-40.2 to network CL 3464164

	#!tests #!rb na

	#!ROBOMERGE-SOURCE: CL 3467826 in //Orion/Release-40.2/... via CL 3467827 via CL 3467828 via CL 3467829
	#!ROBOMERGE-BOT: ORION (Main -> Dev-Cinematics)

Change 3468796 on 2017/06/01 by robomerge

	#!ROBOMERGE-AUTHOR: andrew.grant
	[NULL MERGE]
	Locking Release-40.2 to network CL 3464164

	#!tests #!rb na

	#!ROBOMERGE-SOURCE: CL 3467826 in //Orion/Release-40.2/... via CL 3467827 via CL 3467828 via CL 3467829
	#!ROBOMERGE-BOT: ORION (Main -> Dev-Balance)

Change 3468795 on 2017/06/01 by robomerge

	#!ROBOMERGE-AUTHOR: andrew.grant
	[NULL MERGE]
	Locking Release-40.2 to network CL 3464164

	#!tests #!rb na

	#!ROBOMERGE-SOURCE: CL 3467826 in //Orion/Release-40.2/... via CL 3467827 via CL 3467828 via CL 3467829
	#!ROBOMERGE-BOT: ORION (Main -> Dev-REGS)

Change 3468794 on 2017/06/01 by robomerge

	#!ROBOMERGE-AUTHOR: andrew.grant
	[NULL MERGE]
	Locking Release-40.2 to network CL 3464164

	#!tests #!rb na

	#!ROBOMERGE-SOURCE: CL 3467826 in //Orion/Release-40.2/... via CL 3467827 via CL 3467828 via CL 3467829
	#!ROBOMERGE-BOT: ORION (Main -> Dev-UI)

Change 3468793 on 2017/06/01 by robomerge

	#!ROBOMERGE-AUTHOR: andrew.grant
	[NULL MERGE]
	Locking Release-40.2 to network CL 3464164

	#!tests #!rb na

	#!ROBOMERGE-SOURCE: CL 3467826 in //Orion/Release-40.2/... via CL 3467827 via CL 3467828 via CL 3467829
	#!ROBOMERGE-BOT: ORION (Main -> Dev-General)

Change 3468661 on 2017/05/31 by [REDACTED_PERSON]

	Merging fix, mostly to get a new CL
	#!tests #!rb none

Change 3468321 on 2017/05/31 by [REDACTED_PERSON]

	Merging //Orion/Dev-General @ 3466840 to Dev-General-Playtest (//Orion/Dev-General-Playtest)
	#!tests #!rb none

Change 3468107 on 2017/05/31 by robomerge

	#!ROBOMERGE-AUTHOR: mieszko.zielinski
	Changed 'ensureAlways' to 'ensure' in EnvQueryInstance.cpp #!UE4

	A temp fix for hitches in OR-39101. Looking for a root cause now.

	#!rb none
	#!test golden path
	#!jira OR-39101
	#!lockdown Andrew.Grant

	#!ROBOMERGE-SOURCE: CL 3468105 in //Orion/Release-40.3/... via CL 3468106
	#!ROBOMERGE-BOT: ORION (Release-40.4 -> Main)

Change 3468106 on 2017/05/31 by robomerge

	#!ROBOMERGE-AUTHOR: mieszko.zielinski
	Changed 'ensureAlways' to 'ensure' in EnvQueryInstance.cpp #!UE4

	A temp fix for hitches in OR-39101. Looking for a root cause now.

	#!rb none
	#!test golden path
	#!jira OR-39101
	#!lockdown Andrew.Grant


diff --git a/Engine/Config/BaseEditor.ini b/Engine/Config/BaseEditor.ini
--- a/Engine/Config/BaseEditor.ini
+++ b/Engine/Config/BaseEditor.ini
@@ -393,6 +393,7 @@ bUseFullGCAssetClassNames=true
 +ConfigSettingBlacklist=*.OnlineGameplayFramework:/Script/Engine.Engine
 +ConfigSettingBlacklist=*.Engine:OnlineSubsystem*:/Script/Engine.Engine
 +ConfigSettingBlacklist=*.Engine:Audio
++ConfigSettingBlacklist=*.Engine:/Script/LinuxTargetPlatform.LinuxTargetSettings
 
 
 +VersionedIntRValues=r.AllowStaticLighting
diff --git a/Engine/Config/BaseEditorPerProjectUserSettings.ini b/Engine/Config/BaseEditorPerProjectUserSettings.ini
--- a/Engine/Config/BaseEditorPerProjectUserSettings.ini
+++ b/Engine/Config/BaseEditorPerProjectUserSettings.ini
@@ -617,6 +617,7 @@ bResetMaterialSlots=False
 MinimumLodNumber=0
 LodNumber=0
 bAutoComputeLodDistances=True
+bIsReimportPreview=False
 
 
 [/Script/UnrealEd.FbxStaticMeshImportData]
diff --git a/Engine/Plugins/Runtime/GameplayAbilities/Source/GameplayAbilities/Private/GameplayEffect.cpp b/Engine/Plugins/Runtime/GameplayAbilities/Source/GameplayAbilities/Private/GameplayEffect.cpp
--- a/Engine/Plugins/Runtime/GameplayAbilities/Source/GameplayAbilities/Private/GameplayEffect.cpp
+++ b/Engine/Plugins/Runtime/GameplayAbilities/Source/GameplayAbilities/Private/GameplayEffect.cpp
@@ -160,6 +160,8 @@ void UGameplayEffect::PostEditChangeProperty(FPropertyChangedEvent& PropertyChan
 	}
 
 	HasGrantedApplicationImmunityQuery = !GrantedApplicationImmunityQuery.IsEmpty();
+
+	UAbilitySystemGlobals::Get().GameplayEffectPostEditChangeProperty(this, PropertyChangedEvent);
 }
 
 #endif // #if WITH_EDITOR
diff --git a/Engine/Plugins/Runtime/GameplayAbilities/Source/GameplayAbilities/Public/AbilitySystemGlobals.h b/Engine/Plugins/Runtime/GameplayAbilities/Source/GameplayAbilities/Public/AbilitySystemGlobals.h
--- a/Engine/Plugins/Runtime/GameplayAbilities/Source/GameplayAbilities/Public/AbilitySystemGlobals.h
+++ b/Engine/Plugins/Runtime/GameplayAbilities/Source/GameplayAbilities/Public/AbilitySystemGlobals.h
@@ -114,6 +114,11 @@ class GAMEPLAYABILITIES_API UAbilitySystemGlobals : public UObject
 #endif
 	}
 
+#if WITH_EDITOR
+	// Allows projects to override PostEditChangeProeprty on GEs without having to subclass Gameplayeffect. Intended for validation/auto populating based on changed data.
+	virtual void GameplayEffectPostEditChangeProperty(class UGameplayEffect* GE, FPropertyChangedEvent& PropertyChangedEvent) { }
+#endif
+
 	/** The class to instantiate as the globals object. Defaults to this class but can be overridden */
 	UPROPERTY(config)
 	FStringClassReference AbilitySystemGlobalsClassName;
diff --git a/Engine/Plugins/Runtime/GameplayAbilities/Source/GameplayAbilitiesEditor/Private/GameplayAbilitiesEditorModule.cpp b/Engine/Plugins/Runtime/GameplayAbilities/Source/GameplayAbilitiesEditor/Private/GameplayAbilitiesEditorModule.cpp
--- a/Engine/Plugins/Runtime/GameplayAbilities/Source/GameplayAbilitiesEditor/Private/GameplayAbilitiesEditorModule.cpp
+++ b/Engine/Plugins/Runtime/GameplayAbilities/Source/GameplayAbilitiesEditor/Private/GameplayAbilitiesEditorModule.cpp
@@ -45,6 +45,8 @@
 #include "LevelEditor.h"
 #include "Misc/HotReloadInterface.h"
 #include "EditorReimportHandler.h"
+#include "GameplayEffectCreationMenu.h"
+#include "ISettingsModule.h"
 
 
 class FGameplayAbilitiesEditorModule : public IGameplayAbilitiesEditorModule
@@ -140,6 +142,17 @@ void FGameplayAbilitiesEditorModule::StartupModule()
 	TSharedRef<IAssetTypeActions> GABAction = MakeShareable(new FAssetTypeActions_GameplayAbilitiesBlueprint());
 	RegisterAssetTypeAction(AssetTools, GABAction);
 
+	if (ISettingsModule* SettingsModule = FModuleManager::GetModulePtr<ISettingsModule>("Settings"))
+	{
+		SettingsModule->RegisterSettings("Project", "Project", "Gameplay Effect Parents",
+						NSLOCTEXT("GameplayAbilitiesEditorModule", "GameplayEffectParentName", "Gameplay Effect Parents"),
+						NSLOCTEXT("GameplayAbilitiesEditorModule", "GameplayEffectParentNameDesc", "Data Driven way of specifying common parent Gameplay Effect classes that are accessible through File menu"),
+						GetMutableDefault<UGameplayEffectCreationMenu>()
+						);
+
+		GetDefault<UGameplayEffectCreationMenu>()->AddMenuExtensions();
+	}
+
 	// Register factories for pins and nodes
 	GameplayAbilitiesGraphPanelPinFactory = MakeShareable(new FGameplayAbilitiesGraphPanelPinFactory());
 	FEdGraphUtilities::RegisterVisualPinFactory(GameplayAbilitiesGraphPanelPinFactory);
diff --git a/Engine/Source/Editor/ContentBrowser/Private/SContentBrowser.cpp b/Engine/Source/Editor/ContentBrowser/Private/SContentBrowser.cpp
--- a/Engine/Source/Editor/ContentBrowser/Private/SContentBrowser.cpp
+++ b/Engine/Source/Editor/ContentBrowser/Private/SContentBrowser.cpp
@@ -1737,16 +1737,28 @@ FString SContentBrowser::GetCurrentPath() const
 
 TSharedRef<SWidget> SContentBrowser::MakeAddNewContextMenu(bool bShowGetContent, bool bShowImport)
 {
+	const FSourcesData& SourcesData = AssetViewPtr->GetSourcesData();
+
+	int32 NumAssetPaths, NumClassPaths;
+	ContentBrowserUtils::CountPathTypes(SourcesData.PackagePaths, NumAssetPaths, NumClassPaths);
+
 	// Get all menu extenders for this context menu from the content browser module
 	FContentBrowserModule& ContentBrowserModule = FModuleManager::GetModuleChecked<FContentBrowserModule>( TEXT("ContentBrowser") );
-	TArray<FContentBrowserMenuExtender> MenuExtenderDelegates = ContentBrowserModule.GetAllAssetContextMenuExtenders();
+	TArray<FContentBrowserMenuExtender_SelectedPaths> MenuExtenderDelegates = ContentBrowserModule.GetAllAssetContextMenuExtenders();
+	
+	// Delegate wants paths as FStrings
+	TArray<FString> SelectPaths;
+	for (FName PathName: SourcesData.PackagePaths)
+	{
+		SelectPaths.Add(PathName.ToString());
+	}
 
 	TArray<TSharedPtr<FExtender>> Extenders;
 	for (int32 i = 0; i < MenuExtenderDelegates.Num(); ++i)
 	{
 		if (MenuExtenderDelegates[i].IsBound())
 		{
-			Extenders.Add(MenuExtenderDelegates[i].Execute());
+			Extenders.Add(MenuExtenderDelegates[i].Execute(SelectPaths));
 		}
 	}
 	TSharedPtr<FExtender> MenuExtender = FExtender::Combine(Extenders);
@@ -1760,10 +1772,6 @@ TSharedRef<SWidget> SContentBrowser::MakeAddNewContextMenu(bool bShowGetContent,
 		OnNewFolderRequested = FNewAssetOrClassContextMenu::FOnNewFolderRequested::CreateSP(this, &SContentBrowser::NewFolderRequested);
 	}
 
-	const FSourcesData& SourcesData = AssetViewPtr->GetSourcesData();
-
-	int32 NumAssetPaths, NumClassPaths;
-	ContentBrowserUtils::CountPathTypes(SourcesData.PackagePaths, NumAssetPaths, NumClassPaths);
 
 	// New feature packs don't depend on the current paths, so we always add this item if it was requested
 	FNewAssetOrClassContextMenu::FOnGetContentRequested OnGetContentRequested;
diff --git a/Engine/Source/Editor/ContentBrowser/Public/ContentBrowserModule.h b/Engine/Source/Editor/ContentBrowser/Public/ContentBrowserModule.h
--- a/Engine/Source/Editor/ContentBrowser/Public/ContentBrowserModule.h
+++ b/Engine/Source/Editor/ContentBrowser/Public/ContentBrowserModule.h
@@ -44,7 +44,7 @@ public:
 	virtual IContentBrowserSingleton& Get() const;
 
 	/** Delegates to be called to extend the content browser menus */
-	virtual TArray<FContentBrowserMenuExtender>& GetAllAssetContextMenuExtenders() {return AssetContextMenuExtenders;}
+	virtual TArray<FContentBrowserMenuExtender_SelectedPaths>& GetAllAssetContextMenuExtenders() {return AssetContextMenuExtenders;}
 	virtual TArray<FContentBrowserMenuExtender_SelectedPaths>& GetAllPathViewContextMenuExtenders() {return PathViewContextMenuExtenders;}
 	virtual TArray<FContentBrowserMenuExtender>& GetAllCollectionListContextMenuExtenders() {return CollectionListContextMenuExtenders;}
 	virtual TArray<FContentBrowserMenuExtender>& GetAllCollectionViewContextMenuExtenders() {return CollectionViewContextMenuExtenders;}
@@ -66,7 +66,7 @@ private:
 	TSharedPtr<class FContentBrowserSpawner> ContentBrowserSpawner;
 
 	/** All extender delegates for the content browser menus */
-	TArray<FContentBrowserMenuExtender> AssetContextMenuExtenders;
+	TArray<FContentBrowserMenuExtender_SelectedPaths> AssetContextMenuExtenders;
 	TArray<FContentBrowserMenuExtender_SelectedPaths> PathViewContextMenuExtenders;
 	TArray<FContentBrowserMenuExtender> CollectionListContextMenuExtenders;
 	TArray<FContentBrowserMenuExtender> CollectionViewContextMenuExtenders;
diff --git a/Engine/Source/Editor/EditorStyle/Private/SlateEditorStyle.cpp b/Engine/Source/Editor/EditorStyle/Private/SlateEditorStyle.cpp
--- a/Engine/Source/Editor/EditorStyle/Private/SlateEditorStyle.cpp
+++ b/Engine/Source/Editor/EditorStyle/Private/SlateEditorStyle.cpp
@@ -365,6 +365,7 @@ void FSlateEditorStyle::FStyle::SetupGeneralStyles()
 		Set( "Symbols.DoubleUpArrow", new IMAGE_BRUSH( "Common/UpArrow2", Icon8x8 ) );
 		Set( "Symbols.DownArrow", new IMAGE_BRUSH( "Common/DownArrow", Icon8x8 ) );
 		Set( "Symbols.DoubleDownArrow", new IMAGE_BRUSH( "Common/DownArrow2", Icon8x8 ) );
+		Set( "Symbols.RightArrow", new IMAGE_BRUSH("Common/SubmenuArrow", Icon8x8));
 		Set( "Symbols.Check", new IMAGE_BRUSH( "Common/Check", Icon16x16 ) );
 	}
 
@@ -7108,6 +7109,9 @@ void FSlateEditorStyle::FStyle::SetupAutomationStyles()
 		Set("FBXIcon.ReimportSameContent", new IMAGE_BRUSH("Icons/FBX/FbxReimportSameContent_16x16px", Icon16x16));
 		Set("FBXIcon.ReimportError", new IMAGE_BRUSH("Icons/FBX/FbxReimportError_16x16px", Icon16x16));
 
+		Set("FBXIcon.ReimportCompareAdd", new IMAGE_BRUSH("Icons/FBX/FbxReimportCompare-Add_16x16px", Icon16x16));
+		Set("FBXIcon.ReimportCompareRemoved", new IMAGE_BRUSH("Icons/FBX/FbxReimportCompare-Remove_16x16px", Icon16x16));
+
 		const FTextBlockStyle FBXLargeFont =
 			FTextBlockStyle(NormalText)
 			.SetFont(TTF_CORE_FONT("Fonts/Roboto-Regular", 12))
diff --git a/Engine/Source/Editor/MeshPaint/Private/BaseMeshPaintGeometryAdapter.cpp b/Engine/Source/Editor/MeshPaint/Private/BaseMeshPaintGeometryAdapter.cpp
--- a/Engine/Source/Editor/MeshPaint/Private/BaseMeshPaintGeometryAdapter.cpp
+++ b/Engine/Source/Editor/MeshPaint/Private/BaseMeshPaintGeometryAdapter.cpp
@@ -95,7 +95,7 @@ TArray<uint32> FBaseMeshPaintGeometryAdapter::SphereIntersectTriangles(const flo
 void FBaseMeshPaintGeometryAdapter::GetInfluencedVertexIndices(const float ComponentSpaceSquaredBrushRadius, const FVector& ComponentSpaceBrushPosition, const FVector& ComponentSpaceCameraPosition, const bool bOnlyFrontFacing, TSet<int32> &InfluencedVertices) const
 {
 	// Get a list of (optionally front-facing) triangles that are within a reasonable distance to the brush
-	TArray<uint32> InfluencedTriangles = SphereIntersectTriangles(
+	const TArray<uint32> InfluencedTriangles = SphereIntersectTriangles(
 		ComponentSpaceSquaredBrushRadius,
 		ComponentSpaceBrushPosition,
 		ComponentSpaceCameraPosition,
@@ -106,7 +106,7 @@ void FBaseMeshPaintGeometryAdapter::GetInfluencedVertexIndices(const float Compo
 	check(NumIndexBufferIndices % 3 == 0);
 
 	InfluencedVertices.Reserve(InfluencedTriangles.Num());
-	for (int32 InfluencedTriangle : InfluencedTriangles)
+	for (const int32 InfluencedTriangle : InfluencedTriangles)
 	{
 		for (int32 Index = 0; Index < 3; ++Index)
 		{
diff --git a/Engine/Source/Editor/MeshPaint/Private/MeshPaintHelpers.cpp b/Engine/Source/Editor/MeshPaint/Private/MeshPaintHelpers.cpp
--- a/Engine/Source/Editor/MeshPaint/Private/MeshPaintHelpers.cpp
+++ b/Engine/Source/Editor/MeshPaint/Private/MeshPaintHelpers.cpp
@@ -132,7 +132,7 @@ bool MeshPaintHelpers::PropagateColorsToRawMesh(UStaticMesh* StaticMesh, int32 L
 	else
 	{
 		// If there's no raw mesh data, don't try to do any fixup here
-		if (SrcModel.RawMeshBulkData->IsEmpty())
+		if (SrcModel.RawMeshBulkData->IsEmpty() || ComponentLODInfo.OverrideMapBuildData == nullptr)
 		{
 			return false;
 		}
@@ -1491,11 +1491,15 @@ bool MeshPaintHelpers::ApplyPerVertexPaintAction(FPerVertexPaintActionArgs& InAr
 			Action.ExecuteIfBound(Adapter, VertexIndex);
 		}
 	});*/
-		
-	for (int32 VertexIndex : InfluencedVertices)
+	if (InfluencedVertices.Num())
 	{
-		// Apply the action!
-		Action.ExecuteIfBound(InArgs, VertexIndex);
+		InArgs.Adapter->PreEdit();
+		for (const int32 VertexIndex : InfluencedVertices)
+		{
+		// Apply the action!			
+			Action.ExecuteIfBound(InArgs, VertexIndex);
+		}
+		InArgs.Adapter->PostEdit();
 	}
 
 	return (InfluencedVertices.Num() > 0);
diff --git a/Engine/Source/Editor/MeshPaint/Private/MeshPaintStaticMeshAdapter.cpp b/Engine/Source/Editor/MeshPaint/Private/MeshPaintStaticMeshAdapter.cpp
--- a/Engine/Source/Editor/MeshPaint/Private/MeshPaintStaticMeshAdapter.cpp
+++ b/Engine/Source/Editor/MeshPaint/Private/MeshPaintStaticMeshAdapter.cpp
@@ -104,8 +104,8 @@ bool FMeshPaintGeometryAdapterForStaticMeshes::InitializeVertexData()
 
 void FMeshPaintGeometryAdapterForStaticMeshes::PostEdit()
 {
-	// Recreate all component states using the referenced static mesh
-	TUniquePtr<FStaticMeshComponentRecreateRenderStateContext> RecreateRenderStateContext = MakeUnique<FStaticMeshComponentRecreateRenderStateContext>(ReferencedStaticMesh);
+	// Lighting does not need to be invalidated when mesh painting
+	const bool bUnbuildLighting = false;
 	const bool bUsingInstancedVertexColors = true;
 
 	// Update gpu resource data 
@@ -116,6 +116,9 @@ void FMeshPaintGeometryAdapterForStaticMeshes::PostEdit()
 	}
 	else
 	{
+		// Recreate all component states using the referenced static mesh
+		TUniquePtr<FStaticMeshComponentRecreateRenderStateContext> RecreateRenderStateContext = MakeUnique<FStaticMeshComponentRecreateRenderStateContext>(ReferencedStaticMesh, bUnbuildLighting);
+
 		// Reinitialize the static mesh's resources.
 		ReferencedStaticMesh->InitResources();
 	}
diff --git a/Engine/Source/Editor/MeshPaintMode/Private/PaintModePainter.cpp b/Engine/Source/Editor/MeshPaintMode/Private/PaintModePainter.cpp
--- a/Engine/Source/Editor/MeshPaintMode/Private/PaintModePainter.cpp
+++ b/Engine/Source/Editor/MeshPaintMode/Private/PaintModePainter.cpp
@@ -1023,7 +1023,6 @@ bool FPaintModePainter::PaintInternal(const FVector& InCameraOrigin, const FVect
 
 			if (PaintSettings->PaintMode == EPaintMode::Vertices && MeshAdapter->SupportsVertexPaint())
 			{
-				MeshAdapter->PreEdit();
 
 				FPerVertexPaintActionArgs Args;
 				Args.Adapter = MeshAdapter;
@@ -1033,7 +1032,6 @@ bool FPaintModePainter::PaintInternal(const FVector& InCameraOrigin, const FVect
 				Args.Action = PaintAction;
 
 				bPaintApplied |= MeshPaintHelpers::ApplyPerVertexPaintAction(Args, FPerVertexPaintAction::CreateRaw(this, &FPaintModePainter::ApplyVertexColor, Params));
-				MeshAdapter->PostEdit();
 			}
 			else if (PaintSettings->PaintMode == EPaintMode::Textures&& MeshAdapter->SupportsTexturePaint())
 			{
diff --git a/Engine/Source/Editor/Persona/Private/AnimationEditorPreviewScene.h b/Engine/Source/Editor/Persona/Private/AnimationEditorPreviewScene.h
--- a/Engine/Source/Editor/Persona/Private/AnimationEditorPreviewScene.h
+++ b/Engine/Source/Editor/Persona/Private/AnimationEditorPreviewScene.h
@@ -118,6 +118,24 @@ public:
 	virtual bool AllowMeshHitProxies() const override;
 	virtual void SetAllowMeshHitProxies(bool bState) override;
 
+	virtual void RegisterOnSelectedLODChanged(const FOnSelectedLODChanged &Delegate) override
+	{
+		OnSelectedLODChanged.Add(Delegate);
+	}
+	
+	virtual void UnRegisterOnSelectedLODChanged(void* Thing) override
+	{
+		OnSelectedLODChanged.RemoveAll(Thing);
+	}
+	
+	virtual void BroadcastOnSelectedLODChanged() override
+	{
+		if (OnSelectedLODChanged.IsBound())
+		{
+			OnSelectedLODChanged.Broadcast();
+		}
+	}
+
 	/** FPreviewScene interface */
 	virtual void Tick(float InDeltaTime) override;
 	virtual void AddComponent(class UActorComponent* Component, const FTransform& LocalToWorld, bool bAttachToRoot = false) override;
@@ -280,4 +298,7 @@ private:
 
 	/** Whether or not mesh section hit proxies should be enabled or not */
 	bool bEnableMeshHitProxies;
+
+	/* Selected LOD changed delegate */
+	FOnSelectedLODChangedMulticaster OnSelectedLODChanged;
 };
diff --git a/Engine/Source/Editor/Persona/Private/PersonaMeshDetails.cpp b/Engine/Source/Editor/Persona/Private/PersonaMeshDetails.cpp
--- a/Engine/Source/Editor/Persona/Private/PersonaMeshDetails.cpp
+++ b/Engine/Source/Editor/Persona/Private/PersonaMeshDetails.cpp
@@ -26,6 +26,7 @@
 #include "Widgets/Layout/SUniformGridPanel.h"
 #include "Widgets/Input/SButton.h"
 #include "Widgets/Input/SSpinBox.h"
+#include "Widgets/Input/SComboButton.h"
 #include "Editor/UnrealEdEngine.h"
 #include "EditorDirectories.h"
 #include "UnrealEdGlobals.h"
@@ -975,8 +976,17 @@ void FPersonaMeshDetails::OnPasteMaterialList()
 			Mesh->PreEditChange(Property);
 			FScopedTransaction Transaction(LOCTEXT("PersonaChangedPasteMaterialList", "Pasted material list"));
 			Mesh->Modify();
-
-			FJsonObjectConverter::JsonValueToUProperty(RootJsonValue, Property, &Mesh->Materials, 0, 0);
+			TArray<FSkeletalMaterial> TempMaterials;
+			FJsonObjectConverter::JsonValueToUProperty(RootJsonValue, Property, &TempMaterials, 0, 0);
+			//Do not change the number of material in the array
+			for (int32 MaterialIndex = 0; MaterialIndex < TempMaterials.Num(); ++MaterialIndex)
+			{
+				if (Mesh->Materials.IsValidIndex(MaterialIndex))
+				{
+					Mesh->Materials[MaterialIndex].MaterialInterface = TempMaterials[MaterialIndex].MaterialInterface;
+				}
+			}
+			
 
 			Mesh->PostEditChange();
 		}
@@ -1045,7 +1055,9 @@ void FPersonaMeshDetails::OnPasteMaterialItem(int32 CurrentSlot)
 
 			if (Mesh->Materials.IsValidIndex(CurrentSlot))
 			{
-				FJsonObjectConverter::JsonObjectToUStruct(RootJsonObject.ToSharedRef(), FSkeletalMaterial::StaticStruct(), &Mesh->Materials[CurrentSlot], 0, 0);
+				FSkeletalMaterial TmpSkeletalMaterial;
+				FJsonObjectConverter::JsonObjectToUStruct(RootJsonObject.ToSharedRef(), FSkeletalMaterial::StaticStruct(), &TmpSkeletalMaterial, 0, 0);
+				Mesh->Materials[CurrentSlot].MaterialInterface = TmpSkeletalMaterial.MaterialInterface;
 			}
 
 			Mesh->PostEditChange();
@@ -1072,7 +1084,7 @@ void FPersonaMeshDetails::AddLODLevelCategories(IDetailLayoutBuilder& DetailLayo
 
 		//Create material list panel to let users control the materials array
 		{
-			FString MaterialCategoryName = FString(TEXT("Materials"));
+			FString MaterialCategoryName = FString(TEXT("Material Slots"));
 			IDetailCategoryBuilder& MaterialCategory = DetailLayout.EditCategory(*MaterialCategoryName, FText::GetEmpty(), ECategoryPriority::Important);
 			MaterialCategory.AddCustomRow(LOCTEXT("AddLODLevelCategories_MaterialArrayOperationAdd", "Materials Operation Add Material Slot"))
 				.CopyAction(FUIAction(FExecuteAction::CreateSP(this, &FPersonaMeshDetails::OnCopyMaterialList), FCanExecuteAction::CreateSP(this, &FPersonaMeshDetails::OnCanCopyMaterialList)))
@@ -1138,10 +1150,16 @@ void FPersonaMeshDetails::AddLODLevelCategories(IDetailLayoutBuilder& DetailLayo
 				MaterialListDelegates.OnCanCopyMaterialItem.BindSP(this, &FPersonaMeshDetails::OnCanCopyMaterialItem);
 				MaterialListDelegates.OnPasteMaterialItem.BindSP(this, &FPersonaMeshDetails::OnPasteMaterialItem);
 
-				MaterialCategory.AddCustomBuilder(MakeShareable(new FMaterialList(MaterialCategory.GetParentLayout(), MaterialListDelegates, false)));
+				MaterialCategory.AddCustomBuilder(MakeShareable(new FMaterialList(MaterialCategory.GetParentLayout(), MaterialListDelegates, false, true, true)));
 			}
 		}
 
+		int32 CurrentLodIndex = 0;
+		if (GetPersonaToolkit()->GetPreviewMeshComponent() != nullptr)
+		{
+			CurrentLodIndex = GetPersonaToolkit()->GetPreviewMeshComponent()->ForcedLodModel;
+		}
+		LodCategories.Empty(SkelMeshLODCount);
 		// Create information panel for each LOD level.
 		for (int32 LODIndex = 0; LODIndex < SkelMeshLODCount; ++LODIndex)
 		{
@@ -1182,9 +1200,10 @@ void FPersonaMeshDetails::AddLODLevelCategories(IDetailLayoutBuilder& DetailLayo
 			FString CategoryName = FString(TEXT("LOD"));
 			CategoryName.AppendInt(LODIndex);
 
-			FText LODLevelString = FText::Format(LOCTEXT("LODLevel", "LOD{0}"), FText::AsNumber(LODIndex));
+			FText LODLevelString = FText::FromString(FString(TEXT("LOD Sections (LOD")) + FString::FromInt(LODIndex) + TEXT(")") );
 
 			IDetailCategoryBuilder& LODCategory = DetailLayout.EditCategory(*CategoryName, LODLevelString, ECategoryPriority::Important);
+			LodCategories.Add(&LODCategory);
 			TSharedRef<SWidget> LODCategoryWidget =
 
 				SNew(SBox)
@@ -1204,6 +1223,7 @@ void FPersonaMeshDetails::AddLODLevelCategories(IDetailLayoutBuilder& DetailLayo
 				SectionListDelegates.OnSectionChanged.BindSP(this, &FPersonaMeshDetails::OnSectionChanged);
 				SectionListDelegates.OnGenerateCustomNameWidgets.BindSP(this, &FPersonaMeshDetails::OnGenerateCustomNameWidgetsForSection);
 				SectionListDelegates.OnGenerateCustomSectionWidgets.BindSP(this, &FPersonaMeshDetails::OnGenerateCustomSectionWidgetsForSection);
+				SectionListDelegates.OnGenerateLodComboBox.BindSP(this, &FPersonaMeshDetails::OnGenerateLodComboBoxForSectionList);
 
 				SectionListDelegates.OnCopySectionList.BindSP(this, &FPersonaMeshDetails::OnCopySectionList, LODIndex);
 				SectionListDelegates.OnCanCopySectionList.BindSP(this, &FPersonaMeshDetails::OnCanCopySectionList, LODIndex);
@@ -1212,7 +1232,9 @@ void FPersonaMeshDetails::AddLODLevelCategories(IDetailLayoutBuilder& DetailLayo
 				SectionListDelegates.OnCanCopySectionItem.BindSP(this, &FPersonaMeshDetails::OnCanCopySectionItem);
 				SectionListDelegates.OnPasteSectionItem.BindSP(this, &FPersonaMeshDetails::OnPasteSectionItem);
 
-				LODCategory.AddCustomBuilder(MakeShareable(new FSectionList(LODCategory.GetParentLayout(), SectionListDelegates, (LODIndex > 0))));
+				LODCategory.AddCustomBuilder(MakeShareable(new FSectionList(LODCategory.GetParentLayout(), SectionListDelegates, false, 85, LODIndex)));
+
+				GetPersonaToolkit()->GetPreviewScene()->RegisterOnSelectedLODChanged(FOnSelectedLODChanged::CreateSP(this, &FPersonaMeshDetails::UpdateLODCategoryVisibility));
 			}
 
 			if (LODInfoProperty->IsValidHandle())
@@ -1260,6 +1282,7 @@ void FPersonaMeshDetails::AddLODLevelCategories(IDetailLayoutBuilder& DetailLayo
 					.OnReimportNewFileClicked(this, &FPersonaMeshDetails::OnReimportLodClicked, &DetailLayout, EReimportButtonType::ReimportWithNewFile, LODIndex)
 				];
 			}
+			LODCategory.SetCategoryVisibility( (CurrentLodIndex == 0 ? 0 : CurrentLodIndex - 1) == LODIndex);
 		}
 	}
 }
@@ -2230,36 +2253,7 @@ TSharedRef<SWidget> FPersonaMeshDetails::OnGenerateCustomSectionWidgetsForSectio
 {
 	extern ENGINE_API bool IsGPUSkinCacheAvailable();
 
-	TSharedRef<SVerticalBox> SectionWidget
-		= SNew(SVerticalBox)
-
-		+ SVerticalBox::Slot()
-		.Padding(0, 2, 0, 0)
-		[
-			SNew(SCheckBox)
-			.IsChecked(this, &FPersonaMeshDetails::IsSectionShadowCastingEnabled, LODIndex, SectionIndex)
-			.OnCheckStateChanged(this, &FPersonaMeshDetails::OnSectionShadowCastingChanged, LODIndex, SectionIndex)
-			[
-				SNew(STextBlock)
-				.Font(FEditorStyle::GetFontStyle("StaticMeshEditor.NormalFont"))
-				.Text(LOCTEXT("Cast Shadows", "Cast Shadows"))
-			]
-		]
-
-		+ SVerticalBox::Slot()
-		.Padding(0, 2, 0, 0)
-		[
-			SNew(SCheckBox)
-			.IsEnabled(IsGPUSkinCacheAvailable())
-			.IsChecked(this, &FPersonaMeshDetails::IsSectionRecomputeTangentEnabled, LODIndex, SectionIndex)
-			.OnCheckStateChanged(this, &FPersonaMeshDetails::OnSectionRecomputeTangentChanged, LODIndex, SectionIndex)
-			[
-				SNew(STextBlock)
-				.Font(FEditorStyle::GetFontStyle("StaticMeshEditor.NormalFont"))
-				.Text(LOCTEXT("RecomputeTangent_Title", "Recompute Tangent"))
-				.ToolTipText(LOCTEXT("RecomputeTangent_Tooltip", "This feature only works if you enable (Support Skincache Shaders) in the Project Settings. Please note that skin cache is an experimental feature and only works if you have compute shaders."))
-			]
-		];
+	TSharedRef<SVerticalBox> SectionWidget = SNew(SVerticalBox);
 
 #if WITH_APEX_CLOTHING
 
@@ -2270,19 +2264,25 @@ TSharedRef<SWidget> FPersonaMeshDetails::OnGenerateCustomSectionWidgetsForSectio
 	SectionWidget->AddSlot()
 	.AutoHeight()
 	.Padding(0, 2, 0, 0)
+	.HAlign(HAlign_Fill)
 	[
 	SNew(SHorizontalBox)
 	+ SHorizontalBox::Slot()
 		.AutoWidth()
 		.VAlign(VAlign_Center)
+		[
+			SNew(SBox)
+			.HAlign(HAlign_Right)
+			.MinDesiredWidth(65.0f)
 		[
 			SNew(STextBlock)
+				.Font(IDetailLayoutBuilder::GetDetailFont())
 			.Text(LOCTEXT("Clothing", "Clothing"))
 		]
-
+		]
 		+ SHorizontalBox::Slot()
-		.AutoWidth()
-		.Padding(2, 2, 0, 0)
+		.FillWidth(1.0f)
+		.Padding(5, 2, 0, 0)
 		[
 			SAssignNew(ClothComboBoxes.Last(), SClothComboBox)
 			.OnGenerateWidget(this, &FPersonaMeshDetails::OnGenerateWidgetForClothingEntry)
@@ -2296,8 +2296,165 @@ TSharedRef<SWidget> FPersonaMeshDetails::OnGenerateCustomSectionWidgetsForSectio
 	];
 
 #endif// #if WITH_APEX_CLOTHING
+	SectionWidget->AddSlot()
+	.AutoHeight()
+	.Padding(0, 2, 0, 0)
+	[
+		SNew(SHorizontalBox)
+		+SHorizontalBox::Slot()
+		.AutoWidth()
+		.Padding(2, 0, 2, 0)
+		[
+			SNew(SCheckBox)
+			.IsChecked(this, &FPersonaMeshDetails::IsSectionShadowCastingEnabled, LODIndex, SectionIndex)
+			.OnCheckStateChanged(this, &FPersonaMeshDetails::OnSectionShadowCastingChanged, LODIndex, SectionIndex)
+			[
+				SNew(STextBlock)
+				.Font(IDetailLayoutBuilder::GetDetailFont())
+				.Text(LOCTEXT("Cast Shadows", "Cast Shadows"))
+			]
+		]
+		+SHorizontalBox::Slot()
+		.AutoWidth()
+		.Padding(2, 0, 2, 0)
+		[
+			SNew(SCheckBox)
+			.IsEnabled(IsGPUSkinCacheAvailable())
+			.IsChecked(this, &FPersonaMeshDetails::IsSectionRecomputeTangentEnabled, LODIndex, SectionIndex)
+			.OnCheckStateChanged(this, &FPersonaMeshDetails::OnSectionRecomputeTangentChanged, LODIndex, SectionIndex)
+			[
+				SNew(STextBlock)
+				.Font(IDetailLayoutBuilder::GetDetailFont())
+				.Text(LOCTEXT("RecomputeTangent_Title", "Recompute Tangent"))
+				.ToolTipText(LOCTEXT("RecomputeTangent_Tooltip", "This feature only works if you enable (Support Skincache Shaders) in the Project Settings. Please note that skin cache is an experimental feature and only works if you have compute shaders."))
+			]
+		]
+	];
 	return SectionWidget;
 }
+
+void FPersonaMeshDetails::SetCurrentLOD(int32 NewLodIndex)
+{
+	if (GetPersonaToolkit()->GetPreviewMeshComponent() == nullptr)
+	{
+		return;
+	}
+	int32 CurrentDisplayLOD = GetPersonaToolkit()->GetPreviewMeshComponent()->ForcedLodModel;
+	int32 RealCurrentDisplayLOD = CurrentDisplayLOD == 0 ? 0 : CurrentDisplayLOD - 1;
+	int32 RealNewLOD = NewLodIndex == 0 ? 0 : NewLodIndex - 1;
+	if (CurrentDisplayLOD == NewLodIndex || !LodCategories.IsValidIndex(RealCurrentDisplayLOD) || !LodCategories.IsValidIndex(RealNewLOD))
+	{
+		return;
+	}
+	LodCategories[RealCurrentDisplayLOD]->SetCategoryVisibility(false);
+	LodCategories[RealNewLOD]->SetCategoryVisibility(true);
+	GetPersonaToolkit()->GetPreviewMeshComponent()->SetForcedLOD(NewLodIndex);
+	
+	//Reset the preview section since we do not edit the same LOD
+	if (GetPersonaToolkit()->GetMesh() != nullptr)
+	{
+		GetPersonaToolkit()->GetPreviewMeshComponent()->SetSectionPreview(INDEX_NONE);
+		GetPersonaToolkit()->GetMesh()->SelectedEditorSection = INDEX_NONE;
+	}
+
+	GetPersonaToolkit()->GetPreviewScene()->BroadcastOnSelectedLODChanged();
+}
+
+void FPersonaMeshDetails::UpdateLODCategoryVisibility() const
+{
+	bool bAutoLod = false;
+	if (GetPersonaToolkit()->GetPreviewMeshComponent() != nullptr)
+	{
+		bAutoLod = GetPersonaToolkit()->GetPreviewMeshComponent()->ForcedLodModel == 0;
+	}
+	int32 CurrentDisplayLOD = bAutoLod ? 0 : GetPersonaToolkit()->GetPreviewMeshComponent()->ForcedLodModel - 1;
+	if (LodCategories.IsValidIndex(CurrentDisplayLOD) && GetPersonaToolkit()->GetMesh())
+	{
+		int32 SkeletalMeshLodNumber = GetPersonaToolkit()->GetMesh()->LODInfo.Num();
+		for (int32 LodCategoryIndex = 0; LodCategoryIndex < SkeletalMeshLodNumber; ++LodCategoryIndex)
+		{
+			LodCategories[LodCategoryIndex]->SetCategoryVisibility(CurrentDisplayLOD == LodCategoryIndex);
+		}
+	}
+
+	//Reset the preview section since we do not edit the same LOD
+	if (GetPersonaToolkit()->GetMesh() != nullptr)
+	{
+		GetPersonaToolkit()->GetPreviewMeshComponent()->SetSectionPreview(INDEX_NONE);
+		GetPersonaToolkit()->GetMesh()->SelectedEditorSection = INDEX_NONE;
+	}
+}
+
+FText FPersonaMeshDetails::GetCurrentLodName() const
+{
+	bool bAutoLod = false;
+	if (GetPersonaToolkit()->GetPreviewMeshComponent() != nullptr)
+	{
+		bAutoLod = GetPersonaToolkit()->GetPreviewMeshComponent()->ForcedLodModel == 0;
+	}
+	int32 CurrentDisplayLOD = bAutoLod ? 0 : GetPersonaToolkit()->GetPreviewMeshComponent()->ForcedLodModel - 1;
+	return FText::FromString(bAutoLod ? FString(TEXT("Auto (LOD0)")) : (FString(TEXT("LOD")) + FString::FromInt(CurrentDisplayLOD)));
+}
+
+FText FPersonaMeshDetails::GetCurrentLodTooltip() const
+{
+	if (GetPersonaToolkit()->GetPreviewMeshComponent() != nullptr && GetPersonaToolkit()->GetPreviewMeshComponent()->ForcedLodModel == 0)
+	{
+		return FText::FromString(TEXT("LOD0 is edit when selecting Auto LOD"));
+	}
+	return FText::GetEmpty();
+}
+
+TSharedRef<SWidget> FPersonaMeshDetails::OnGenerateLodComboBoxForSectionList(int32 LodIndex)
+{
+	return SNew(SComboButton)
+	.OnGetMenuContent(this, &FPersonaMeshDetails::OnGenerateLodMenuForSectionList, LodIndex)
+	.VAlign(VAlign_Center)
+	.ContentPadding(2)
+	.ButtonContent()
+	[
+		SNew(STextBlock)
+		.Font(IDetailLayoutBuilder::GetDetailFont())
+		.Text(this, &FPersonaMeshDetails::GetCurrentLodName)
+		.ToolTipText(this, &FPersonaMeshDetails::GetCurrentLodTooltip)
+	];
+}
+
+TSharedRef<SWidget> FPersonaMeshDetails::OnGenerateLodMenuForSectionList(int32 LodIndex)
+{
+	USkeletalMesh* SkelMesh = GetPersonaToolkit()->GetMesh();
+
+	if (SkelMesh == nullptr)
+	{
+		return SNullWidget::NullWidget;
+	}
+
+	bool bAutoLod = false;
+	if (GetPersonaToolkit()->GetPreviewMeshComponent() != nullptr)
+	{
+		bAutoLod = GetPersonaToolkit()->GetPreviewMeshComponent()->ForcedLodModel == 0;
+	}
+	const int32 SkelMeshLODCount = SkelMesh->LODInfo.Num();
+	if(SkelMeshLODCount < 2)
+	{
+		return SNullWidget::NullWidget;
+	}
+	FMenuBuilder MenuBuilder(true, NULL);
+
+	FText AutoLodText = FText::FromString((TEXT("Auto LOD")));
+	FUIAction AutoLodAction(FExecuteAction::CreateSP(this, &FPersonaMeshDetails::SetCurrentLOD, 0));
+	MenuBuilder.AddMenuEntry(AutoLodText, LOCTEXT("OnGenerateLodMenuForSectionList_Auto_ToolTip", "LOD0 is edit when selecting Auto LOD"), FSlateIcon(), AutoLodAction);
+	// Add a menu item for each texture.  Clicking on the texture will display it in the content browser
+	for (int32 AllLodIndex = 0; AllLodIndex < SkelMeshLODCount; ++AllLodIndex)
+	{
+		FText LODLevelString = FText::FromString((TEXT("LOD ") + FString::FromInt(AllLodIndex)));
+		FUIAction Action(FExecuteAction::CreateSP(this, &FPersonaMeshDetails::SetCurrentLOD, AllLodIndex+1));
+		MenuBuilder.AddMenuEntry(LODLevelString, FText::GetEmpty(), FSlateIcon(), Action);
+	}
+
+	return MenuBuilder.MakeWidget();
+}
+
 ECheckBoxState FPersonaMeshDetails::IsMaterialSelected(int32 MaterialIndex) const
 {
 	ECheckBoxState State = ECheckBoxState::Unchecked;
diff --git a/Engine/Source/Editor/Persona/Private/PersonaMeshDetails.h b/Engine/Source/Editor/Persona/Private/PersonaMeshDetails.h
--- a/Engine/Source/Editor/Persona/Private/PersonaMeshDetails.h
+++ b/Engine/Source/Editor/Persona/Private/PersonaMeshDetails.h
@@ -211,6 +211,17 @@ private:
 	 */
 	TSharedRef<SWidget> OnGenerateCustomSectionWidgetsForSection(int32 LODIndex, int32 SectionIndex);
 
+	TSharedRef<SWidget> OnGenerateLodComboBoxForSectionList(int32 LodIndex);
+	/*
+	 * Generate the context menu to choose the LOD we will display the section list
+	*/
+	TSharedRef<SWidget> OnGenerateLodMenuForSectionList(int32 LodIndex);
+	void UpdateLODCategoryVisibility() const;
+	FText GetCurrentLodName() const;
+	FText GetCurrentLodTooltip() const;
+
+	void SetCurrentLOD(int32 NewLodIndex);
+
 	FText GetMaterialNameText(int32 MaterialIndex)const ;
 	void OnMaterialNameCommitted(const FText& InValue, ETextCommit::Type CommitType, int32 MaterialIndex);
 	void OnMaterialNameChanged(const FText& InValue, int32 MaterialIndex);
@@ -481,6 +492,8 @@ private:
 	/* This is to know if material are used by any LODs sections. */
 	TMap<int32, TArray<FSectionLocalizer>> MaterialUsedMap;
 
+	TArray<class IDetailCategoryBuilder*> LodCategories;
+
 #if WITH_APEX_CLOTHING
 private:
 
diff --git a/Engine/Source/Editor/Persona/Private/SAnimationEditorViewport.cpp b/Engine/Source/Editor/Persona/Private/SAnimationEditorViewport.cpp
--- a/Engine/Source/Editor/Persona/Private/SAnimationEditorViewport.cpp
+++ b/Engine/Source/Editor/Persona/Private/SAnimationEditorViewport.cpp
@@ -308,7 +308,6 @@ void SAnimationEditorViewportTabBody::Construct(const FArguments& InArgs, const
 	BlueprintEditorPtr = InArgs._BlueprintEditor;
 	bShowTimeline = InArgs._ShowTimeline;
 	OnInvokeTab = InArgs._OnInvokeTab;
-	LODSelection = 0;
 
 	// register delegates for change notifications
 	InPreviewScene->RegisterOnAnimChanged(FOnAnimChanged::CreateSP(this, &SAnimationEditorViewportTabBody::AnimChanged));
@@ -680,6 +679,7 @@ void SAnimationEditorViewportTabBody::BindCommands()
 
 #endif// #if WITH_APEX_CLOTHING		
 
+	GetPreviewScene()->RegisterOnSelectedLODChanged(FOnSelectedLODChanged::CreateSP(this, &SAnimationEditorViewportTabBody::OnLODModelChanged));
 	//Bind LOD preview menu commands
 	const FAnimViewportLODCommands& ViewportLODMenuCommands = FAnimViewportLODCommands::Get();
 
@@ -1325,19 +1325,48 @@ void SAnimationEditorViewportTabBody::UpdateShowFlagForMeshEdges()
 	LevelViewportClient->EngineShowFlags.SetMeshEdges(bUseOverlayMaterial || bShowMeshEdgesViewMode);
 }
 
+int32 SAnimationEditorViewportTabBody::GetLODSelection() const
+{
+	UDebugSkelMeshComponent* PreviewComponent = GetPreviewScene()->GetPreviewMeshComponent();
+
+	if (PreviewComponent)
+	{
+		return PreviewComponent->ForcedLodModel;
+	}
+	return 0;
+}
+
 bool SAnimationEditorViewportTabBody::IsLODModelSelected(int32 LODSelectionType) const
 {
-	return (LODSelection == LODSelectionType) ? true : false;
+	UDebugSkelMeshComponent* PreviewComponent = GetPreviewScene()->GetPreviewMeshComponent();
+
+	if (PreviewComponent)
+	{
+		return (PreviewComponent->ForcedLodModel == LODSelectionType) ? true : false;
+	}
+	return false;
 }
 
 void SAnimationEditorViewportTabBody::OnSetLODModel(int32 LODSelectionType)
 {
-	LODSelection = LODSelectionType;
-
 	UDebugSkelMeshComponent* PreviewComponent = GetPreviewScene()->GetPreviewMeshComponent();
+	
 	if( PreviewComponent )
 	{
-		PreviewComponent->ForcedLodModel = LODSelection;
+		LODSelection = LODSelectionType;
+		PreviewComponent->ForcedLodModel = LODSelectionType;
+		PopulateUVChoices();
+		GetPreviewScene()->BroadcastOnSelectedLODChanged();
+	}
+}
+
+void SAnimationEditorViewportTabBody::OnLODModelChanged()
+{
+	UDebugSkelMeshComponent* PreviewComponent = GetPreviewScene()->GetPreviewMeshComponent();
+
+	if (PreviewComponent && LODSelection != PreviewComponent->ForcedLodModel)
+	{
+		LODSelection = PreviewComponent->ForcedLodModel;
 		PopulateUVChoices();
 	}
 }
diff --git a/Engine/Source/Editor/Persona/Private/SAnimationEditorViewport.h b/Engine/Source/Editor/Persona/Private/SAnimationEditorViewport.h
--- a/Engine/Source/Editor/Persona/Private/SAnimationEditorViewport.h
+++ b/Engine/Source/Editor/Persona/Private/SAnimationEditorViewport.h
@@ -121,7 +121,6 @@ public:
 	/** SWidget interface */
 	virtual FReply OnKeyDown(const FGeometry& MyGeometry, const FKeyEvent& InKeyEvent) override;
 
-
 	void RefreshViewport();
 
 	/**
@@ -137,7 +136,7 @@ public:
 
 	/** LOD model selection checking function*/
 	bool IsLODModelSelected( int32 LODSelectionType ) const;
-	int32 GetLODSelection() const { return LODSelection; };
+	int32 GetLODSelection() const;
 
 	/** Function to set the current playback speed*/
 	void OnSetPlaybackSpeed(int32 PlaybackSpeedMode);
@@ -198,6 +197,7 @@ public:
 
 	/** Function to set LOD model selection*/
 	void OnSetLODModel(int32 LODSelectionType);
+	void OnLODModelChanged();
 
 	/** Get the skeleton tree we are bound to */
 	TSharedRef<class ISkeletonTree> GetSkeletonTree() const { return SkeletonTreePtr.Pin().ToSharedRef(); }
diff --git a/Engine/Source/Editor/Persona/Public/IPersonaPreviewScene.h b/Engine/Source/Editor/Persona/Public/IPersonaPreviewScene.h
--- a/Engine/Source/Editor/Persona/Public/IPersonaPreviewScene.h
+++ b/Engine/Source/Editor/Persona/Public/IPersonaPreviewScene.h
@@ -29,6 +29,10 @@ DECLARE_MULTICAST_DELEGATE_TwoParams(FOnMeshClickMulticaster, HActor*, const FVi
 
 typedef FOnMeshClickMulticaster::FDelegate FOnMeshClick;
 
+//The selected LOD changed
+DECLARE_MULTICAST_DELEGATE(FOnSelectedLODChangedMulticaster);
+typedef FOnSelectedLODChangedMulticaster::FDelegate FOnSelectedLODChanged;
+
 /** Modes that the preview scene defaults to (usually depending on asset editor context) */
 enum class EPreviewSceneDefaultAnimationMode : int32
 {
@@ -190,4 +194,11 @@ public:
 
 	/** Set whether or not to ignore mesh hit proxies */
 	virtual void SetAllowMeshHitProxies(bool bState) = 0;
+
+	/** Register callback to be able to be notify when the select LOD is change */
+	virtual void RegisterOnSelectedLODChanged(const FOnSelectedLODChanged &Delegate) = 0;
+	/** Unregister callback to free up the ressources */
+	virtual void UnRegisterOnSelectedLODChanged(void* Thing) = 0;
+	/** Broadcast select LOD changed */
+	virtual void BroadcastOnSelectedLODChanged() = 0;
 };
diff --git a/Engine/Source/Editor/PropertyEditor/Private/PropertyCustomizationHelpers.cpp b/Engine/Source/Editor/PropertyEditor/Private/PropertyCustomizationHelpers.cpp
--- a/Engine/Source/Editor/PropertyEditor/Private/PropertyCustomizationHelpers.cpp
+++ b/Engine/Source/Editor/PropertyEditor/Private/PropertyCustomizationHelpers.cpp
@@ -707,9 +707,10 @@ public:
 		FOnGenerateWidgetsForMaterial InOnGenerateWidgetsForMaterial, 
 		FOnResetMaterialToDefaultClicked InOnResetToDefaultClicked,
 		int32 InMultipleMaterialCount,
-		bool bShowUsedTextures)
+		bool bShowUsedTextures,
+		bool bDisplayCompactSize)
 	{
-		return MakeShareable( new FMaterialItemView( Material, InOnMaterialChanged, InOnGenerateNameWidgetsForMaterial, InOnGenerateWidgetsForMaterial, InOnResetToDefaultClicked, InMultipleMaterialCount, bShowUsedTextures) );
+		return MakeShareable( new FMaterialItemView( Material, InOnMaterialChanged, InOnGenerateNameWidgetsForMaterial, InOnGenerateWidgetsForMaterial, InOnResetToDefaultClicked, InMultipleMaterialCount, bShowUsedTextures, bDisplayCompactSize) );
 	}
 
 	TSharedRef<SWidget> CreateNameContent()
@@ -757,24 +758,38 @@ public:
 						.OnSetObject(this, &FMaterialItemView::OnSetObject)
 						.DisplayThumbnail(true)
 						.ThumbnailPool(ThumbnailPool)
+						.ThumbnailSize(bDisplayCompactSize ? FIntPoint(48, 48) : FIntPoint(64, 64) )
+						.DisplayCompactSize(bDisplayCompactSize)
 						.CustomContentSlot()
 						[
 							SNew( SBox )
 							.HAlign(HAlign_Left)
 							[
-								// Add a menu for displaying all textures 
-								SNew( SComboButton )
-								.OnGetMenuContent( this, &FMaterialItemView::OnGetTexturesMenuForMaterial )
-								.VAlign( VAlign_Center )
-								.ContentPadding(2)
-								.IsEnabled( this, &FMaterialItemView::IsTexturesMenuEnabled )
-								.Visibility( bShowUsedTextures ? EVisibility::Visible : EVisibility::Hidden )
-								.ButtonContent()
+								SNew(SHorizontalBox)
+								+SHorizontalBox::Slot()
+								.Padding(0.0f, 0.0f, 3.0f, 0.0f)
+								.AutoWidth()
 								[
-									SNew( STextBlock )
-									.Font( IDetailLayoutBuilder::GetDetailFont() )
-									.ToolTipText( LOCTEXT("ViewTexturesToolTip", "View the textures used by this material" ) )
-									.Text( LOCTEXT("ViewTextures","Textures") )
+									// Add a menu for displaying all textures 
+									SNew( SComboButton )
+									.OnGetMenuContent( this, &FMaterialItemView::OnGetTexturesMenuForMaterial )
+									.VAlign( VAlign_Center )
+									.ContentPadding(2)
+									.IsEnabled( this, &FMaterialItemView::IsTexturesMenuEnabled )
+									.Visibility( bShowUsedTextures ? EVisibility::Visible : EVisibility::Hidden )
+									.ButtonContent()
+									[
+										SNew( STextBlock )
+										.Font( IDetailLayoutBuilder::GetDetailFont() )
+										.ToolTipText( LOCTEXT("ViewTexturesToolTip", "View the textures used by this material" ) )
+										.Text( LOCTEXT("ViewTextures","Textures") )
+									]
+								]
+								+SHorizontalBox::Slot()
+								.Padding(3.0f, 0.0f)
+								.FillWidth(1.0f)
+								[
+									OnGenerateCustomMaterialWidgets.IsBound() && bDisplayCompactSize ? OnGenerateCustomMaterialWidgets.Execute(MaterialItem.Material.Get(), MaterialItem.SlotIndex) : StaticCastSharedRef<SWidget>(SNullWidget::NullWidget)
 								]
 							]
 						]
@@ -802,7 +817,7 @@ public:
 				.Padding(2)
 				.VAlign( VAlign_Center )
 				[
-					OnGenerateCustomMaterialWidgets.IsBound() ? OnGenerateCustomMaterialWidgets.Execute( MaterialItem.Material.Get(), MaterialItem.SlotIndex ) : StaticCastSharedRef<SWidget>( SNullWidget::NullWidget )
+					OnGenerateCustomMaterialWidgets.IsBound() && !bDisplayCompactSize ? OnGenerateCustomMaterialWidgets.Execute( MaterialItem.Material.Get(), MaterialItem.SlotIndex ) : StaticCastSharedRef<SWidget>( SNullWidget::NullWidget )
 				]
 			];
 
@@ -816,7 +831,8 @@ private:
 						FOnGenerateWidgetsForMaterial& InOnGenerateMaterialWidgets, 
 						FOnResetMaterialToDefaultClicked& InOnResetToDefaultClicked,
 						int32 InMultipleMaterialCount,
-						bool bInShowUsedTextures)
+						bool bInShowUsedTextures,
+						bool bInDisplayCompactSize)
 						
 		: MaterialItem( InMaterial )
 		, OnMaterialChanged( InOnMaterialChanged )
@@ -825,6 +841,7 @@ private:
 		, OnResetToDefaultClicked( InOnResetToDefaultClicked )
 		, MultipleMaterialCount( InMultipleMaterialCount )
 		, bShowUsedTextures( bInShowUsedTextures )
+		, bDisplayCompactSize(bInDisplayCompactSize)
 	{
 
 	}
@@ -936,15 +953,17 @@ private:
 	FOnResetMaterialToDefaultClicked OnResetToDefaultClicked;
 	int32 MultipleMaterialCount;
 	bool bShowUsedTextures;
+	bool bDisplayCompactSize;
 };
 
 
-FMaterialList::FMaterialList(IDetailLayoutBuilder& InDetailLayoutBuilder, FMaterialListDelegates& InMaterialListDelegates, bool bInAllowCollapse, bool bInShowUsedTextures)
+FMaterialList::FMaterialList(IDetailLayoutBuilder& InDetailLayoutBuilder, FMaterialListDelegates& InMaterialListDelegates, bool bInAllowCollapse, bool bInShowUsedTextures, bool bInDisplayCompactSize)
 	: MaterialListDelegates( InMaterialListDelegates )
 	, DetailLayoutBuilder( InDetailLayoutBuilder )
 	, MaterialListBuilder( new FMaterialListBuilder )
 	, bAllowCollpase(bInAllowCollapse)
 	, bShowUsedTextures(bInShowUsedTextures)
+	, bDisplayCompactSize(bInDisplayCompactSize)
 {
 }
 
@@ -1180,7 +1199,7 @@ void FMaterialList::AddMaterialItem( FDetailWidgetRow& Row, int32 CurrentSlot, c
 {
 	uint32 NumMaterials = MaterialListBuilder->GetNumMaterialsInSlot(CurrentSlot);
 
-	TSharedRef<FMaterialItemView> NewView = FMaterialItemView::Create( Item, MaterialListDelegates.OnMaterialChanged, MaterialListDelegates.OnGenerateCustomNameWidgets, MaterialListDelegates.OnGenerateCustomMaterialWidgets, MaterialListDelegates.OnResetMaterialToDefaultClicked, NumMaterials, bShowUsedTextures );
+	TSharedRef<FMaterialItemView> NewView = FMaterialItemView::Create( Item, MaterialListDelegates.OnMaterialChanged, MaterialListDelegates.OnGenerateCustomNameWidgets, MaterialListDelegates.OnGenerateCustomMaterialWidgets, MaterialListDelegates.OnResetMaterialToDefaultClicked, NumMaterials, bShowUsedTextures, bDisplayCompactSize);
 
 	TSharedPtr<SWidget> RightSideContent;
 	if( bDisplayLink )
@@ -1432,7 +1451,11 @@ class FSectionListBuilder : public ISectionListBuilder
 {
 	friend class FSectionList;
 public:
-
+	
+	FSectionListBuilder(int32 InThumbnailSize)
+		:ThumbnailSize(InThumbnailSize)
+	{}
+	
 	/**
 	* Adds a new Section to the list
 	*
@@ -1442,7 +1465,7 @@ public:
 	*/
 	virtual void AddSection(int32 LodIndex, int32 SectionIndex, FName InMaterialSlotName, int32 InMaterialSlotIndex, FName InOriginalMaterialSlotName, const TMap<int32, FName> &InAvailableMaterialSlotName, const UMaterialInterface* Material, bool IsSectionUsingCloth) override
 	{
-		FSectionListItem SectionItem(LodIndex, SectionIndex, InMaterialSlotName, InMaterialSlotIndex, InOriginalMaterialSlotName, InAvailableMaterialSlotName, Material, IsSectionUsingCloth);
+		FSectionListItem SectionItem(LodIndex, SectionIndex, InMaterialSlotName, InMaterialSlotIndex, InOriginalMaterialSlotName, InAvailableMaterialSlotName, Material, IsSectionUsingCloth, ThumbnailSize);
 		if (!Sections.Contains(SectionItem))
 		{
 			Sections.Add(SectionItem);
@@ -1507,6 +1530,8 @@ private:
 	TArray<FSectionListItem> Sections;
 	/** All Section items in the list */
 	TMap<int32, TArray<FSectionListItem>> SectionsByLOD;
+
+	int32 ThumbnailSize;
 };
 
 
@@ -1528,9 +1553,10 @@ public:
 		FOnGenerateWidgetsForSection InOnGenerateNameWidgetsForSection,
 		FOnGenerateWidgetsForSection InOnGenerateWidgetsForSection,
 		FOnResetSectionToDefaultClicked InOnResetToDefaultClicked,
-		int32 InMultipleSectionCount)
+		int32 InMultipleSectionCount,
+		int32 InThumbnailSize)
 	{
-		return MakeShareable(new FSectionItemView(Section, InOnSectionChanged, InOnGenerateNameWidgetsForSection, InOnGenerateWidgetsForSection, InOnResetToDefaultClicked, InMultipleSectionCount));
+		return MakeShareable(new FSectionItemView(Section, InOnSectionChanged, InOnGenerateNameWidgetsForSection, InOnGenerateWidgetsForSection, InOnResetToDefaultClicked, InMultipleSectionCount, InThumbnailSize));
 	}
 
 	TSharedRef<SWidget> CreateNameContent()
@@ -1560,6 +1586,7 @@ public:
 		return
 			SNew(SHorizontalBox)
 			+ SHorizontalBox::Slot()
+			.HAlign(HAlign_Fill)
 			[
 				SNew(SVerticalBox)
 				+ SVerticalBox::Slot()
@@ -1576,62 +1603,74 @@ public:
 						.ObjectPath(SectionItem.Material->GetPathName())
 						.Class(UMaterialInterface::StaticClass())
 						.DisplayThumbnail(true)
-						.ThumbnailSize(FIntPoint(32,32))
+						.ThumbnailSize(FIntPoint(ThumbnailSize, ThumbnailSize))
 						.DisplayUseSelected(false)
 						.AllowClear(false)
 						.DisplayBrowse(false)
 						.EnableContentPicker(false)
-
 						.ThumbnailPool(ThumbnailPool)
-					]
-				]
-				+ SVerticalBox::Slot()
-				.AutoHeight()
-				.Padding(2)
-				.VAlign(VAlign_Center)
-				[
-					SNew(SHorizontalBox)
-					+ SHorizontalBox::Slot()
-					.Padding(0)
-					.VAlign(VAlign_Center)
-					.AutoWidth()
-					[
-						SNew(STextBlock)
-						.Font(IDetailLayoutBuilder::GetDetailFont())
-						.Text(LOCTEXT("SectionListItemMaterialSlotNameLabel", "Material Slot"))
-						.ToolTipText(MaterialSlotNameTooltipText)
-					]
-					+ SHorizontalBox::Slot()
-					.VAlign(VAlign_Center)
-					.AutoWidth()
-					.Padding(5, 0, 0, 0)
-					[
-						SNew(SBox)
-						.HAlign(HAlign_Left)
-						.VAlign(VAlign_Center)
+						.DisplayCompactSize(true)
+						.CustomContentSlot()
 						[
-							//Material Slot Name
-							SNew(SComboButton)
-							.OnGetMenuContent(this, &FSectionItemView::OnGetMaterialSlotNameMenuForSection)
-							.VAlign(VAlign_Center)
-							.ContentPadding(2)
-							.IsEnabled(!SectionItem.IsSectionUsingCloth)
-							.ButtonContent()
+							SNew( SBox )
+							.HAlign(HAlign_Fill)
 							[
-								SNew(STextBlock)
-								.Font(IDetailLayoutBuilder::GetDetailFont())
-								.Text(this, &FSectionItemView::GetCurrentMaterialSlotName)
-								.ToolTipText(MaterialSlotNameTooltipText)
+								SNew(SVerticalBox)
+								+SVerticalBox::Slot()
+								.AutoHeight()
+								[
+									SNew(SHorizontalBox)
+									+ SHorizontalBox::Slot()
+									.Padding(0)
+									.VAlign(VAlign_Center)
+									.AutoWidth()
+									[
+										SNew(SBox)
+										.HAlign(HAlign_Right)
+										.MinDesiredWidth(65.0f)
+										[
+											SNew(STextBlock)
+											.Font(IDetailLayoutBuilder::GetDetailFont())
+											.Text(LOCTEXT("SectionListItemMaterialSlotNameLabel", "Material Slot"))
+											.ToolTipText(MaterialSlotNameTooltipText)
+										]
+									]
+									+ SHorizontalBox::Slot()
+									.VAlign(VAlign_Center)
+									.FillWidth(1.0f)
+									.Padding(5, 0, 0, 0)
+									[
+										SNew(SBox)
+										.HAlign(HAlign_Fill)
+										.VAlign(VAlign_Center)
+										.MinDesiredWidth(210.0f)
+										[
+											//Material Slot Name
+											SNew(SComboButton)
+											.OnGetMenuContent(this, &FSectionItemView::OnGetMaterialSlotNameMenuForSection)
+											.VAlign(VAlign_Center)
+											.ContentPadding(2)
+											.IsEnabled(!SectionItem.IsSectionUsingCloth)
+											.ButtonContent()
+											[
+												SNew(STextBlock)
+												.Font(IDetailLayoutBuilder::GetDetailFont())
+												.Text(this, &FSectionItemView::GetCurrentMaterialSlotName)
+												.ToolTipText(MaterialSlotNameTooltipText)
+											]
+										]
+									]
+								]
+								+SVerticalBox::Slot()
+								.AutoHeight()
+								.VAlign(VAlign_Center)
+								[
+									OnGenerateCustomSectionWidgets.IsBound() ? OnGenerateCustomSectionWidgets.Execute(SectionItem.LodIndex, SectionItem.SectionIndex) : StaticCastSharedRef<SWidget>(SNullWidget::NullWidget)
+								]
 							]
 						]
 					]
 				]
-				+ SVerticalBox::Slot()
-				.AutoHeight()
-				.VAlign(VAlign_Center)
-				[
-					OnGenerateCustomSectionWidgets.IsBound() ? OnGenerateCustomSectionWidgets.Execute(SectionItem.LodIndex, SectionItem.SectionIndex) : StaticCastSharedRef<SWidget>(SNullWidget::NullWidget)
-				]
 			];
 	}
 
@@ -1642,13 +1681,15 @@ private:
 		FOnGenerateWidgetsForSection& InOnGenerateNameWidgets,
 		FOnGenerateWidgetsForSection& InOnGenerateSectionWidgets,
 		FOnResetSectionToDefaultClicked& InOnResetToDefaultClicked,
-		int32 InMultipleSectionCount)
+		int32 InMultipleSectionCount,
+		int32 InThumbnailSize)
 		: SectionItem(InSection)
 		, OnSectionChanged(InOnSectionChanged)
 		, OnGenerateCustomNameWidgets(InOnGenerateNameWidgets)
 		, OnGenerateCustomSectionWidgets(InOnGenerateSectionWidgets)
 		, OnResetToDefaultClicked(InOnResetToDefaultClicked)
 		, MultipleSectionCount(InMultipleSectionCount)
+		, ThumbnailSize(InThumbnailSize)
 	{
 
 	}
@@ -1703,14 +1744,17 @@ private:
 	FOnGenerateWidgetsForSection OnGenerateCustomSectionWidgets;
 	FOnResetSectionToDefaultClicked OnResetToDefaultClicked;
 	int32 MultipleSectionCount;
+	int32 ThumbnailSize;
 };
 
 
-FSectionList::FSectionList(IDetailLayoutBuilder& InDetailLayoutBuilder, FSectionListDelegates& InSectionListDelegates, bool bInAllowCollapse /*= false*/)
+FSectionList::FSectionList(IDetailLayoutBuilder& InDetailLayoutBuilder, FSectionListDelegates& InSectionListDelegates, bool bInAllowCollapse, int32 InThumbnailSize, int32 InSectionsLodIndex)
 	: SectionListDelegates(InSectionListDelegates)
 	, DetailLayoutBuilder(InDetailLayoutBuilder)
-	, SectionListBuilder(new FSectionListBuilder)
+	, SectionListBuilder(new FSectionListBuilder(InThumbnailSize))
 	, bAllowCollpase(bInAllowCollapse)
+	, ThumbnailSize(InThumbnailSize)
+	, SectionsLodIndex(InSectionsLodIndex)
 {
 }
 
@@ -1789,6 +1833,13 @@ void FSectionList::GenerateHeaderRowContent(FDetailWidgetRow& NodeRow)
 		.Text(LOCTEXT("SectionHeaderTitle", "Sections"))
 		.Font(IDetailLayoutBuilder::GetDetailFont())
 	];
+	if (SectionListDelegates.OnGenerateLodComboBox.IsBound())
+	{
+		NodeRow.ValueContent()
+		[
+			SectionListDelegates.OnGenerateLodComboBox.Execute(SectionsLodIndex)
+		];
+	}
 }
 
 void FSectionList::GenerateChildContent(IDetailChildrenBuilder& ChildrenBuilder)
@@ -1814,7 +1865,7 @@ void FSectionList::GenerateChildContent(IDetailChildrenBuilder& ChildrenBuilder)
 			if (bDisplayAllSectionsInSlot)
 			{
 				FDetailWidgetRow& ChildRow = ChildrenBuilder.AddChildContent(Section.Material.IsValid() ? FText::FromString(Section.Material->GetName()) : FText::GetEmpty());
-				AddSectionItem(ChildRow, CurrentLODIndex, FSectionListItem(CurrentLODIndex, Section.SectionIndex, Section.MaterialSlotName, Section.MaterialSlotIndex, Section.OriginalMaterialSlotName, Section.AvailableMaterialSlotName, Section.Material.Get(), Section.IsSectionUsingCloth), !bDisplayAllSectionsInSlot);
+				AddSectionItem(ChildRow, CurrentLODIndex, FSectionListItem(CurrentLODIndex, Section.SectionIndex, Section.MaterialSlotName, Section.MaterialSlotIndex, Section.OriginalMaterialSlotName, Section.AvailableMaterialSlotName, Section.Material.Get(), Section.IsSectionUsingCloth, ThumbnailSize), !bDisplayAllSectionsInSlot);
 			}
 		}
 	}
@@ -1891,7 +1942,7 @@ void FSectionList::AddSectionItem(FDetailWidgetRow& Row, int32 LodIndex, const s
 {
 	uint32 NumSections = SectionListBuilder->GetNumSections(LodIndex);
 
-	TSharedRef<FSectionItemView> NewView = FSectionItemView::Create(Item, SectionListDelegates.OnSectionChanged, SectionListDelegates.OnGenerateCustomNameWidgets, SectionListDelegates.OnGenerateCustomSectionWidgets, SectionListDelegates.OnResetSectionToDefaultClicked, NumSections);
+	TSharedRef<FSectionItemView> NewView = FSectionItemView::Create(Item, SectionListDelegates.OnSectionChanged, SectionListDelegates.OnGenerateCustomNameWidgets, SectionListDelegates.OnGenerateCustomSectionWidgets, SectionListDelegates.OnResetSectionToDefaultClicked, NumSections, ThumbnailSize);
 
 	TSharedPtr<SWidget> RightSideContent;
 	if (bDisplayLink)
@@ -1943,11 +1994,6 @@ void SMaterialSlotWidget::Construct(const FArguments& InArgs, int32 SlotIndex, b
 			LOCTEXT("CustomNameMaterialNotUsedDeleteTooltip", "Delete this material slot"),
 			InArgs._CanDeleteMaterialSlot);
 
-	if(bIsMaterialUsed)
-	{
-		DeleteButton->SetVisibility(EVisibility::Hidden);
-	}
-
 	ChildSlot
 	[
 		SAssignNew(SlotNameBox, SHorizontalBox)
@@ -1958,22 +2004,20 @@ void SMaterialSlotWidget::Construct(const FArguments& InArgs, int32 SlotIndex, b
 		[
 			SNew(SBox)
 			.VAlign(VAlign_Center)
-			// Match the size of the thumbnail
-			.MinDesiredWidth(64.0f)
 			[
 				SNew(STextBlock)
 				.Font(IDetailLayoutBuilder::GetDetailFont())
-				.Text(LOCTEXT("MaterialArrayNameLabelStringKey", "Slot name"))
+				.Text(LOCTEXT("MaterialArrayNameLabelStringKey", "Slot Name"))
 			]
 		]
 		+ SHorizontalBox::Slot()
-		.VAlign(VAlign_Center)
-		.AutoWidth()
-		.Padding(12.0f, 3.0f, 0.0f, 3.0f)
+		.FillWidth(1.0f)
+		.Padding(5.0f, 3.0f, 0.0f, 3.0f)
 		[
 			SNew(SBox)
 			.VAlign(VAlign_Center)
-			.MinDesiredWidth(125.0f)
+			.HAlign(HAlign_Fill)
+			.MinDesiredWidth(160.0f)
 			[
 				SNew(SEditableTextBox)
 				.Text(InArgs._MaterialName)
@@ -1984,6 +2028,13 @@ void SMaterialSlotWidget::Construct(const FArguments& InArgs, int32 SlotIndex, b
 		]
 	];
 
+	
+	if (bIsMaterialUsed)
+	{
+		DeleteButton->SetVisibility(EVisibility::Hidden);
+	}
+	
+
 	SlotNameBox->AddSlot()
 		.AutoWidth()
 		.VAlign(VAlign_Center)
diff --git a/Engine/Source/Editor/PropertyEditor/Private/UserInterface/PropertyEditor/SPropertyEditorAsset.cpp b/Engine/Source/Editor/PropertyEditor/Private/UserInterface/PropertyEditor/SPropertyEditorAsset.cpp
--- a/Engine/Source/Editor/PropertyEditor/Private/UserInterface/PropertyEditor/SPropertyEditorAsset.cpp
+++ b/Engine/Source/Editor/PropertyEditor/Private/UserInterface/PropertyEditor/SPropertyEditorAsset.cpp
@@ -76,6 +76,7 @@ void SPropertyEditorAsset::Construct( const FArguments& InArgs, const TSharedPtr
 	PropertyHandle = InArgs._PropertyHandle;
 	OnSetObject = InArgs._OnSetObject;
 	OnShouldFilterAsset = InArgs._OnShouldFilterAsset;
+	bool DisplayCompactedSize = InArgs._DisplayCompactSize;
 
 	UProperty* Property = nullptr;
 	if(PropertyEditor.IsValid())
@@ -289,29 +290,62 @@ void SPropertyEditorAsset::Construct( const FArguments& InArgs, const TSharedPtr
 			]
 		];
 
-
-		ValueContentBox->AddSlot()
-		[
-			SNew( SBox )
-			.VAlign( VAlign_Center )
+		if(DisplayCompactedSize)
+		{
+			ValueContentBox->AddSlot()
 			[
-				SAssignNew( CustomContentBox, SVerticalBox )
-				+ SVerticalBox::Slot()
-				.Padding(0.0f, 4.0f, 0.0f, 0.0f)
+				SNew( SBox )
+				.VAlign( VAlign_Center )
 				[
-					AssetComboButton.ToSharedRef()
+					SAssignNew( CustomContentBox, SVerticalBox )
+					+ SVerticalBox::Slot()
+					.AutoHeight()
+					.Padding(0.0f, 4.0f, 0.0f, 0.0f)
+					[
+						SNew(SHorizontalBox)
+						+SHorizontalBox::Slot()
+						.FillWidth(1.0f)
+						[
+							AssetComboButton.ToSharedRef()
+						]
+						+ SHorizontalBox::Slot()
+						.AutoWidth()
+						[
+							SAssignNew(ButtonBoxWrapper, SBox)
+							.Padding(FMargin(0.0f, 2.0f, 4.0f, 2.0f))
+							[
+								ButtonBox
+							]
+						]
+					]
 				]
-				+ SVerticalBox::Slot()
-				.AutoHeight()
+			];
+		}
+		else
+		{
+			ValueContentBox->AddSlot()
+			[
+				SNew( SBox )
+				.VAlign( VAlign_Center )
 				[
-					SAssignNew(ButtonBoxWrapper, SBox)
-					.Padding( FMargin( 0.0f, 2.0f, 4.0f, 2.0f ) )
+					SAssignNew( CustomContentBox, SVerticalBox )
+					+ SVerticalBox::Slot()
+					.Padding(0.0f, 4.0f, 0.0f, 0.0f)
 					[
-						ButtonBox
+						AssetComboButton.ToSharedRef()
+					]
+					+ SVerticalBox::Slot()
+					.AutoHeight()
+					[
+						SAssignNew(ButtonBoxWrapper, SBox)
+						.Padding( FMargin( 0.0f, 2.0f, 4.0f, 2.0f ) )
+						[
+							ButtonBox
+						]
 					]
 				]
-			]
-		];
+			];
+		}
 	}
 	else
 	{
diff --git a/Engine/Source/Editor/PropertyEditor/Private/UserInterface/PropertyEditor/SPropertyEditorAsset.h b/Engine/Source/Editor/PropertyEditor/Private/UserInterface/PropertyEditor/SPropertyEditorAsset.h
--- a/Engine/Source/Editor/PropertyEditor/Private/UserInterface/PropertyEditor/SPropertyEditorAsset.h
+++ b/Engine/Source/Editor/PropertyEditor/Private/UserInterface/PropertyEditor/SPropertyEditorAsset.h
@@ -37,6 +37,7 @@ public:
 		, _DisplayUseSelected(true)
 		, _DisplayBrowse(true)
 		, _EnableContentPicker(true)
+		, _DisplayCompactSize(false)
 		, _ThumbnailPool(NULL)
 		, _ThumbnailSize( FIntPoint(64, 64) )
 		, _ObjectPath()
@@ -51,6 +52,7 @@ public:
 		SLATE_ARGUMENT( bool, DisplayUseSelected )
 		SLATE_ARGUMENT( bool, DisplayBrowse )
 		SLATE_ARGUMENT( bool, EnableContentPicker)
+		SLATE_ARGUMENT( bool, DisplayCompactSize)
 		SLATE_ARGUMENT( TSharedPtr<FAssetThumbnailPool>, ThumbnailPool )
 		SLATE_ARGUMENT( FIntPoint, ThumbnailSize )
 		SLATE_ATTRIBUTE( FString, ObjectPath )
diff --git a/Engine/Source/Editor/PropertyEditor/Public/PropertyCustomizationHelpers.h b/Engine/Source/Editor/PropertyEditor/Public/PropertyCustomizationHelpers.h
--- a/Engine/Source/Editor/PropertyEditor/Public/PropertyCustomizationHelpers.h
+++ b/Engine/Source/Editor/PropertyEditor/Public/PropertyCustomizationHelpers.h
@@ -522,7 +522,7 @@ class FMaterialList
 	, public TSharedFromThis<FMaterialList>
 {
 public:
-	PROPERTYEDITOR_API FMaterialList( IDetailLayoutBuilder& InDetailLayoutBuilder, FMaterialListDelegates& MaterialListDelegates, bool bInAllowCollapse = false, bool bInShowUsedTextures = true);
+	PROPERTYEDITOR_API FMaterialList( IDetailLayoutBuilder& InDetailLayoutBuilder, FMaterialListDelegates& MaterialListDelegates, bool bInAllowCollapse = false, bool bInShowUsedTextures = true, bool bInDisplayCompactSize = false);
 
 	/**
 	 * @return true if materials are being displayed.                                                          
@@ -598,6 +598,8 @@ private:
 	bool bAllowCollpase;
 	/** Whether or not to use the used textures menu for each material entry */
 	bool bShowUsedTextures;
+	/** Whether or not to display a compact form of material entry*/
+	bool bDisplayCompactSize;
 };
 
 
@@ -636,6 +638,8 @@ DECLARE_DELEGATE_RetVal_TwoParams(TSharedRef<SWidget>, FOnGenerateWidgetsForSect
 
 DECLARE_DELEGATE_TwoParams(FOnResetSectionToDefaultClicked, int32, int32);
 
+DECLARE_DELEGATE_RetVal_OneParam(TSharedRef<SWidget>, FOnGenerateLODComboBox, int32);
+
 DECLARE_DELEGATE_RetVal(bool, FOnCanCopySectionList);
 DECLARE_DELEGATE(FOnCopySectionList);
 DECLARE_DELEGATE(FOnPasteSectionList);
@@ -652,6 +656,7 @@ struct FSectionListDelegates
 		, OnGenerateCustomNameWidgets()
 		, OnGenerateCustomSectionWidgets()
 		, OnResetSectionToDefaultClicked()
+		, OnGenerateLodComboBox()
 	{}
 
 	/** Delegate called to populate the list with Sections */
@@ -665,6 +670,9 @@ struct FSectionListDelegates
 	/** Delegate called when a Section list item should be reset to default */
 	FOnResetSectionToDefaultClicked OnResetSectionToDefaultClicked;
 
+	/** Delegate called when a Section list generate the LOD section combo box */
+	FOnGenerateLODComboBox OnGenerateLodComboBox;
+
 	/** Delegate called Copying a section list */
 	FOnCopySectionList OnCopySectionList;
 	/** Delegate called to know if we can copy a section list */
@@ -711,6 +719,9 @@ struct FSectionListItem
 	/* Is this section is using cloth */
 	bool IsSectionUsingCloth;
 
+	/* Size of the preview material thumbnail */
+	int32 ThumbnailSize;
+
 	/** Material being readonly view in the list */
 	TWeakObjectPtr<UMaterialInterface> Material;
 
@@ -723,10 +734,11 @@ struct FSectionListItem
 	TMap<int32, FName> AvailableMaterialSlotName;
 
 
-	FSectionListItem(int32 InLodIndex, int32 InSectionIndex, FName InMaterialSlotName, int32 InMaterialSlotIndex, FName InOriginalMaterialSlotName, const TMap<int32, FName> &InAvailableMaterialSlotName, const UMaterialInterface* InMaterial, bool InIsSectionUsingCloth)
+	FSectionListItem(int32 InLodIndex, int32 InSectionIndex, FName InMaterialSlotName, int32 InMaterialSlotIndex, FName InOriginalMaterialSlotName, const TMap<int32, FName> &InAvailableMaterialSlotName, const UMaterialInterface* InMaterial, bool InIsSectionUsingCloth, int32 InThumbnailSize)
 		: LodIndex(InLodIndex)
 		, SectionIndex(InSectionIndex)
 		, IsSectionUsingCloth(InIsSectionUsingCloth)
+		, ThumbnailSize(InThumbnailSize)
 		, Material(InMaterial)
 		, MaterialSlotName(InMaterialSlotName)
 		, MaterialSlotIndex(InMaterialSlotIndex)
@@ -767,7 +779,7 @@ struct FSectionListItem
 class FSectionList : public IDetailCustomNodeBuilder, public TSharedFromThis<FSectionList>
 {
 public:
-	PROPERTYEDITOR_API FSectionList(IDetailLayoutBuilder& InDetailLayoutBuilder, FSectionListDelegates& SectionListDelegates, bool bInAllowCollapse = false);
+	PROPERTYEDITOR_API FSectionList(IDetailLayoutBuilder& InDetailLayoutBuilder, FSectionListDelegates& SectionListDelegates, bool bInAllowCollapse, int32 InThumbnailSize, int32 InSectionsLodIndex);
 
 	/**
 	* @return true if Sections are being displayed
@@ -828,4 +840,7 @@ private:
 	TSharedRef<class FSectionListBuilder> SectionListBuilder;
 	/** Allow Collapse of Section header row. Right now if you allow collapse, it will initially collapse. */
 	bool bAllowCollpase;
+
+	int32 ThumbnailSize;
+	int32 SectionsLodIndex;
 };
diff --git a/Engine/Source/Editor/StaticMeshEditor/Private/StaticMeshEditor.cpp b/Engine/Source/Editor/StaticMeshEditor/Private/StaticMeshEditor.cpp
--- a/Engine/Source/Editor/StaticMeshEditor/Private/StaticMeshEditor.cpp
+++ b/Engine/Source/Editor/StaticMeshEditor/Private/StaticMeshEditor.cpp
@@ -1099,6 +1099,7 @@ void FStaticMeshEditor::RefreshTool()
 		UpdateLODStats(LODIndex);
 	}
 
+	OnSelectedLODChangedResetOnRefresh.Clear();
 	bool bForceRefresh = true;
 	StaticMeshDetailsView->SetObject( StaticMesh, bForceRefresh );
 
@@ -1206,11 +1207,22 @@ void FStaticMeshEditor::ComboBoxSelectionChanged( TSharedPtr<FString> NewSelecti
 
 void FStaticMeshEditor::LODLevelsSelectionChanged( TSharedPtr<FString> NewSelection, ESelectInfo::Type /*SelectInfo*/ )
 {
-	int32 CurrentLOD = GetCurrentLODLevel();
-
+	int32 CurrentLOD = 0;
+	LODLevels.Find(LODLevelCombo->GetSelectedItem(), CurrentLOD);
+	if (GetStaticMeshComponent() != nullptr)
+	{
+		GetStaticMeshComponent()->ForcedLodModel = CurrentLOD;
+	}
 	UpdateLODStats( CurrentLOD > 0? CurrentLOD - 1 : 0 );
-
 	Viewport->ForceLODLevel(CurrentLOD);
+	if (OnSelectedLODChanged.IsBound())
+	{
+		OnSelectedLODChanged.Broadcast();
+	}
+	if (OnSelectedLODChangedResetOnRefresh.IsBound())
+	{
+		OnSelectedLODChangedResetOnRefresh.Broadcast();
+	}
 }
 
 int32 FStaticMeshEditor::GetCurrentUVChannel()
@@ -1225,13 +1237,21 @@ int32 FStaticMeshEditor::GetCurrentLODLevel()
 {
 	int32 Index = 0;
 	LODLevels.Find(LODLevelCombo->GetSelectedItem(), Index);
+	if (GetStaticMeshComponent() != nullptr)
+	{
+		if (GetStaticMeshComponent()->ForcedLodModel != Index)
+		{
+			LODLevelCombo->SetSelectedItem(LODLevels[GetStaticMeshComponent()->ForcedLodModel]);
+			LODLevels.Find(LODLevelCombo->GetSelectedItem(), Index);
+		}
+	}
 
 	return Index;
 }
 
 int32 FStaticMeshEditor::GetCurrentLODIndex()
 {
-	int32 Index = LODLevels.Find(LODLevelCombo->GetSelectedItem());
+	int32 Index = GetCurrentLODLevel();
 
 	return Index == 0? 0 : Index - 1;
 }
diff --git a/Engine/Source/Editor/StaticMeshEditor/Private/StaticMeshEditor.h b/Engine/Source/Editor/StaticMeshEditor/Private/StaticMeshEditor.h
--- a/Engine/Source/Editor/StaticMeshEditor/Private/StaticMeshEditor.h
+++ b/Engine/Source/Editor/StaticMeshEditor/Private/StaticMeshEditor.h
@@ -129,6 +129,24 @@ public:
 	/** Get the active view mode */
 	virtual EViewModeIndex GetViewMode() const override;
 
+	virtual void RegisterOnSelectedLODChanged(const FOnSelectedLODChanged &Delegate, bool UnregisterOnRefresh) override
+	{
+		if (!UnregisterOnRefresh)
+		{
+			OnSelectedLODChanged.Add(Delegate);
+		}
+		else
+		{
+			OnSelectedLODChangedResetOnRefresh.Add(Delegate);
+		}
+	}
+	
+	virtual void UnRegisterOnSelectedLODChanged(void* Thing) override
+	{
+		OnSelectedLODChanged.RemoveAll(Thing);
+		OnSelectedLODChangedResetOnRefresh.RemoveAll(Thing);
+	}
+
 private:
 	TSharedRef<SDockTab> SpawnTab_Viewport(const FSpawnTabArgs& Args);
 	TSharedRef<SDockTab> SpawnTab_Properties(const FSpawnTabArgs& Args);
@@ -331,6 +349,9 @@ private:
 	const float MinPrimSize;
 	const FVector OverlapNudge;
 
+	FOnSelectedLODChangedMulticaster OnSelectedLODChanged;
+	FOnSelectedLODChangedMulticaster OnSelectedLODChangedResetOnRefresh;
+
 	/**	The tab ids for all the tabs used */
 	static const FName ViewportTabId;
 	static const FName PropertiesTabId;
diff --git a/Engine/Source/Editor/StaticMeshEditor/Private/StaticMeshEditorTools.cpp b/Engine/Source/Editor/StaticMeshEditor/Private/StaticMeshEditorTools.cpp
--- a/Engine/Source/Editor/StaticMeshEditor/Private/StaticMeshEditorTools.cpp
+++ b/Engine/Source/Editor/StaticMeshEditor/Private/StaticMeshEditorTools.cpp
@@ -1337,6 +1337,7 @@ void FMeshSectionSettingsLayout::AddToCategory( IDetailCategoryBuilder& Category
 	SectionListDelegates.OnSectionChanged.BindSP(this, &FMeshSectionSettingsLayout::OnSectionChanged);
 	SectionListDelegates.OnGenerateCustomNameWidgets.BindSP(this, &FMeshSectionSettingsLayout::OnGenerateCustomNameWidgetsForSection);
 	SectionListDelegates.OnGenerateCustomSectionWidgets.BindSP(this, &FMeshSectionSettingsLayout::OnGenerateCustomSectionWidgetsForSection);
+	SectionListDelegates.OnGenerateLodComboBox.BindSP(this, &FMeshSectionSettingsLayout::OnGenerateLodComboBoxForSectionList);
 
 	SectionListDelegates.OnCopySectionList.BindSP(this, &FMeshSectionSettingsLayout::OnCopySectionList, LODIndex);
 	SectionListDelegates.OnCanCopySectionList.BindSP(this, &FMeshSectionSettingsLayout::OnCanCopySectionList, LODIndex);
@@ -1345,7 +1346,9 @@ void FMeshSectionSettingsLayout::AddToCategory( IDetailCategoryBuilder& Category
 	SectionListDelegates.OnCanCopySectionItem.BindSP(this, &FMeshSectionSettingsLayout::OnCanCopySectionItem);
 	SectionListDelegates.OnPasteSectionItem.BindSP(this, &FMeshSectionSettingsLayout::OnPasteSectionItem);
 
-	CategoryBuilder.AddCustomBuilder(MakeShareable(new FSectionList(CategoryBuilder.GetParentLayout(), SectionListDelegates, (LODIndex > 0))));
+	CategoryBuilder.AddCustomBuilder(MakeShareable(new FSectionList(CategoryBuilder.GetParentLayout(), SectionListDelegates, false, 61, LODIndex)));
+
+	StaticMeshEditor.RegisterOnSelectedLODChanged(FOnSelectedLODChanged::CreateSP(this, &FMeshSectionSettingsLayout::UpdateLODCategoryVisibility), true);
 }
 
 void FMeshSectionSettingsLayout::OnCopySectionList(int32 CurrentLODIndex)
@@ -1656,9 +1659,10 @@ TSharedRef<SWidget> FMeshSectionSettingsLayout::OnGenerateCustomNameWidgetsForSe
 
 TSharedRef<SWidget> FMeshSectionSettingsLayout::OnGenerateCustomSectionWidgetsForSection(int32 ForLODIndex, int32 SectionIndex)
 {
-	return SNew(SVerticalBox)
-		+SVerticalBox::Slot()
-		.AutoHeight()
+	return SNew(SHorizontalBox)
+		+ SHorizontalBox::Slot()
+		.AutoWidth()
+		.Padding(2, 0, 2, 0)
 		[
 			SNew(SCheckBox)
 				.IsChecked(this, &FMeshSectionSettingsLayout::DoesSectionCastShadow, SectionIndex)
@@ -1669,9 +1673,9 @@ TSharedRef<SWidget> FMeshSectionSettingsLayout::OnGenerateCustomSectionWidgetsFo
 					.Text(LOCTEXT("CastShadow", "Cast Shadow"))
 			]
 		]
-		+SVerticalBox::Slot()
-			.AutoHeight()
-			.Padding(0,2,0,0)
+		+SHorizontalBox::Slot()
+		.AutoWidth()
+		.Padding(2,0,2,0)
 		[
 			SNew(SCheckBox)
 				.IsEnabled(this, &FMeshSectionSettingsLayout::SectionCollisionEnabled)
@@ -1832,6 +1836,121 @@ void FMeshSectionSettingsLayout::CallPostEditChange(UProperty* PropertyChanged/*
 	StaticMeshEditor.RefreshViewport();
 }
 
+void FMeshSectionSettingsLayout::SetCurrentLOD(int32 NewLodIndex)
+{
+	if (StaticMeshEditor.GetStaticMeshComponent() == nullptr || LodCategoriesPtr == nullptr)
+	{
+		return;
+	}
+	int32 CurrentDisplayLOD = StaticMeshEditor.GetStaticMeshComponent()->ForcedLodModel;
+	int32 RealCurrentDisplayLOD = CurrentDisplayLOD == 0 ? 0 : CurrentDisplayLOD - 1;
+	int32 RealNewLOD = NewLodIndex == 0 ? 0 : NewLodIndex - 1;
+	
+	if (CurrentDisplayLOD == NewLodIndex || !LodCategoriesPtr->IsValidIndex(RealCurrentDisplayLOD) || !LodCategoriesPtr->IsValidIndex(RealNewLOD))
+	{
+		return;
+	}
+	(*LodCategoriesPtr)[RealCurrentDisplayLOD]->SetCategoryVisibility(false);
+	(*LodCategoriesPtr)[RealNewLOD]->SetCategoryVisibility(true);
+	StaticMeshEditor.GetStaticMeshComponent()->SetForcedLodModel(NewLodIndex);
+
+	//Reset the preview section since we do not edit the same LOD
+	StaticMeshEditor.GetStaticMeshComponent()->SetSectionPreview(INDEX_NONE);
+	StaticMeshEditor.GetStaticMeshComponent()->SelectedEditorSection = INDEX_NONE;
+}
+
+void FMeshSectionSettingsLayout::UpdateLODCategoryVisibility()
+{
+	bool bAutoLod = false;
+	if (StaticMeshEditor.GetStaticMeshComponent() != nullptr)
+	{
+		bAutoLod = StaticMeshEditor.GetStaticMeshComponent()->ForcedLodModel == 0;
+	}
+	int32 CurrentDisplayLOD = bAutoLod ? 0 : StaticMeshEditor.GetStaticMeshComponent()->ForcedLodModel - 1;
+
+	if (LodCategoriesPtr != nullptr && LodCategoriesPtr->IsValidIndex(CurrentDisplayLOD) && StaticMeshEditor.GetStaticMesh())
+	{
+		int32 StaticMeshLodNumber = StaticMeshEditor.GetStaticMesh()->GetNumLODs();
+		for (int32 LodCategoryIndex = 0; LodCategoryIndex < StaticMeshLodNumber; ++LodCategoryIndex)
+		{
+			if (!LodCategoriesPtr->IsValidIndex(LodCategoryIndex))
+			{
+				break;
+			}
+			(*LodCategoriesPtr)[LodCategoryIndex]->SetCategoryVisibility(CurrentDisplayLOD == LodCategoryIndex);
+		}
+	}
+}
+
+FText FMeshSectionSettingsLayout::GetCurrentLodName() const
+{
+	bool bAutoLod = false;
+	if (StaticMeshEditor.GetStaticMeshComponent() != nullptr)
+	{
+		bAutoLod = StaticMeshEditor.GetStaticMeshComponent()->ForcedLodModel == 0;
+	}
+	int32 CurrentDisplayLOD = bAutoLod ? 0 : StaticMeshEditor.GetStaticMeshComponent()->ForcedLodModel - 1;
+	return FText::FromString(bAutoLod ? FString(TEXT("Auto (LOD0)")) : (FString(TEXT("LOD")) + FString::FromInt(CurrentDisplayLOD)));
+}
+
+FText FMeshSectionSettingsLayout::GetCurrentLodTooltip() const
+{
+	if (StaticMeshEditor.GetStaticMeshComponent() != nullptr && StaticMeshEditor.GetStaticMeshComponent()->ForcedLodModel == 0)
+	{
+		return FText::FromString(TEXT("LOD0 is edit when selecting Auto LOD"));
+	}
+	return FText::GetEmpty();
+}
+
+TSharedRef<SWidget> FMeshSectionSettingsLayout::OnGenerateLodComboBoxForSectionList(int32 LodIndex)
+{
+	return SNew(SComboButton)
+		.OnGetMenuContent(this, &FMeshSectionSettingsLayout::OnGenerateLodMenuForSectionList, LodIndex)
+		.VAlign(VAlign_Center)
+		.ContentPadding(2)
+		.ButtonContent()
+		[
+			SNew(STextBlock)
+			.Font(IDetailLayoutBuilder::GetDetailFont())
+		.Text(this, &FMeshSectionSettingsLayout::GetCurrentLodName)
+		.ToolTipText(this, &FMeshSectionSettingsLayout::GetCurrentLodTooltip)
+		];
+}
+
+TSharedRef<SWidget> FMeshSectionSettingsLayout::OnGenerateLodMenuForSectionList(int32 LodIndex)
+{
+	UStaticMesh* StaticMesh = StaticMeshEditor.GetStaticMesh();
+
+	if (StaticMesh == nullptr)
+	{
+		return SNullWidget::NullWidget;
+	}
+
+	bool bAutoLod = false;
+	if (StaticMeshEditor.GetStaticMeshComponent() != nullptr)
+	{
+		bAutoLod = StaticMeshEditor.GetStaticMeshComponent()->ForcedLodModel == 0;
+	}
+	const int32 StaticMeshLODCount = StaticMesh->GetNumLODs();
+	if (StaticMeshLODCount < 2)
+	{
+		return SNullWidget::NullWidget;
+	}
+	FMenuBuilder MenuBuilder(true, NULL);
+
+	FText AutoLodText = FText::FromString((TEXT("Auto LOD")));
+	FUIAction AutoLodAction(FExecuteAction::CreateSP(this, &FMeshSectionSettingsLayout::SetCurrentLOD, 0));
+	MenuBuilder.AddMenuEntry(AutoLodText, LOCTEXT("OnGenerateLodMenuForSectionList_Auto_ToolTip", "LOD0 is edit when selecting Auto LOD"), FSlateIcon(), AutoLodAction);
+	// Add a menu item for each texture.  Clicking on the texture will display it in the content browser
+	for (int32 AllLodIndex = 0; AllLodIndex < StaticMeshLODCount; ++AllLodIndex)
+	{
+		FText LODLevelString = FText::FromString((TEXT("LOD ") + FString::FromInt(AllLodIndex)));
+		FUIAction Action(FExecuteAction::CreateSP(this, &FMeshSectionSettingsLayout::SetCurrentLOD, AllLodIndex + 1));
+		MenuBuilder.AddMenuEntry(LODLevelString, FText::GetEmpty(), FSlateIcon(), Action);
+	}
+
+	return MenuBuilder.MakeWidget();
+}
 
 //////////////////////////////////////////////////////////////////////////
 // FMeshMaterialLayout
@@ -1914,7 +2033,7 @@ void FMeshMaterialsLayout::AddToCategory(IDetailCategoryBuilder& CategoryBuilder
 	MaterialListDelegates.OnCanCopyMaterialItem.BindSP(this, &FMeshMaterialsLayout::OnCanCopyMaterialItem);
 	MaterialListDelegates.OnPasteMaterialItem.BindSP(this, &FMeshMaterialsLayout::OnPasteMaterialItem);
 
-	CategoryBuilder.AddCustomBuilder(MakeShareable(new FMaterialList(CategoryBuilder.GetParentLayout(), MaterialListDelegates)));
+	CategoryBuilder.AddCustomBuilder(MakeShareable(new FMaterialList(CategoryBuilder.GetParentLayout(), MaterialListDelegates, false, true, true)));
 }
 
 void FMeshMaterialsLayout::OnCopyMaterialList()
@@ -1960,7 +2079,16 @@ void FMeshMaterialsLayout::OnPasteMaterialList()
 		FScopedTransaction Transaction(LOCTEXT("StaticMeshToolChangedPasteMaterialList", "Pasted material list"));
 		GetStaticMesh().Modify();
 
-		FJsonObjectConverter::JsonValueToUProperty(RootJsonValue, Property, &GetStaticMesh().StaticMaterials, 0, 0);
+		TArray<FStaticMaterial> TempMaterials;
+		FJsonObjectConverter::JsonValueToUProperty(RootJsonValue, Property, &TempMaterials, 0, 0);
+		//Do not change the number of material in the array
+		for (int32 MaterialIndex = 0; MaterialIndex < TempMaterials.Num(); ++MaterialIndex)
+		{
+			if (GetStaticMesh().StaticMaterials.IsValidIndex(MaterialIndex))
+			{
+				GetStaticMesh().StaticMaterials[MaterialIndex].MaterialInterface = TempMaterials[MaterialIndex].MaterialInterface;
+			}
+		}
 
 		CallPostEditChange(Property);
 	}
@@ -2016,7 +2144,9 @@ void FMeshMaterialsLayout::OnPasteMaterialItem(int32 CurrentSlot)
 
 		if (GetStaticMesh().StaticMaterials.IsValidIndex(CurrentSlot))
 		{
-			FJsonObjectConverter::JsonObjectToUStruct(RootJsonObject.ToSharedRef(), FSkeletalMaterial::StaticStruct(), &GetStaticMesh().StaticMaterials[CurrentSlot], 0, 0);
+			FStaticMaterial TmpStaticMaterial;
+			FJsonObjectConverter::JsonObjectToUStruct(RootJsonObject.ToSharedRef(), FStaticMaterial::StaticStruct(), &TmpStaticMaterial, 0, 0);
+			GetStaticMesh().StaticMaterials[CurrentSlot].MaterialInterface = TmpStaticMaterial.MaterialInterface;
 		}
 
 		CallPostEditChange(Property);
@@ -2762,11 +2892,19 @@ void FLevelOfDetailSettingsLayout::AddLODLevelCategories( IDetailLayoutBuilder&
 		{
 			FString CategoryName = FString(TEXT("StaticMeshMaterials"));
 
-			IDetailCategoryBuilder& MaterialsCategory = DetailBuilder.EditCategory(*CategoryName, LOCTEXT("StaticMeshMaterialsLabel", "Materials"), ECategoryPriority::Important);
+			IDetailCategoryBuilder& MaterialsCategory = DetailBuilder.EditCategory(*CategoryName, LOCTEXT("StaticMeshMaterialsLabel", "Material Slots"), ECategoryPriority::Important);
 
 			MaterialsLayoutWidget = MakeShareable(new FMeshMaterialsLayout(StaticMeshEditor));
 			MaterialsLayoutWidget->AddToCategory(MaterialsCategory);
 		}
+
+		int32 CurrentLodIndex = 0;
+		
+		if (StaticMeshEditor.GetStaticMeshComponent() != nullptr)
+		{
+			CurrentLodIndex = StaticMeshEditor.GetStaticMeshComponent()->ForcedLodModel;
+		}
+		LodCategories.Empty(StaticMeshLODCount);
 		// Create information panel for each LOD level.
 		for(int32 LODIndex = 0; LODIndex < StaticMeshLODCount; ++LODIndex)
 		{
@@ -2811,27 +2949,10 @@ void FLevelOfDetailSettingsLayout::AddLODLevelCategories( IDetailLayoutBuilder&
 			FString CategoryName = FString(TEXT("LOD"));
 			CategoryName.AppendInt( LODIndex );
 
-			FText LODLevelString = FText::Format( LOCTEXT("LODLevel", "LOD{0}"), FText::AsNumber( LODIndex ) );
+			FText LODLevelString = FText::FromString(FString(TEXT("LOD Sections (LOD")) + FString::FromInt(LODIndex) + TEXT(")") );
 
 			IDetailCategoryBuilder& LODCategory = DetailBuilder.EditCategory( *CategoryName, LODLevelString, ECategoryPriority::Important );
-
-			if (LODIndex != 0)
-			{
-				LODCategory.AddCustomRow( LOCTEXT("RemoveLOD", "Remove LOD") )
-				.ValueContent()
-				.HAlign(HAlign_Left)
-				[
-					SNew(SButton)
-					.OnClicked(this, &FLevelOfDetailSettingsLayout::OnRemoveLOD, LODIndex)
-					.IsEnabled(this, &FLevelOfDetailSettingsLayout::CanRemoveLOD, LODIndex)
-					.ToolTipText( LOCTEXT("RemoveLOD_ToolTip", "Removes this LOD from the Static Mesh") )
-					[
-						SNew(STextBlock)
-						.Text( LOCTEXT("RemoveLOD", "Remove LOD") )
-						.Font( DetailBuilder.GetDetailFont() )
-					]
-				];
-			}
+			LodCategories.Add(&LODCategory);
 
 			LODCategory.HeaderContent
 			(
@@ -2868,7 +2989,7 @@ void FLevelOfDetailSettingsLayout::AddLODLevelCategories( IDetailLayoutBuilder&
 			);
 			
 			
-			SectionSettingsWidgets[ LODIndex ] = MakeShareable( new FMeshSectionSettingsLayout( StaticMeshEditor, LODIndex ) );
+			SectionSettingsWidgets[ LODIndex ] = MakeShareable( new FMeshSectionSettingsLayout( StaticMeshEditor, LODIndex, LodCategories) );
 			SectionSettingsWidgets[ LODIndex ]->AddToCategory( LODCategory );
 
 			LODCategory.AddCustomRow(( LOCTEXT("ScreenSizeRow", "ScreenSize")))
@@ -2901,6 +3022,26 @@ void FLevelOfDetailSettingsLayout::AddLODLevelCategories( IDetailLayoutBuilder&
 				LODCategory.AddCustomBuilder( ReductionSettingsWidgets[LODIndex].ToSharedRef() );
 			}
 
+			if (LODIndex != 0)
+			{
+				LODCategory.AddCustomRow( LOCTEXT("RemoveLOD", "Remove LOD") )
+				.ValueContent()
+				.HAlign(HAlign_Left)
+				[
+					SNew(SButton)
+					.OnClicked(this, &FLevelOfDetailSettingsLayout::OnRemoveLOD, LODIndex)
+					.IsEnabled(this, &FLevelOfDetailSettingsLayout::CanRemoveLOD, LODIndex)
+					.ToolTipText( LOCTEXT("RemoveLOD_ToolTip", "Removes this LOD from the Static Mesh") )
+					[
+						SNew(STextBlock)
+						.Text( LOCTEXT("RemoveLOD", "Remove LOD") )
+						.Font( DetailBuilder.GetDetailFont() )
+					]
+				];
+			}
+
+			LODCategory.SetCategoryVisibility((CurrentLodIndex == 0 ? 0 : CurrentLodIndex - 1) == LODIndex);
+
 		}
 	}
 
diff --git a/Engine/Source/Editor/StaticMeshEditor/Private/StaticMeshEditorTools.h b/Engine/Source/Editor/StaticMeshEditor/Private/StaticMeshEditorTools.h
--- a/Engine/Source/Editor/StaticMeshEditor/Private/StaticMeshEditorTools.h
+++ b/Engine/Source/Editor/StaticMeshEditor/Private/StaticMeshEditorTools.h
@@ -252,9 +252,10 @@ private:
 class FMeshSectionSettingsLayout : public TSharedFromThis<FMeshSectionSettingsLayout>
 {
 public:
-	FMeshSectionSettingsLayout( IStaticMeshEditor& InStaticMeshEditor, int32 InLODIndex )
+	FMeshSectionSettingsLayout( IStaticMeshEditor& InStaticMeshEditor, int32 InLODIndex, TArray<class IDetailCategoryBuilder*> &InLodCategories)
 		: StaticMeshEditor( InStaticMeshEditor )
 		, LODIndex( InLODIndex )
+		, LodCategoriesPtr(&InLodCategories)
 	{}
 
 	virtual ~FMeshSectionSettingsLayout();
@@ -319,9 +320,22 @@ private:
 	void OnSectionIsolatedChanged(ECheckBoxState NewState, int32 SectionIndex);
 
 	void CallPostEditChange(UProperty* PropertyChanged=nullptr);
+
+	TSharedRef<SWidget> OnGenerateLodComboBoxForSectionList(int32 LodIndex);
+	/*
+	* Generate the context menu to choose the LOD we will display the section list
+	*/
+	TSharedRef<SWidget> OnGenerateLodMenuForSectionList(int32 LodIndex);
+	void UpdateLODCategoryVisibility();
+	FText GetCurrentLodName() const;
+	FText GetCurrentLodTooltip() const;
+
+	void SetCurrentLOD(int32 NewLodIndex);
 	
 	IStaticMeshEditor& StaticMeshEditor;
 	int32 LODIndex;
+
+	TArray<class IDetailCategoryBuilder*> *LodCategoriesPtr;
 };
 
 struct FSectionLocalizer
@@ -503,4 +517,6 @@ private:
 	bool bBuildSettingsExpanded[MAX_STATIC_MESH_LODS];
 	bool bReductionSettingsExpanded[MAX_STATIC_MESH_LODS];
 	bool bSectionSettingsExpanded[MAX_STATIC_MESH_LODS];
+
+	TArray<class IDetailCategoryBuilder*> LodCategories;
 };
diff --git a/Engine/Source/Editor/StaticMeshEditor/Public/IStaticMeshEditor.h b/Engine/Source/Editor/StaticMeshEditor/Public/IStaticMeshEditor.h
--- a/Engine/Source/Editor/StaticMeshEditor/Public/IStaticMeshEditor.h
+++ b/Engine/Source/Editor/StaticMeshEditor/Public/IStaticMeshEditor.h
@@ -11,6 +11,11 @@ class UStaticMesh;
 class UStaticMeshComponent;
 class UStaticMeshSocket;
 
+
+DECLARE_MULTICAST_DELEGATE(FOnSelectedLODChangedMulticaster);
+
+typedef FOnSelectedLODChangedMulticaster::FDelegate FOnSelectedLODChanged;
+
 /**
  * Public interface to Static Mesh Editor
  */
@@ -220,6 +225,11 @@ public:
 
 	/** Get the active view mode */
 	virtual EViewModeIndex GetViewMode() const = 0;
+
+	/* Register callback to be able to be notify when the select LOD is change */
+	virtual void RegisterOnSelectedLODChanged(const FOnSelectedLODChanged &Delegate, bool UnregisterOnRefresh) = 0;
+	/* Unregister callback to free up the ressources */
+	virtual void UnRegisterOnSelectedLODChanged(void* Thing) = 0;
 };
 
 
diff --git a/Engine/Source/Editor/UnrealEd/Classes/CookOnTheSide/CookOnTheFlyServer.h b/Engine/Source/Editor/UnrealEd/Classes/CookOnTheSide/CookOnTheFlyServer.h
--- a/Engine/Source/Editor/UnrealEd/Classes/CookOnTheSide/CookOnTheFlyServer.h
+++ b/Engine/Source/Editor/UnrealEd/Classes/CookOnTheSide/CookOnTheFlyServer.h
@@ -973,7 +973,7 @@ private:
 
 	/** Map of platform name to asset registry generators, which hold the state of asset registry data for a platform */
 	TMap<FName, FAssetRegistryGenerator*> RegistryGenerators;
-
+	
 	//////////////////////////////////////////////////////////////////////////
 	// iterative ini settings checking
 	// growing list of ini settings which are accessed over the course of the cook
@@ -996,6 +996,18 @@ private:
 	* used to reset the cached cooked shaders
 	*/
 	void OnTargetPlatformChangedSupportedFormats(const ITargetPlatform* TargetPlatform);
+
+	/**
+	 * Returns the current set of cooking targetplatforms
+	 *  mostly used for cook on the fly or in situations where the cooker can't figure out what the target platform is
+	 * 
+	 * @return Array of target platforms which are can be used
+	 */
+	const TArray<ITargetPlatform*>& GetCookingTargetPlatforms() const;
+
+	/** Cached cooking target platforms from the targetmanager, these are used when we don't know what platforms we should be targeting */
+	mutable TArray<ITargetPlatform*> CookingTargetPlatforms;
+
 public:
 
 	enum ECookOnTheSideResult
@@ -1394,6 +1406,15 @@ private:
 
 	//////////////////////////////////////////////////////////////////////////
 	// cook on the fly specific functions
+
+	/**
+	 * When we get a new connection from the network make sure the version is compatible 
+	 *		Will terminate the connection if return false
+	 * 
+	 * @return return false if not compatible, true if it is
+	 */
+	bool HandleNetworkFileServerNewConnection( const FString& VersionInfo, const FString& PlatformName );
+
 	/**
 	* Cook requests for a package from network
 	*  blocks until cook is complete
@@ -1417,6 +1438,16 @@ private:
 	 */
 	FString HandleNetworkGetSandboxPath();
 
+
+	/**
+	 * HandleNetworkGetPrecookedList 
+	 * this is used specifically for cook on the fly with shared cooked builds
+	 * returns the list of files which are still valid in the pak file which was initially loaded
+	 * 
+	 * @param PrecookedFileList all the files which are still valid in the client pak file
+	 */
+	void HandleNetworkGetPrecookedList( const FString& PlatformName, TMap<FString, FDateTime>& PrecookedFileList );
+
 	//////////////////////////////////////////////////////////////////////////
 	// general functions
 
diff --git a/Engine/Source/Editor/UnrealEd/Classes/Factories/FbxImportUI.h b/Engine/Source/Editor/UnrealEd/Classes/Factories/FbxImportUI.h
--- a/Engine/Source/Editor/UnrealEd/Classes/Factories/FbxImportUI.h
+++ b/Engine/Source/Editor/UnrealEd/Classes/Factories/FbxImportUI.h
@@ -151,6 +151,8 @@ public:
 	UPROPERTY(EditAnywhere, config, AdvancedDisplay, Category = Material, meta = (OBJRestrict = "true", ImportType = "Mesh"))
 	uint32 bResetMaterialSlots : 1;
 
+	void ResetToDefault();
+
 	/** UObject Interface */
 	virtual bool CanEditChange( const UProperty* InProperty ) const override;
 
diff --git a/Engine/Source/Editor/UnrealEd/Classes/Factories/FbxSceneImportFactory.h b/Engine/Source/Editor/UnrealEd/Classes/Factories/FbxSceneImportFactory.h
--- a/Engine/Source/Editor/UnrealEd/Classes/Factories/FbxSceneImportFactory.h
+++ b/Engine/Source/Editor/UnrealEd/Classes/Factories/FbxSceneImportFactory.h
@@ -381,6 +381,10 @@ class UNREALED_API UFbxSceneImportFactory : public USceneImportFactory
 	/* Default Options always have the same name "Default" */
 	static FString DefaultOptionName;
 
+public:
+	static TSharedPtr<FFbxSceneInfo> ConvertSceneInfo(void* VoidFbxImporter, void* VoidFbxSceneInfo);
+	static void ExtractMaterialInfo(void* FbxImporterVoid, TSharedPtr<FFbxSceneInfo> SceneInfoPtr);
+
 protected:
 	/** Convert the scene and remake all the transform for the SceneInfo pass in parameter.
 	 *  We need this because EvaluateGlobal and EvaluateLocal are dependent of the scene conversion.
@@ -428,8 +432,6 @@ protected:
 	/** Create a package for the specified node. Package will be the concatenation of UFbxSceneImportFactory::Path and Node->GetName(). */
 	UPackage *CreatePackageForNode(FString PackageName, FString &StaticMeshName);
 
-	static TSharedPtr<FFbxSceneInfo> ConvertSceneInfo(void* VoidFbxImporter, void* VoidFbxSceneInfo);
-	static void ExtractMaterialInfo(void* FbxImporterVoid, TSharedPtr<FFbxSceneInfo> SceneInfoPtr);
 	bool SetStaticMeshComponentOverrideMaterial(class UStaticMeshComponent* StaticMeshComponent, TSharedPtr<FFbxNodeInfo> NodeInfo);
 
 	/** The path of the asset to import */
diff --git a/Engine/Source/Editor/UnrealEd/Classes/Settings/EditorExperimentalSettings.h b/Engine/Source/Editor/UnrealEd/Classes/Settings/EditorExperimentalSettings.h
--- a/Engine/Source/Editor/UnrealEd/Classes/Settings/EditorExperimentalSettings.h
+++ b/Engine/Source/Editor/UnrealEd/Classes/Settings/EditorExperimentalSettings.h
@@ -83,6 +83,9 @@ public:
 	UPROPERTY(EditAnywhere, config, Category = Cooking, meta = (DisplayName = "Disable Cook In The Editor feature (cooks from launch on will be run in a separate process if disabled)", ConfigRestartRequired=true))
 	bool bDisableCookInEditor;
 
+	UPROPERTY(EditAnywhere, config, Category = Cooking, meta = (DisplayName = "Use shared cooked builds in launch on", ConfigRestartRequired = true))
+	bool bSharedCookedBuilds;
+
 	UPROPERTY(EditAnywhere, config, Category = Cooking, meta = (DisplayName = "Use multiple processes when cooking (only affects File -> Package)"))
 	int32 MultiProcessCooking;
 
diff --git a/Engine/Source/Editor/UnrealEd/Private/Commandlets/AssetRegistryGenerator.cpp b/Engine/Source/Editor/UnrealEd/Private/Commandlets/AssetRegistryGenerator.cpp
--- a/Engine/Source/Editor/UnrealEd/Private/Commandlets/AssetRegistryGenerator.cpp
+++ b/Engine/Source/Editor/UnrealEd/Private/Commandlets/AssetRegistryGenerator.cpp
@@ -354,17 +354,8 @@ void FAssetRegistryGenerator::CleanManifestDirectories()
 bool FAssetRegistryGenerator::LoadPreviousAssetRegistry(const FString& Filename)
 {
 	// First try development asset registry
-	FString DevelopmentFilename = Filename.Replace(TEXT("AssetRegistry.bin"), TEXT("DevelopmentAssetRegistry.bin"));
 	FArrayReader SerializedAssetData;
-
-	if (IFileManager::Get().FileExists(*DevelopmentFilename) && FFileHelper::LoadFileToArray(SerializedAssetData, *DevelopmentFilename))
-	{
-		FAssetRegistrySerializationOptions Options;
-		Options.ModifyForDevelopment();
-
-		return PreviousState.Serialize(SerializedAssetData, Options);
-	}
-
+	
 	if (IFileManager::Get().FileExists(*Filename) && FFileHelper::LoadFileToArray(SerializedAssetData, *Filename))
 	{
 		FAssetRegistrySerializationOptions Options;
@@ -473,6 +464,8 @@ void FAssetRegistryGenerator::UpdatePackageSourceHashes()
 
 		State.GetDependencies(PackageName, Dependencies, EAssetRegistryDependencyType::Hard);
 
+		Dependencies.Sort([&](const FAssetIdentifier& A, const FAssetIdentifier& B) { return A.PackageName > B.PackageName; } );
+
 		for (FAssetIdentifier& Dependency : Dependencies)
 		{
 			const FAssetPackageData* DependencyData = State.GetAssetPackageData(Dependency.PackageName);
@@ -752,7 +745,7 @@ void FAssetRegistryGenerator::AddAssetToFileOrderRecursive(FAssetData* InAsset,
 	}
 }
 
-bool FAssetRegistryGenerator::SaveAssetRegistry(const FString& SandboxPath)
+bool FAssetRegistryGenerator::SaveAssetRegistry(const FString& SandboxPath, bool bSerializeDevelopmentAssetRegistry )
 {
 	UE_LOG(LogAssetRegistryGenerator, Display, TEXT("Saving asset registry."));
 	const TMap<FName, const FAssetData*>& ObjectToDataMap = State.GetObjectPathToAssetDataMap();
@@ -771,7 +764,7 @@ bool FAssetRegistryGenerator::SaveAssetRegistry(const FString& SandboxPath)
 
 	AssetRegistry.InitializeTemporaryAssetRegistryState(State, SaveOptions, true);
 
-	if (DevelopmentSaveOptions.bSerializeAssetRegistry)
+	if (DevelopmentSaveOptions.bSerializeAssetRegistry && bSerializeDevelopmentAssetRegistry)
 	{
 		// Create development registry data, used for incremental cook and editor viewing
 		FArrayWriter SerializedAssetRegistry;
diff --git a/Engine/Source/Editor/UnrealEd/Private/Commandlets/PackageUtilities.cpp b/Engine/Source/Editor/UnrealEd/Private/Commandlets/PackageUtilities.cpp
--- a/Engine/Source/Editor/UnrealEd/Private/Commandlets/PackageUtilities.cpp
+++ b/Engine/Source/Editor/UnrealEd/Private/Commandlets/PackageUtilities.cpp
@@ -38,6 +38,7 @@
 #include "FileHelpers.h"
 
 #include "Animation/AnimCompress_BitwiseCompressOnly.h"
+#include "Animation/AnimCompress_Automatic.h"
 
 
 #include "CollectionManagerTypes.h"
@@ -1558,6 +1559,25 @@ struct CompressAnimationsFunctor
 	template< typename OBJECTYPE >
 	void DoIt( UCommandlet* Commandlet, UPackage* Package, TArray<FString>& Tokens, TArray<FString>& Switches )
 	{
+		// Count the number of animations to provide some limited progress indication
+		int32 NumAnimationsInPackage = 0;
+		for (TObjectIterator<OBJECTYPE> It; It; ++It)
+		{
+			OBJECTYPE* AnimSeq = *It;
+			if (!AnimSeq->IsIn(Package))
+			{
+				continue;
+			}
+
+			++NumAnimationsInPackage;
+		}
+
+		// Skip packages that contain no Animations.
+		if (NumAnimationsInPackage == 0)
+		{
+			return;
+		}
+
 		// @todoanim: we expect this won't work properly since it won't have any skeletalmesh,
 		// but soon, the compression will changed based on skeleton. 
 		// when that happens, this doesn't have to worry about skeletalmesh not loaded
@@ -1631,25 +1651,17 @@ struct CompressAnimationsFunctor
 		// Get version number. Bump this up every time you want to recompress all animations.
 		const int32 CompressCommandletVersion = UAnimationSettings::Get()->CompressCommandletVersion;
 
-		// Count the number of animations to provide some limited progress indication
-		int32 NumAnimationsInPackage = 0;
-		for (TObjectIterator<OBJECTYPE> It; It; ++It)
-		{
-			++NumAnimationsInPackage;
-		}
-
-
 		int32 ActiveAnimationIndex = 0;
 		for (TObjectIterator<OBJECTYPE> It; It; ++It)
 		{
 			OBJECTYPE* AnimSeq = *It;
-			++ActiveAnimationIndex;
-
 			if (!AnimSeq->IsIn(Package))
 			{
 				continue;
 			}
 
+			++ActiveAnimationIndex;
+
 			// If animation hasn't been compressed, force it.
 			bool bForceCompression = (AnimSeq->CompressedTrackOffsets.Num() == 0);
 
@@ -1669,6 +1681,10 @@ struct CompressAnimationsFunctor
 
 			USkeleton* Skeleton = AnimSeq->GetSkeleton();
 			check (Skeleton);
+			if (Skeleton->HasAnyFlags(RF_NeedLoad))
+			{
+				Skeleton->GetLinker()->Preload(Skeleton);
+			}
 
 			if( bAnalyze )
 			{
@@ -2072,12 +2088,11 @@ struct CompressAnimationsFunctor
 				}
 			}
 #endif
-			int32 OldSize;
-			int32 NewSize;
+			SIZE_T OldSize;
+			SIZE_T NewSize;
 
 			{
-				FArchiveCountMem CountBytesSize( AnimSeq );
-				OldSize = CountBytesSize.GetNum();
+				OldSize = AnimSeq->GetResourceSizeBytes(EResourceSizeMode::Inclusive);
 			}
 
 			// Clear bDoNotOverrideCompression flag
@@ -2101,6 +2116,20 @@ struct CompressAnimationsFunctor
 				// Force an update.
 				AnimSeq->CompressCommandletVersion = 0;
 			}
+			
+			// Do not perform automatic recompression on animations marked as 'bDoNotOverrideCompression'
+			// Unless they have no compression scheme, or they're using automatic compression.
+			if (AnimSeq->bDoNotOverrideCompression && (AnimSeq->CompressionScheme != nullptr) && !AnimSeq->CompressionScheme->IsA(UAnimCompress_Automatic::StaticClass()))
+			{
+				continue;
+			}
+
+			// Set version since we've checked this animation for recompression.
+			if (AnimSeq->CompressCommandletVersion != CompressCommandletVersion)
+			{
+				AnimSeq->CompressCommandletVersion = CompressCommandletVersion;
+				bDirtyPackage = true;
+			}
 
 			UE_LOG(LogPackageUtilities, Warning, TEXT("Compressing animation '%s' [#%d / %d in package '%s']"),
 				*AnimSeq->GetName(),
@@ -2108,22 +2137,16 @@ struct CompressAnimationsFunctor
 				NumAnimationsInPackage,
 				*PackageFileName);
 
-			// @todoanim: expect this won't work
-			AnimSeq->RequestAnimCompression(false, true, false);
+			UAnimCompress* CompressionAlgorithm = NewObject<UAnimCompress_Automatic>();
+			AnimSeq->CompressionScheme = static_cast<UAnimCompress*>(StaticDuplicateObject(CompressionAlgorithm, AnimSeq));
+			AnimSeq->RequestAnimCompression(false, false, false);
 			{
-				FArchiveCountMem CountBytesSize( AnimSeq );
-				NewSize = CountBytesSize.GetNum();
-			}
-
-			// Set version since we've checked this animation for recompression.
-			if( AnimSeq->CompressCommandletVersion != CompressCommandletVersion )
-			{
-				AnimSeq->CompressCommandletVersion = CompressCommandletVersion;
-				bDirtyPackage = true;
+				NewSize = AnimSeq->GetResourceSizeBytes(EResourceSizeMode::Inclusive);
 			}
 
 			// Only save package if size has changed.
-			bDirtyPackage = (bDirtyPackage || bForceCompression || (OldSize != NewSize));
+			const int64 DeltaSize = NewSize - OldSize;
+			bDirtyPackage = (bDirtyPackage || bForceCompression || (DeltaSize != 0));
 
 			// if Dirty, then we need to be able to write to this package. 
 			// If we can't, abort, don't want to waste time!!
@@ -2175,7 +2198,7 @@ struct CompressAnimationsFunctor
 
 		// End of recompression
 		// Does package need to be saved?
-		bDirtyPackage = bDirtyPackage || Package->IsDirty();
+/*		bDirtyPackage = bDirtyPackage || Package->IsDirty();*/
 
 		// If we need to save package, do so.
 		if( bDirtyPackage && !bAnalyze )
@@ -2245,27 +2268,6 @@ int32 UCompressAnimationsCommandlet::Main( const FString& Params )
 	}
 	else
 	{
-		// First scan all Skeletal Meshes
-		UE_LOG(LogPackageUtilities, Warning, TEXT("Scanning for all SkeletalMeshes..."));
-
-		// If we have SKIPREADONLY, then override this, as we need to scan all packages for skeletal meshes.
-		FString SearchAllMeshesParams = ParamsUpperCase;
-		SearchAllMeshesParams += FString(TEXT(" -OVERRIDEREADONLY"));
-		SearchAllMeshesParams += FString(TEXT(" -OVERRIDELOADMAPS"));
-		// Prevent recompression here, we'll do it after we gathered all skeletal meshes
-		GDisableAnimationRecompression = true;
-		DoActionToAllPackages<USkeletalMesh, AddAllSkeletalMeshesToListFunctor>(this, SearchAllMeshesParams);
-		GDisableAnimationRecompression = false;
-
-		int32 Count = 0;
-		for( TObjectIterator<USkeletalMesh> It; It; ++It )
-		{
-			USkeletalMesh* SkelMesh = *It;
-			UE_LOG(LogPackageUtilities, Warning, TEXT("[%i] %s"), Count, *SkelMesh->GetFName().ToString());
-			Count++;
-		}
-		UE_LOG(LogPackageUtilities, Warning, TEXT("%i SkeletalMeshes found!"), Count);
-
 		// Then do the animation recompression
 		UE_LOG(LogPackageUtilities, Warning, TEXT("Recompressing all animations..."));
 		DoActionToAllPackages<UAnimSequence, CompressAnimationsFunctor>(this, ParamsUpperCase);
