Copying //UE4/Dev-Framework to //UE4/Dev-Main (Source: //UE4/Dev-Framework @ 3582324)

diff --git a/Engine/Source/Editor/PropertyEditor/Private/PropertyHandleImpl.cpp b/Engine/Source/Editor/PropertyEditor/Private/PropertyHandleImpl.cpp
--- a/Engine/Source/Editor/PropertyEditor/Private/PropertyHandleImpl.cpp
+++ b/Engine/Source/Editor/PropertyEditor/Private/PropertyHandleImpl.cpp
@@ -899,14 +899,14 @@ FPropertyAccess::Result FPropertyValueImpl::OnUseSelected()
 		UObjectPropertyBase* ObjProp = Cast<UObjectPropertyBase>( NodeProperty );
 		UInterfaceProperty* IntProp = Cast<UInterfaceProperty>( NodeProperty );
 		UClassProperty* ClassProp = Cast<UClassProperty>( NodeProperty );
-		UAssetClassProperty* AssetClassProperty = Cast<UAssetClassProperty>( NodeProperty );
+		USoftClassProperty* SoftClassProperty = Cast<USoftClassProperty>( NodeProperty );
 		UClass* const InterfaceThatMustBeImplemented = ObjProp ? ObjProp->GetOwnerProperty()->GetClassMetaData(TEXT("MustImplement")) : nullptr;
 
-		if(ClassProp || AssetClassProperty)
+		if(ClassProp || SoftClassProperty)
 		{
 			FEditorDelegates::LoadSelectedAssetsIfNeeded.Broadcast();
 
-			const UClass* const SelectedClass = GEditor->GetFirstSelectedClass(ClassProp ? ClassProp->MetaClass : AssetClassProperty->MetaClass);
+			const UClass* const SelectedClass = GEditor->GetFirstSelectedClass(ClassProp ? ClassProp->MetaClass : SoftClassProperty->MetaClass);
 			if(SelectedClass)
 			{
 				if (!InterfaceThatMustBeImplemented || SelectedClass->ImplementsInterface(InterfaceThatMustBeImplemented))
@@ -2596,11 +2596,11 @@ bool FPropertyHandleBase::GeneratePossibleValues(TArray< TSharedPtr<FString> >&
 			}
 		}
 	}
-	else if( Property->IsA(UClassProperty::StaticClass()) || Property->IsA(UAssetClassProperty::StaticClass()) )		
+	else if( Property->IsA(UClassProperty::StaticClass()) || Property->IsA(USoftClassProperty::StaticClass()) )		
 	{
 		UClass* MetaClass = Property->IsA(UClassProperty::StaticClass()) 
 			? CastChecked<UClassProperty>(Property)->MetaClass
-			: CastChecked<UAssetClassProperty>(Property)->MetaClass;
+			: CastChecked<USoftClassProperty>(Property)->MetaClass;
 
 		TSharedPtr< FString > NoneStr( new FString( TEXT("None") ) );
 		OutOptionStrings.Add( NoneStr );
@@ -3390,7 +3390,7 @@ FPropertyAccess::Result FPropertyHandleObject::SetValue(const FAssetData& NewVal
 	// Instanced references can not be set this way (most likely editinlinenew )
 	if (!Property->HasAnyPropertyFlags(CPF_InstancedReference))
 	{
-		if ( !Property->IsA( UAssetObjectProperty::StaticClass() ) )
+		if ( !Property->IsA( USoftObjectProperty::StaticClass() ) )
 		{
 			// Make sure the asset is loaded if we are not a string asset reference.
 			NewValue.GetAsset();
diff --git a/Engine/Source/Editor/PropertyEditor/Private/UserInterface/PropertyEditor/SPropertyEditorAsset.cpp b/Engine/Source/Editor/PropertyEditor/Private/UserInterface/PropertyEditor/SPropertyEditorAsset.cpp
--- a/Engine/Source/Editor/PropertyEditor/Private/UserInterface/PropertyEditor/SPropertyEditorAsset.cpp
+++ b/Engine/Source/Editor/PropertyEditor/Private/UserInterface/PropertyEditor/SPropertyEditorAsset.cpp
@@ -10,6 +10,7 @@
 #include "UObject/UObjectHash.h"
 #include "UObject/UObjectIterator.h"
 #include "Widgets/Layout/SBox.h"
+#include "Widgets/Images/SImage.h"
 #include "Components/SkeletalMeshComponent.h"
 #include "Particles/ParticleSystem.h"
 #include "UserInterface/PropertyEditor/PropertyEditorConstants.h"
@@ -38,22 +39,21 @@ bool SPropertyEditorAsset::ShouldDisplayThumbnail( const FArguments& InArgs, con
 		// also check metadata for thumbnail & text display
 		if(InArgs._ThumbnailPool.IsValid())
 		{
-
 			const UProperty* ArrayParent = PropertyEditorHelpers::GetArrayParent( *PropertyEditor->GetPropertyNode() );
 			const UProperty* SetParent = PropertyEditorHelpers::GetSetParent( *PropertyEditor->GetPropertyNode() );
 			const UProperty* MapParent = PropertyEditorHelpers::GetMapParent( *PropertyEditor->GetPropertyNode() );
 
 			const UProperty* PropertyToCheck = PropertyEditor->GetProperty();
-			if( ArrayParent != NULL )
+			if( ArrayParent != nullptr )
 			{
 				// If the property is a child of an array property, the parent will have the display thumbnail metadata
 				PropertyToCheck = ArrayParent;
 			}
-			else if ( SetParent != NULL )
+			else if ( SetParent != nullptr )
 			{
 				PropertyToCheck = SetParent;
 			}
-			else if ( MapParent != NULL )
+			else if ( MapParent != nullptr )
 			{
 				PropertyToCheck = MapParent;
 			}
@@ -64,10 +64,8 @@ bool SPropertyEditorAsset::ShouldDisplayThumbnail( const FArguments& InArgs, con
 				bDisplayThumbnail = DisplayThumbnailString == TEXT("true");
 			}
 		}
-
 	}
 
-
 	return bDisplayThumbnail;
 }
 
@@ -173,7 +171,7 @@ void SPropertyEditorAsset::Construct( const FArguments& InArgs, const TSharedPtr
 		NewAssetFactories = PropertyCustomizationHelpers::GetNewAssetFactoriesForClasses(CustomClassFilters);
 	}
 	
-	TSharedPtr<SHorizontalBox> ValueContentBox = NULL;
+	TSharedPtr<SHorizontalBox> ValueContentBox = nullptr;
 	ChildSlot
 	[
 		SNew( SAssetDropTarget )
@@ -233,11 +231,26 @@ void SPropertyEditorAsset::Construct( const FArguments& InArgs, const TSharedPtr
 		.ContentPadding(2.0f)
 		.ButtonContent()
 		[
-			// Show the name of the asset or actor
-			SNew(STextBlock)
-			.TextStyle( FEditorStyle::Get(), "PropertyEditor.AssetClass" )
-			.Font( FEditorStyle::GetFontStyle( PropertyEditorConstants::PropertyFontStyle ) )
-			.Text(this,&SPropertyEditorAsset::OnGetAssetName)
+			SNew(SHorizontalBox)
+			+SHorizontalBox::Slot()
+			.AutoWidth()
+			.HAlign(HAlign_Left)
+			.VAlign(VAlign_Center)
+			[
+				SNew(SImage)
+				.Image( this, &SPropertyEditorAsset::GetStatusIcon )
+			]
+
+			+SHorizontalBox::Slot()
+			.FillWidth(1)
+			.VAlign(VAlign_Center)
+			[
+				// Show the name of the asset or actor
+				SNew(STextBlock)
+				.TextStyle( FEditorStyle::Get(), "PropertyEditor.AssetClass" )
+				.Font( FEditorStyle::GetFontStyle( PropertyEditorConstants::PropertyFontStyle ) )
+				.Text(this,&SPropertyEditorAsset::OnGetAssetName)
+			]
 		];
 
 	if (bOldEnableAttribute && !InArgs._EnableContentPicker)
@@ -257,9 +270,7 @@ void SPropertyEditorAsset::Construct( const FArguments& InArgs, const TSharedPtr
 
 		AssetThumbnail = MakeShareable( new FAssetThumbnail( Value.AssetData, InArgs._ThumbnailSize.X, InArgs._ThumbnailSize.Y, InArgs._ThumbnailPool ) );
 
-
 		FAssetThumbnailConfig AssetThumbnailConfig;
-
 		TSharedPtr<IAssetTypeActions> AssetTypeActions;
 		if (ObjectClass != nullptr)
 		{
@@ -469,6 +480,56 @@ const FSlateBrush* SPropertyEditorAsset::GetThumbnailBorder() const
 	}
 }
 
+const FSlateBrush* SPropertyEditorAsset::GetStatusIcon() const
+{
+	static FSlateNoResource EmptyBrush = FSlateNoResource();
+
+	EActorReferenceState State = GetActorReferenceState();
+
+	if (State == EActorReferenceState::Unknown)
+	{
+		return FEditorStyle::GetBrush("Icons.Warning");
+	}
+	else if (State == EActorReferenceState::Error)
+	{
+		return FEditorStyle::GetBrush("Icons.Error");
+	}
+
+	return &EmptyBrush;
+}
+
+SPropertyEditorAsset::EActorReferenceState SPropertyEditorAsset::GetActorReferenceState() const
+{
+	if (bIsActor)
+	{
+		FObjectOrAssetData Value;
+		GetValue(Value);
+
+		if (Value.Object != nullptr)
+		{
+			return EActorReferenceState::Loaded;
+		}
+		else if (Value.ObjectPath.IsNull())
+		{
+			return EActorReferenceState::Null;
+		}
+		else
+		{
+			// Get a path pointing to the owning map
+			FSoftObjectPath MapObjectPath = FSoftObjectPath(Value.ObjectPath.GetAssetPathName(), FString());
+
+			if (MapObjectPath.ResolveObject())
+			{
+				// If the map is valid but the 
+				return EActorReferenceState::Error;
+			}
+
+			return EActorReferenceState::Unknown;
+		}
+	}
+	return EActorReferenceState::NotAnActor;
+}
+
 void SPropertyEditorAsset::Tick( const FGeometry& AllottedGeometry, const double InCurrentTime, const float InDeltaTime )
 {
 	if( AssetThumbnail.IsValid() )
@@ -501,9 +562,9 @@ bool SPropertyEditorAsset::Supports( const UProperty* NodeProperty )
 	const UObjectPropertyBase* ObjectProperty = Cast<const UObjectPropertyBase>( NodeProperty );
 	const UInterfaceProperty* InterfaceProperty = Cast<const UInterfaceProperty>( NodeProperty );
 
-	if ( ( ObjectProperty != NULL || InterfaceProperty != NULL )
+	if ( ( ObjectProperty != nullptr || InterfaceProperty != nullptr )
 		 && !NodeProperty->IsA(UClassProperty::StaticClass()) 
-		 && !NodeProperty->IsA(UAssetClassProperty::StaticClass()) )
+		 && !NodeProperty->IsA(USoftClassProperty::StaticClass()) )
 	{
 		return true;
 	}
@@ -563,12 +624,12 @@ FText SPropertyEditorAsset::OnGetAssetName() const
 	FText Name = LOCTEXT("None", "None");
 	if( Result == FPropertyAccess::Success )
 	{
-		if(Value.Object != NULL)
+		if(Value.Object != nullptr)
 		{
 			if( bIsActor )
 			{
 				AActor* Actor = CastChecked<AActor>(Value.Object);
-				Name = FText::FromString(Actor->GetActorLabel());
+				Name = FText::AsCultureInvariant(Actor->GetActorLabel());
 			}
 			else if (UField* AsField = Cast<UField>(Value.Object))
 			{
@@ -576,12 +637,16 @@ FText SPropertyEditorAsset::OnGetAssetName() const
 			}
 			else
 			{
-				Name = FText::FromString(Value.Object->GetName());
+				Name = FText::AsCultureInvariant(Value.Object->GetName());
 			}
 		}
 		else if( Value.AssetData.IsValid() )
 		{
-			Name = FText::FromName(Value.AssetData.AssetName);
+			Name = FText::AsCultureInvariant(Value.AssetData.AssetName.ToString());
+		}
+		else if (Value.ObjectPath.IsValid())
+		{
+			Name = FText::AsCultureInvariant(Value.ObjectPath.ToString());
 		}
 	}
 	else if( Result == FPropertyAccess::MultipleValues )
@@ -597,7 +662,7 @@ FText SPropertyEditorAsset::OnGetAssetClassName() const
 	UClass* Class = GetDisplayedClass();
 	if(Class)
 	{
-		return FText::FromString(Class->GetName());
+		return FText::AsCultureInvariant(Class->GetName());
 	}
 	return FText::GetEmpty();
 }
@@ -611,14 +676,37 @@ FText SPropertyEditorAsset::OnGetToolTip() const
 
 	if( Result == FPropertyAccess::Success )
 	{
-		if(Value.Object != NULL && !bIsActor )
+		if ( bIsActor )
+		{
+			// Always show full path instead of label
+			EActorReferenceState State = GetActorReferenceState();
+			FFormatNamedArguments Args;
+			Args.Add(TEXT("Actor"), FText::AsCultureInvariant(Value.ObjectPath.ToString()));
+			if (State == EActorReferenceState::Null)
+			{
+				ToolTipText = LOCTEXT("EmptyActorReference", "None");
+			}
+			else if (State == EActorReferenceState::Error)
+			{
+				ToolTipText = FText::Format(LOCTEXT("BrokenActorReference", "Broken reference to Actor ID '{Actor}', it was deleted or renamed"), Args);
+			}
+			else if (State == EActorReferenceState::Unknown)
+			{
+				ToolTipText = FText::Format(LOCTEXT("UnknownActorReference", "Unloaded reference to Actor ID '{Actor}', use Browse to load level"), Args);
+			}
+			else
+			{
+				ToolTipText = FText::Format(LOCTEXT("GoodActorReference", "Reference to Actor ID '{Actor}'"), Args);
+			}
+		}
+		else if( Value.Object != nullptr )
 		{
 			// Display the package name which is a valid path to the object without redundant information
-			ToolTipText = FText::FromString(Value.Object->GetOutermost()->GetName());
+			ToolTipText = FText::AsCultureInvariant(Value.Object->GetOutermost()->GetName());
 		}
 		else if( Value.AssetData.IsValid() )
 		{
-			ToolTipText = FText::FromName(Value.AssetData.PackageName);
+			ToolTipText = FText::AsCultureInvariant(Value.AssetData.PackageName.ToString());
 		}
 	}
 	else if( Result == FPropertyAccess::MultipleValues )
@@ -628,7 +716,7 @@ FText SPropertyEditorAsset::OnGetToolTip() const
 
 	if( ToolTipText.IsEmpty() )
 	{
-		ToolTipText = FText::FromString(ObjectPath.Get());
+		ToolTipText = FText::AsCultureInvariant(ObjectPath.Get());
 	}
 
 	return ToolTipText;
@@ -664,10 +752,10 @@ FPropertyAccess::Result SPropertyEditorAsset::GetValue( FObjectOrAssetData& OutV
 
 	if( PropertyEditor.IsValid() && PropertyEditor->GetPropertyHandle()->IsValidHandle() )
 	{
-		UObject* Object = NULL;
+		UObject* Object = nullptr;
 		Result = PropertyEditor->GetPropertyHandle()->GetValue(Object);
 
-		if (Object == NULL)
+		if (Object == nullptr)
 		{
 			// Check to see if it's pointing to an unloaded object
 			FString CurrentObjectPath;
@@ -675,16 +763,32 @@ FPropertyAccess::Result SPropertyEditorAsset::GetValue( FObjectOrAssetData& OutV
 
 			if (CurrentObjectPath.Len() > 0 && CurrentObjectPath != TEXT("None"))
 			{
-				if( !CachedAssetData.IsValid() || CachedAssetData.ObjectPath.ToString() != CurrentObjectPath )
+				FSoftObjectPath SoftObjectPath = FSoftObjectPath(CurrentObjectPath);
+
+				if (SoftObjectPath.IsAsset())
 				{
-					static FName AssetRegistryName("AssetRegistry");
+					if (!CachedAssetData.IsValid() || CachedAssetData.ObjectPath.ToString() != CurrentObjectPath)
+					{
+						static FName AssetRegistryName("AssetRegistry");
 
-					FAssetRegistryModule& AssetRegistryModule = FModuleManager::Get().LoadModuleChecked<FAssetRegistryModule>(AssetRegistryName);
-					CachedAssetData = AssetRegistryModule.Get().GetAssetByObjectPath( *CurrentObjectPath );
+						FAssetRegistryModule& AssetRegistryModule = FModuleManager::Get().LoadModuleChecked<FAssetRegistryModule>(AssetRegistryName);
+						CachedAssetData = AssetRegistryModule.Get().GetAssetByObjectPath(*CurrentObjectPath);
+					}
+
+					Result = FPropertyAccess::Success;
+					OutValue = FObjectOrAssetData(CachedAssetData);
 				}
+				else
+				{
+					// This is an actor or other subobject reference
+					if (CachedAssetData.IsValid())
+					{
+						CachedAssetData = FAssetData();
+					}
 
-				Result = FPropertyAccess::Success;
-				OutValue = FObjectOrAssetData( CachedAssetData );
+					Result = FPropertyAccess::Success;
+					OutValue = FObjectOrAssetData(SoftObjectPath);
+				}
 
 				return Result;
 			}
@@ -702,13 +806,13 @@ FPropertyAccess::Result SPropertyEditorAsset::GetValue( FObjectOrAssetData& OutV
 	}
 	else
 	{
-		UObject* Object = NULL;
+		UObject* Object = nullptr;
 		if (PropertyHandle.IsValid())
 		{
 			Result = PropertyHandle->GetValue(Object);
 		}
 
-		if (Object != NULL)
+		if (Object != nullptr)
 		{
 #if !UE_BUILD_SHIPPING
 			if (!Object->IsValidLowLevel())
@@ -725,37 +829,48 @@ FPropertyAccess::Result SPropertyEditorAsset::GetValue( FObjectOrAssetData& OutV
 			const FString CurrentObjectPath = ObjectPath.Get();
 			Result = FPropertyAccess::Success;
 
-			if (CurrentObjectPath != TEXT("None") && (!CachedAssetData.IsValid() || CachedAssetData.ObjectPath.ToString() != CurrentObjectPath))
+			FSoftObjectPath SoftObjectPath = FSoftObjectPath(CurrentObjectPath);
+
+			if (SoftObjectPath.IsAsset())
 			{
-				static FName AssetRegistryName("AssetRegistry");
+				if (CurrentObjectPath != TEXT("None") && (!CachedAssetData.IsValid() || CachedAssetData.ObjectPath.ToString() != CurrentObjectPath))
+				{
+					static FName AssetRegistryName("AssetRegistry");
 
-				FAssetRegistryModule& AssetRegistryModule = FModuleManager::Get().LoadModuleChecked<FAssetRegistryModule>(AssetRegistryName);
-				CachedAssetData = AssetRegistryModule.Get().GetAssetByObjectPath(*CurrentObjectPath);
+					FAssetRegistryModule& AssetRegistryModule = FModuleManager::Get().LoadModuleChecked<FAssetRegistryModule>(AssetRegistryName);
+					CachedAssetData = AssetRegistryModule.Get().GetAssetByObjectPath(*CurrentObjectPath);
+				}
 
-				if (PropertyHandle.IsValid())
+				OutValue = FObjectOrAssetData(CachedAssetData);
+			}
+			else
+			{
+				// This is an actor or other subobject reference
+				if (CachedAssetData.IsValid())
 				{
-					// No property editor was specified so check if multiple property values are associated with the property handle
-					TArray<FString> ObjectValues;
-					PropertyHandle->GetPerObjectValues(ObjectValues);
+					CachedAssetData = FAssetData();
+				}
 
-					if (ObjectValues.Num() > 1)
+				OutValue = FObjectOrAssetData(SoftObjectPath);
+			}
+
+			if (PropertyHandle.IsValid())
+			{
+				// No property editor was specified so check if multiple property values are associated with the property handle
+				TArray<FString> ObjectValues;
+				PropertyHandle->GetPerObjectValues(ObjectValues);
+
+				if (ObjectValues.Num() > 1)
+				{
+					for (int32 ObjectIndex = 1; ObjectIndex < ObjectValues.Num() && Result == FPropertyAccess::Success; ++ObjectIndex)
 					{
-						for (int32 ObjectIndex = 1; ObjectIndex < ObjectValues.Num() && Result == FPropertyAccess::Success; ++ObjectIndex)
+						if (ObjectValues[ObjectIndex] != ObjectValues[0])
 						{
-							if (ObjectValues[ObjectIndex] != ObjectValues[0])
-							{
-								Result = FPropertyAccess::MultipleValues;
-							}
+							Result = FPropertyAccess::MultipleValues;
 						}
 					}
 				}
 			}
-			else if (CurrentObjectPath == TEXT("None"))
-			{
-				CachedAssetData = FAssetData();
-			}
-
-			OutValue = FObjectOrAssetData(CachedAssetData);
 		}
 	}
 
@@ -766,7 +881,7 @@ UClass* SPropertyEditorAsset::GetDisplayedClass() const
 {
 	FObjectOrAssetData Value;
 	GetValue( Value );
-	if(Value.Object != NULL)
+	if(Value.Object != nullptr)
 	{
 		return Value.Object->GetClass();
 	}
@@ -808,6 +923,13 @@ void SPropertyEditorAsset::OnBrowse()
 	FObjectOrAssetData Value;
 	GetValue( Value );
 
+	// Try loading owning object
+	if (Value.Object == nullptr && Value.ObjectPath.IsValid())
+	{
+		FSoftObjectPath MapObjectPath = FSoftObjectPath(Value.ObjectPath.GetAssetPathName(), FString());
+		MapObjectPath.TryLoad();
+	}
+
 	if(PropertyEditor.IsValid() && Value.Object)
 	{
 		// This code only works on loaded objects
@@ -829,8 +951,8 @@ FText SPropertyEditorAsset::GetOnBrowseToolTip() const
 	if (Value.Object)
 	{
 		FFormatNamedArguments Args;
-		Args.Add(TEXT("Asset"), FText::FromString(Value.Object->GetName()));
-		if (Value.Object->IsA(AActor::StaticClass()))
+		Args.Add(TEXT("Asset"), FText::AsCultureInvariant(Value.Object->GetName()));
+		if (bIsActor)
 		{
 			return FText::Format(LOCTEXT( "BrowseToAssetInViewport", "Select '{Asset}' in the viewport"), Args);
 		}
@@ -856,7 +978,7 @@ void SPropertyEditorAsset::OnUse()
 		FEditorDelegates::LoadSelectedAssetsIfNeeded.Broadcast();
 
 		// try to get a selected object of our class
-		UObject* Selection = NULL;
+		UObject* Selection = nullptr;
 		if( ObjectClass && ObjectClass->IsChildOf( AActor::StaticClass() ) )
 		{
 			Selection = GEditor->GetSelectedActors()->GetTop( ObjectClass );
@@ -868,11 +990,11 @@ void SPropertyEditorAsset::OnUse()
 		}
 
 		// Check against custom asset filter
-		if (Selection != NULL
+		if (Selection != nullptr
 			&& OnShouldFilterAsset.IsBound()
 			&& OnShouldFilterAsset.Execute(FAssetData(Selection)))
 		{
-			Selection = NULL;
+			Selection = nullptr;
 		}
 
 		if( Selection )
@@ -884,7 +1006,7 @@ void SPropertyEditorAsset::OnUse()
 
 void SPropertyEditorAsset::OnClear()
 {
-	SetValue(NULL);
+	SetValue(nullptr);
 }
 
 FSlateColor SPropertyEditorAsset::GetAssetClassColor()
@@ -901,7 +1023,7 @@ FSlateColor SPropertyEditorAsset::GetAssetClassColor()
 
 bool SPropertyEditorAsset::OnAssetDraggedOver( const UObject* InObject ) const
 {
-	if (CanEdit() && InObject != NULL && InObject->IsA(ObjectClass))
+	if (CanEdit() && InObject != nullptr && InObject->IsA(ObjectClass))
 	{
 		// Check against custom asset filter
 		if (!OnShouldFilterAsset.IsBound()
@@ -932,6 +1054,10 @@ void SPropertyEditorAsset::OnCopy()
 	{
 		FPlatformApplicationMisc::ClipboardCopy(*Value.AssetData.GetExportTextName());
 	}
+	else
+	{
+		FPlatformApplicationMisc::ClipboardCopy(*Value.ObjectPath.ToString());
+	}
 }
 
 void SPropertyEditorAsset::OnPaste()
@@ -941,11 +1067,11 @@ void SPropertyEditorAsset::OnPaste()
 
 	if(DestPath == TEXT("None"))
 	{
-		SetValue(NULL);
+		SetValue(nullptr);
 	}
 	else
 	{
-		UObject* Object = LoadObject<UObject>(NULL, *DestPath);
+		UObject* Object = LoadObject<UObject>(nullptr, *DestPath);
 		if(Object && Object->IsA(ObjectClass))
 		{
 			// Check against custom asset filter
@@ -1017,18 +1143,18 @@ bool SPropertyEditorAsset::CanSetBasedOnCustomClasses( const FAssetData& InAsset
 
 UClass* SPropertyEditorAsset::GetObjectPropertyClass(const UProperty* Property)
 {
-	UClass* Class = NULL;
+	UClass* Class = nullptr;
 
-	if (Cast<const UObjectPropertyBase>(Property) != NULL)
+	if (Cast<const UObjectPropertyBase>(Property) != nullptr)
 	{
 		Class = Cast<const UObjectPropertyBase>(Property)->PropertyClass;
 	}
-	else if (Cast<const UInterfaceProperty>(Property) != NULL)
+	else if (Cast<const UInterfaceProperty>(Property) != nullptr)
 	{
 		Class = Cast<const UInterfaceProperty>(Property)->InterfaceClass;
 	}
 
-	if (!ensureMsgf(Class != NULL, TEXT("Property (%s) is not an object or interface class"), Property ? *Property->GetFullName() : TEXT("null")))
+	if (!ensureMsgf(Class != nullptr, TEXT("Property (%s) is not an object or interface class"), Property ? *Property->GetFullName() : TEXT("null")))
 	{
 		Class = UObject::StaticClass();
 	}
diff --git a/Engine/Source/Editor/PropertyEditor/Private/UserInterface/PropertyEditor/SPropertyEditorAsset.h b/Engine/Source/Editor/PropertyEditor/Private/UserInterface/PropertyEditor/SPropertyEditorAsset.h
--- a/Engine/Source/Editor/PropertyEditor/Private/UserInterface/PropertyEditor/SPropertyEditorAsset.h
+++ b/Engine/Source/Editor/PropertyEditor/Private/UserInterface/PropertyEditor/SPropertyEditorAsset.h
@@ -29,6 +29,20 @@ class SPropertyEditorAsset : public SCompoundWidget
 {
 public:
 
+	enum class EActorReferenceState : uint8
+	{
+		// This is not pointing to an actor
+		NotAnActor,
+		// This is specifically pointing at no actor
+		Null,
+		// The pointed to actor is fully loaded in memory
+		Loaded,
+		// The pointed to actor is unknown because the pointed to map is not loaded 
+		Unknown,
+		// This is a known bad reference, the owning map is loaded but the actor does not exist
+		Error,
+	};
+
 	SLATE_BEGIN_ARGS( SPropertyEditorAsset )
 		: _AssetFont( FEditorStyle::GetFontStyle("PropertyEditor.AssetName.Font") ) 
 		, _ClassFont( FEditorStyle::GetFontStyle("PropertyEditor.AssetClass.Font") ) 
@@ -87,22 +101,30 @@ private:
 	struct FObjectOrAssetData
 	{
 		UObject* Object;
+		FSoftObjectPath ObjectPath;
 		FAssetData AssetData;
 
-		FObjectOrAssetData( UObject* InObject = NULL )
+		FObjectOrAssetData( UObject* InObject = nullptr )
 			: Object( InObject )
+			, ObjectPath( Object )
 		{
-			AssetData = InObject != NULL && !InObject->IsA<AActor>() ? FAssetData( InObject ) : FAssetData();
+			AssetData = InObject != nullptr && !InObject->IsA<AActor>() ? FAssetData( InObject ) : FAssetData();
 		}
 
+		FObjectOrAssetData( const FSoftObjectPath& InObjectPath )
+			: Object(nullptr)
+			, ObjectPath(InObjectPath)
+		{}
+
 		FObjectOrAssetData( const FAssetData& InAssetData )
-			: Object( NULL )
+			: Object( nullptr )
+			, ObjectPath( InAssetData.ToSoftObjectPath() )
 			, AssetData( InAssetData )
 		{}
 
 		bool IsValid() const
 		{
-			return Object != NULL || AssetData.IsValid();
+			return Object != nullptr || ObjectPath.IsValid() || AssetData.IsValid();
 		}
 	};
 
@@ -116,6 +138,12 @@ private:
 	/** Gets the border brush to show around the thumbnail, changes when the user hovers on it. */
 	const FSlateBrush* GetThumbnailBorder() const;
 
+	/** Returns the status icon, empty for non actors */
+	const FSlateBrush* GetStatusIcon() const;
+
+	/** Returns the state of this actor reference */
+	EActorReferenceState GetActorReferenceState() const;
+
 	/** 
 	 * Get the content to be displayed in the asset/actor picker menu 
 	 * @returns the widget for the menu content
diff --git a/Engine/Source/Editor/PropertyEditor/Private/UserInterface/PropertyEditor/SPropertyEditorClass.cpp b/Engine/Source/Editor/PropertyEditor/Private/UserInterface/PropertyEditor/SPropertyEditorClass.cpp
--- a/Engine/Source/Editor/PropertyEditor/Private/UserInterface/PropertyEditor/SPropertyEditorClass.cpp
+++ b/Engine/Source/Editor/PropertyEditor/Private/UserInterface/PropertyEditor/SPropertyEditorClass.cpp
@@ -71,7 +71,7 @@ bool SPropertyEditorClass::Supports(const TSharedRef< class FPropertyEditor >& I
 	const UProperty* Property = InPropertyEditor->GetProperty();
 	int32 ArrayIndex = PropertyNode->GetArrayIndex();
 
-	if ((Property->IsA(UClassProperty::StaticClass()) || Property->IsA(UAssetClassProperty::StaticClass())) 
+	if ((Property->IsA(UClassProperty::StaticClass()) || Property->IsA(USoftClassProperty::StaticClass())) 
 		&& ((ArrayIndex == -1 && Property->ArrayDim == 1) || (ArrayIndex > -1 && Property->ArrayDim > 0)))
 	{
 		return true;
@@ -92,9 +92,9 @@ void SPropertyEditorClass::Construct(const FArguments& InArgs, const TSharedPtr<
 		{
 			MetaClass = ClassProp->MetaClass;
 		}
-		else if (UAssetClassProperty* const AssetClassProperty = Cast<UAssetClassProperty>(Property))
+		else if (USoftClassProperty* const SoftClassProperty = Cast<USoftClassProperty>(Property))
 		{
-			MetaClass = AssetClassProperty->MetaClass;
+			MetaClass = SoftClassProperty->MetaClass;
 		}
 		else
 		{
diff --git a/Engine/Source/Editor/SceneOutliner/Private/SceneOutlinerItemLabelColumn.cpp b/Engine/Source/Editor/SceneOutliner/Private/SceneOutlinerItemLabelColumn.cpp
--- a/Engine/Source/Editor/SceneOutliner/Private/SceneOutlinerItemLabelColumn.cpp
+++ b/Engine/Source/Editor/SceneOutliner/Private/SceneOutlinerItemLabelColumn.cpp
@@ -322,7 +322,7 @@ private:
 		if (Actor && Actor->IsActorLabelEditable() && !InLabel.ToString().Equals(Actor->GetActorLabel(), ESearchCase::CaseSensitive))
 		{
 			const FScopedTransaction Transaction( LOCTEXT( "SceneOutlinerRenameActorTransaction", "Rename Actor" ) );
-			Actor->SetActorLabel( InLabel.ToString() );
+			FActorLabelUtilities::RenameExistingActor(Actor, InLabel.ToString());
 
 			auto Outliner = WeakSceneOutliner.Pin();
 			if (Outliner.IsValid())
diff --git a/Engine/Source/Editor/UMGEditor/Private/Nodes/K2Node_CreateWidget.cpp b/Engine/Source/Editor/UMGEditor/Private/Nodes/K2Node_CreateWidget.cpp
--- a/Engine/Source/Editor/UMGEditor/Private/Nodes/K2Node_CreateWidget.cpp
+++ b/Engine/Source/Editor/UMGEditor/Private/Nodes/K2Node_CreateWidget.cpp
@@ -78,7 +78,7 @@ UEdGraphPin* UK2Node_CreateWidget::GetOwningPlayerPin() const
 	return Pin;
 }
 
-bool UK2Node_CreateWidget::IsSpawnVarPin(UEdGraphPin* Pin)
+bool UK2Node_CreateWidget::IsSpawnVarPin(UEdGraphPin* Pin) const
 {
 	return( Super::IsSpawnVarPin(Pin) &&
 		Pin->PinName != FK2Node_CreateWidgetHelper::OwningPlayerPinName );
diff --git a/Engine/Source/Editor/UMGEditor/Private/Nodes/K2Node_CreateWidget.h b/Engine/Source/Editor/UMGEditor/Private/Nodes/K2Node_CreateWidget.h
--- a/Engine/Source/Editor/UMGEditor/Private/Nodes/K2Node_CreateWidget.h
+++ b/Engine/Source/Editor/UMGEditor/Private/Nodes/K2Node_CreateWidget.h
@@ -36,5 +36,5 @@ protected:
 	/** Gets base class to use for the 'class' pin.  UObject by default. */
 	virtual UClass* GetClassPinBaseClass() const override;
 	/**  */
-	virtual bool IsSpawnVarPin(UEdGraphPin* Pin) override;
+	virtual bool IsSpawnVarPin(UEdGraphPin* Pin) const override;
 };
diff --git a/Engine/Source/Editor/UMGEditor/Private/Palette/SPaletteView.cpp b/Engine/Source/Editor/UMGEditor/Private/Palette/SPaletteView.cpp
--- a/Engine/Source/Editor/UMGEditor/Private/Palette/SPaletteView.cpp
+++ b/Engine/Source/Editor/UMGEditor/Private/Palette/SPaletteView.cpp
@@ -411,7 +411,7 @@ void SPaletteView::BuildClassWidgetList()
 	auto ActiveWidgetBlueprintClass = GetBlueprint()->GeneratedClass;
 	FName ActiveWidgetBlueprintClassName = ActiveWidgetBlueprintClass->GetFName();
 
-	TArray<FStringClassReference> WidgetClassesToHide = GetDefault<UUMGEditorProjectSettings>()->WidgetClassesToHide;
+	TArray<FSoftClassPath> WidgetClassesToHide = GetDefault<UUMGEditorProjectSettings>()->WidgetClassesToHide;
 
 	// Locate all UWidget classes from code and loaded widget BPs
 	for (TObjectIterator<UClass> ClassIt; ClassIt; ++ClassIt)
@@ -446,7 +446,7 @@ void SPaletteView::BuildClassWidgetList()
 
 		// Excludes this widget if it is on the hide list
 		bool bIsOnList = false;
-		for (FStringClassReference Widget : WidgetClassesToHide)
+		for (FSoftClassPath Widget : WidgetClassesToHide)
 		{
 			if (WidgetAssetData.ObjectPath.ToString().Find(Widget.ToString()) == 0)
 			{
@@ -519,7 +519,7 @@ void SPaletteView::BuildClassWidgetList()
 
 		// Excludes this widget if it is on the hide list
 		bool bIsOnList = false;
-		for (FStringClassReference Widget : WidgetClassesToHide)
+		for (FSoftClassPath Widget : WidgetClassesToHide)
 		{
 			if (Widget.ToString().Find(WidgetBPAssetData.ObjectPath.ToString()) == 0)
 			{
diff --git a/Engine/Source/Editor/UMGEditor/Public/UMGEditorProjectSettings.h b/Engine/Source/Editor/UMGEditor/Public/UMGEditorProjectSettings.h
--- a/Engine/Source/Editor/UMGEditor/Public/UMGEditorProjectSettings.h
+++ b/Engine/Source/Editor/UMGEditor/Public/UMGEditorProjectSettings.h
@@ -5,7 +5,7 @@
 #include "CoreMinimal.h"
 #include "UObject/ObjectMacros.h"
 #include "UObject/Object.h"
-#include "Misc/StringClassReference.h"
+#include "UObject/SoftObjectPath.h"
 #include "Engine/DeveloperSettings.h"
 #include "UMGEditorProjectSettings.generated.h"
 
@@ -66,5 +66,5 @@ public:
 	TArray<FString> CategoriesToHide;
 
 	UPROPERTY(EditAnywhere, config, Category = "Class Filtering", meta = (MetaClass = "Widget"))
-	TArray<FStringClassReference> WidgetClassesToHide;
+	TArray<FSoftClassPath> WidgetClassesToHide;
 };
diff --git a/Engine/Source/Editor/UnrealEd/Classes/CookOnTheSide/CookOnTheFlyServer.h b/Engine/Source/Editor/UnrealEd/Classes/CookOnTheSide/CookOnTheFlyServer.h
--- a/Engine/Source/Editor/UnrealEd/Classes/CookOnTheSide/CookOnTheFlyServer.h
+++ b/Engine/Source/Editor/UnrealEd/Classes/CookOnTheSide/CookOnTheFlyServer.h
@@ -1305,8 +1305,6 @@ public:
 	* Callbacks from UObject globals
 	*/
 	void PreGarbageCollect();
-	void OnStringAssetReferenceLoadedPackage(const FName& PackageName);
-
 
 private:
 
@@ -1543,7 +1541,7 @@ private:
 	 * @param RedirectedPaths map of original to redirected object paths
 	 * @return true if the Package contains a redirector false otherwise
 	 */
-	bool ContainsRedirector(const FName& PackageName, TMap<FString,FString>& RedirectedPaths) const;
+	bool ContainsRedirector(const FName& PackageName, TMap<FName, FName>& RedirectedPaths) const;
 	
 	/**
 	 * Calls BeginCacheForCookedPlatformData on all UObjects in the package
diff --git a/Engine/Source/Editor/UnrealEd/Classes/Editor/EditorEngine.h b/Engine/Source/Editor/UnrealEd/Classes/Editor/EditorEngine.h
--- a/Engine/Source/Editor/UnrealEd/Classes/Editor/EditorEngine.h
+++ b/Engine/Source/Editor/UnrealEd/Classes/Editor/EditorEngine.h
@@ -821,8 +821,8 @@ private:
 	virtual void RemapGamepadControllerIdForPIE(class UGameViewportClient* GameViewport, int32 &ControllerId) override;
 	virtual TSharedPtr<SViewport> GetGameViewportWidget() const override;
 	virtual void TriggerStreamingDataRebuild() override;
-	virtual bool NetworkRemapPath(UNetDriver* Driver, FString &Str, bool reading = true) override;
-	virtual bool NetworkRemapPath(UPendingNetGame* PendingNetGame, FString& Str, bool reading = true) override;
+	virtual bool NetworkRemapPath(UNetDriver* Driver, FString& Str, bool bReading = true) override;
+	virtual bool NetworkRemapPath(UPendingNetGame* PendingNetGame, FString& Str, bool bReading = true) override;
 	virtual bool AreEditorAnalyticsEnabled() const override;
 	virtual void CreateStartupAnalyticsAttributes(TArray<FAnalyticsEventAttribute>& StartSessionAttributes) const override;
 	virtual void VerifyLoadMapWorldCleanup() override;
@@ -1366,10 +1366,11 @@ public:
 	 *
 	 * @param	InWorld				World context
 	 * @param	bVerifyDeletionCanHappen	[opt] If true (default), verify that deletion can be performed.
-	 * @param	bWarnAboutReferences		[opt] If true (default), we prompt the user about referenced actours they are about to delete
+	 * @param	bWarnAboutReferences		[opt] If true (default), we prompt the user about referenced actors they are about to delete
+	 * @param	bWarnAboutSoftReferences	[opt] If true (default), we prompt the user about soft references to actors they are about to delete
 	 * @return								true unless the delete operation was aborted.
 	 */
-	virtual bool edactDeleteSelected(UWorld* InWorld, bool bVerifyDeletionCanHappen=true, bool bWarnAboutReferences = true) { return true; }
+	virtual bool edactDeleteSelected(UWorld* InWorld, bool bVerifyDeletionCanHappen=true, bool bWarnAboutReferences = true, bool bWarnAboutSoftReferences = true) { return true; }
 
 	/**
 	 * Checks the state of the selected actors and notifies the user of any potentially unknown destructive actions which may occur as
@@ -1692,19 +1693,16 @@ public:
 	// @todo gmp: temp hack for Rocket demo
 	void RequestPlaySession(const FVector* StartLocation, const FRotator* StartRotation, bool MobilePreview, bool VulkanPreview, const FString& MobilePreviewTargetDevice);
 
-	/**
-	 * Request to play a game on a remote device 
-	 */
+	/** Request to play a game on a remote device */
 	void RequestPlaySession( const FString& DeviceId, const FString& DeviceName );
 
-	/**
-	 * Cancel request to start a play session
-	 */
+	/** Cancel request to start a play session */
 	void CancelRequestPlaySession();
 
-	/**
-	 * Makes a request to start a play from a Slate editor session
-	 */
+	/** Asks the player to save dirty maps, if this fails it will return false and call CancelRequestPlaySession */
+	bool SaveMapsForPlaySession();
+
+	/** Makes a request to start a play from a Slate editor session */
 	void RequestToggleBetweenPIEandSIE() { bIsToggleBetweenPIEandSIEQueued = true; }
 
 	/** Called when the debugger has paused the active PIE or SIE session */
@@ -1740,16 +1738,6 @@ public:
 	 */
 	void RequestLateJoin();
 
-	/**
-	 * Saves play in editor levels and also fixes up references in AWorldSettings to other levels.
-	 *
-	 * @param	Prefix				Prefix used to save files to disk.
-	 * @param	OutSavedFilenames	The file names of all successfully saved worlds will be added to this (with the P world being the first)
-	 *
-	 * @return	False if the save failed and the user wants to abort what they were doing
-	 */
-	virtual bool SavePlayWorldPackages(UWorld* InWorld, const TCHAR* Prefix, TArray<FString>& OutSavedFilenames);
-
 	/**
 	 * Builds a URL for game spawned by the editor (not including map name!). 
 	 * @param	MapName			The name of the map to put into the URL
@@ -1947,8 +1935,9 @@ public:
 	 *
 	 * @param InWorld		World to get the selected actors from
 	 * @param bShouldCut If true, deletes the selected actors after copying them to the clipboard
+	 * @param bShouldCut If true, this cut is part of a move and the actors will be immediately pasted
 	 */
-	void CopySelectedActorsToClipboard( UWorld* InWorld, const bool bShouldCut );
+	void CopySelectedActorsToClipboard( UWorld* InWorld, const bool bShouldCut, const bool bIsMove = false );
 
 	/**
 	 * Checks to see whether it's possible to perform a paste operation.
@@ -2554,13 +2543,7 @@ private:
 	 * @return	true if a static mesh was loaded; false, otherwise.
 	 */
 	bool LoadPreviewMesh( int32 Index );
-public:
-	/** Creates a PIE world by saving to a temp file and then reloading it */
-	UWorld* CreatePIEWorldBySavingToTemp(FWorldContext &WorldContext, UWorld* InWorld, FString &PlayWorldMapName);
-
-	UWorld* CreatePIEWorldFromEntry(FWorldContext &WorldContext, UWorld* InWorld, FString &PlayWorldMapName);
 
-private:
 	/**
 	 * Login PIE instances with the online platform before actually creating any PIE worlds
 	 *
@@ -2587,6 +2570,9 @@ private:
 	virtual void OnLoginPIEAllComplete();
 
 public:
+	/** Creates a pie world from the default entry map, used by clients that connect to a PIE server */
+	UWorld* CreatePIEWorldFromEntry(FWorldContext &WorldContext, UWorld* InWorld, FString &PlayWorldMapName);
+
 	/**
 	 * Continue the creation of a single PIE world after a login was successful
 	 *
@@ -2880,13 +2866,6 @@ protected:
 
 	void PlayUsingLauncher();
 
-public:
-
-	/** Save the currently loaded world ready for a play session */
-	void SaveWorldForPlay(TArray<FString>& SavedMapNames);
-
-protected:
-
 	/** Called when Matinee is opened */
 	virtual void OnOpenMatinee(){};
 
@@ -2951,7 +2930,7 @@ public:
 protected:
 
 	UPROPERTY(EditAnywhere, config, Category = Advanced, meta = (MetaClass = "ActorGroupingUtils"))
-	FStringClassReference ActorGroupingUtilsClassName;
+	FSoftClassPath ActorGroupingUtilsClassName;
 
 	UPROPERTY()
 	class UActorGroupingUtils* ActorGroupingUtils;
@@ -2987,6 +2966,15 @@ public:
 	 */
 	static void SetActorLabelUnique(AActor* Actor, const FString& NewActorLabel, const FCachedActorLabels* InExistingActorLabels = nullptr);
 
+	/** 
+	 * Does an explicit actor rename. In addition to changing the label this will also fix any soft references pointing to it 
+	 * 
+	 * @param	Actor					The actor to change the label of
+	 * @param	NewActorLabel			The new label string to assign to the actor.  If empty, the actor will have a default label.
+	 * @param	bMakeUnique				If true, it will call SetActorLabelUnique, if false it will use the exact label specified
+	 */
+	static void RenameExistingActor(AActor* Actor, const FString& NewActorLabel, bool bMakeUnique = false);
+
 private:
 	FActorLabelUtilities() {}
 };
diff --git a/Engine/Source/Editor/UnrealEd/Classes/Editor/PropertyEditorTestObject.h b/Engine/Source/Editor/UnrealEd/Classes/Editor/PropertyEditorTestObject.h
--- a/Engine/Source/Editor/UnrealEd/Classes/Editor/PropertyEditorTestObject.h
+++ b/Engine/Source/Editor/UnrealEd/Classes/Editor/PropertyEditorTestObject.h
@@ -6,7 +6,7 @@
 #include "UObject/ObjectMacros.h"
 #include "UObject/Object.h"
 #include "UObject/ScriptInterface.h"
-#include "Misc/StringAssetReference.h"
+#include "UObject/SoftObjectPath.h"
 #include "Engine/BlendableInterface.h"
 #include "PropertyEditorTestObject.generated.h"
 
@@ -86,7 +86,7 @@ struct FPropertyEditorTestSubStruct
 	FLinearColor CustomizedStructInsideUncustomizedStruct;
 
 	UPROPERTY(EditAnywhere, Category=PropertyEditorTestSubStruct)
-	FStringAssetReference CustomizedStructInsideUncustomizedStruct2;
+	FSoftObjectPath CustomizedStructInsideUncustomizedStruct2;
 };
 
 /**
@@ -290,10 +290,10 @@ class UPropertyEditorTestObject : public UObject
 	FPropertyEditorTestBasicStruct StructWithMultipleInstances2;
 
 	UPROPERTY(EditAnywhere, Category=StructTests)
-	FStringAssetReference AssetReferenceCustomStruct;
+	FSoftObjectPath AssetReferenceCustomStruct;
 
 	UPROPERTY(EditAnywhere, Category=StructTests, meta=(DisplayThumbnail = "true"))
-	FStringAssetReference AssetReferenceCustomStructWithThumbnail;
+	FSoftObjectPath AssetReferenceCustomStructWithThumbnail;
 
 	UPROPERTY(EditAnywhere, Category=StructTests, meta=(InlineEditConditionToggle))
 	bool bEditCondition;
@@ -305,7 +305,7 @@ class UPropertyEditorTestObject : public UObject
 	bool bEditConditionAssetReferenceCustomStructWithEditCondition;
 
 	UPROPERTY(EditAnywhere, Category=StructTests, meta=(editcondition = "bEditConditionAssetReferenceCustomStructWithEditCondition"))
-	FStringAssetReference AssetReferenceCustomStructWithEditCondition;
+	FSoftObjectPath AssetReferenceCustomStructWithEditCondition;
 
 	UPROPERTY(EditAnywhere, Category=StructTests)
 	TArray<FPropertyEditorTestBasicStruct> ArrayOfStructs;
diff --git a/Engine/Source/Editor/UnrealEd/Classes/Editor/Transactor.h b/Engine/Source/Editor/UnrealEd/Classes/Editor/Transactor.h
--- a/Engine/Source/Editor/UnrealEd/Classes/Editor/Transactor.h
+++ b/Engine/Source/Editor/UnrealEd/Classes/Editor/Transactor.h
@@ -11,7 +11,6 @@
 #include "UObject/Object.h"
 #include "Serialization/ArchiveUObject.h"
 #include "Misc/ITransaction.h"
-#include "UObject/AssetPtr.h"
 #include "Transactor.generated.h"
 
 
diff --git a/Engine/Source/Editor/UnrealEd/Classes/Editor/UnrealEdEngine.h b/Engine/Source/Editor/UnrealEd/Classes/Editor/UnrealEdEngine.h
--- a/Engine/Source/Editor/UnrealEd/Classes/Editor/UnrealEdEngine.h
+++ b/Engine/Source/Editor/UnrealEd/Classes/Editor/UnrealEdEngine.h
@@ -420,14 +420,15 @@ public:
 	virtual void edactSelectRelevantLights( UWorld* InWorld );
 
 	/**
-	 * Deletes all selected actors.  bIgnoreKismetReferenced is ignored when bVerifyDeletionCanHappen is true.
+	 * Deletes all selected actors
 	 *
-	 * @param	InWorld						World context
+	 * @param	InWorld				World context
 	 * @param	bVerifyDeletionCanHappen	[opt] If true (default), verify that deletion can be performed.
-	 * @param	bWarnAboutReferences		[opt] If true (default), we prompt the user about referenced actours they are about to delete
+	 * @param	bWarnAboutReferences		[opt] If true (default), we prompt the user about referenced actors they are about to delete
+	 * @param	bWarnAboutSoftReferences	[opt] If true (default), we prompt the user about soft references to actors they are about to delete
 	 * @return								true unless the delete operation was aborted.
 	 */
-	virtual bool edactDeleteSelected( UWorld* InWorld, bool bVerifyDeletionCanHappen=true, bool bWarnAboutReferences = true) override;
+	virtual bool edactDeleteSelected( UWorld* InWorld, bool bVerifyDeletionCanHappen=true, bool bWarnAboutReferences = true, bool bWarnAboutSoftReferences = true) override;
 
 	/**
 	 * Creates a new group from the current selection removing any existing groups.
diff --git a/Engine/Source/Editor/UnrealEd/Classes/Factories/FbxTextureImportData.h b/Engine/Source/Editor/UnrealEd/Classes/Factories/FbxTextureImportData.h
--- a/Engine/Source/Editor/UnrealEd/Classes/Factories/FbxTextureImportData.h
+++ b/Engine/Source/Editor/UnrealEd/Classes/Factories/FbxTextureImportData.h
@@ -4,7 +4,7 @@
 
 #include "CoreMinimal.h"
 #include "UObject/ObjectMacros.h"
-#include "Misc/StringAssetReference.h"
+#include "UObject/SoftObjectPath.h"
 #include "Factories/FbxAssetImportData.h"
 #include "MaterialImportHelpers.h"
 #include "FbxTextureImportData.generated.h"
@@ -27,7 +27,7 @@ class UNREALED_API UFbxTextureImportData : public UFbxAssetImportData
 
 	/** Base material to instance from when importing materials. */
 	UPROPERTY(EditAnywhere, config, Category = Material, meta = (ImportType = "Mesh", AllowedClasses = "MaterialInterface"))
-	FStringAssetReference BaseMaterialName;
+	FSoftObjectPath BaseMaterialName;
 
 	UPROPERTY(config, meta = (ImportType = "Mesh"))
 	FString BaseColorName;
diff --git a/Engine/Source/Editor/UnrealEd/Classes/Settings/EditorExperimentalSettings.h b/Engine/Source/Editor/UnrealEd/Classes/Settings/EditorExperimentalSettings.h
--- a/Engine/Source/Editor/UnrealEd/Classes/Settings/EditorExperimentalSettings.h
+++ b/Engine/Source/Editor/UnrealEd/Classes/Settings/EditorExperimentalSettings.h
@@ -59,10 +59,6 @@ public:
 	UPROPERTY(EditAnywhere, config, Category=Blueprints, meta=(DisplayName="Blueprint Break on Exceptions"))
 	bool bBreakOnExceptions;
 
-	/** Enable experimental blueprint performance analysis tools. */
-	UPROPERTY(EditAnywhere, config, Category=Blueprints, meta=(DisplayName="Blueprint Performance Analysis Tools"))
-	bool bBlueprintPerformanceAnalysisTools;
-
 	/** Enables "Find and Replace All" tool in the MyBlueprint window for variables */
 	UPROPERTY(EditAnywhere, config, Category = Blueprints, meta = (DisplayName = "Find and Replace All References Tool"))
 	bool bEnableFindAndReplaceReferences;
diff --git a/Engine/Source/Editor/UnrealEd/Classes/Settings/LevelEditorViewportSettings.h b/Engine/Source/Editor/UnrealEd/Classes/Settings/LevelEditorViewportSettings.h
--- a/Engine/Source/Editor/UnrealEd/Classes/Settings/LevelEditorViewportSettings.h
+++ b/Engine/Source/Editor/UnrealEd/Classes/Settings/LevelEditorViewportSettings.h
@@ -5,9 +5,9 @@
 #include "CoreMinimal.h"
 #include "UObject/ObjectMacros.h"
 #include "UObject/Object.h"
+#include "UObject/SoftObjectPath.h"
 #include "Engine/EngineTypes.h"
 #include "Engine/EngineBaseTypes.h"
-#include "Misc/StringAssetReference.h"
 #include "Viewports.h"
 #include "Editor/UnrealEdTypes.h"
 #include "LevelEditorViewportSettings.generated.h"
@@ -450,7 +450,7 @@ public:
 
 	/** A list of meshes that can be used as preview mesh in the editor view port by holding down the backslash key */
 	UPROPERTY(EditAnywhere, config, Category=Preview, meta=(AllowedClasses = "StaticMesh"))
-	TArray<FStringAssetReference> PreviewMeshes;
+	TArray<FSoftObjectPath> PreviewMeshes;
 
 	UPROPERTY(EditAnywhere, config, AdvancedDisplay, Category=LookAndFeel, meta=(ClampMin = "0.01", UIMin = "0.01", UIMax = "5"))
 	float BillboardScale;
diff --git a/Engine/Source/Editor/UnrealEd/Classes/ThumbnailRendering/SceneThumbnailInfoWithPrimitive.h b/Engine/Source/Editor/UnrealEd/Classes/ThumbnailRendering/SceneThumbnailInfoWithPrimitive.h
--- a/Engine/Source/Editor/UnrealEd/Classes/ThumbnailRendering/SceneThumbnailInfoWithPrimitive.h
+++ b/Engine/Source/Editor/UnrealEd/Classes/ThumbnailRendering/SceneThumbnailInfoWithPrimitive.h
@@ -8,7 +8,7 @@
 
 #include "CoreMinimal.h"
 #include "UObject/ObjectMacros.h"
-#include "Misc/StringAssetReference.h"
+#include "UObject/SoftObjectPtr.h"
 #include "ThumbnailRendering/SceneThumbnailInfo.h"
 #include "ThumbnailRendering/ThumbnailManager.h"
 
@@ -26,7 +26,7 @@ class USceneThumbnailInfoWithPrimitive : public USceneThumbnailInfo
 
 	/** The custom mesh used when the primitive type is TPT_None */
 	UPROPERTY(EditAnywhere, Category=Thumbnail)
-	FStringAssetReference PreviewMesh;
+	FSoftObjectPath PreviewMesh;
 
 	UPROPERTY()
 	bool bUserModifiedShape;
diff --git a/Engine/Source/Editor/UnrealEd/Classes/UserDefinedStructure/UserDefinedStructEditorData.h b/Engine/Source/Editor/UnrealEd/Classes/UserDefinedStructure/UserDefinedStructEditorData.h
--- a/Engine/Source/Editor/UnrealEd/Classes/UserDefinedStructure/UserDefinedStructEditorData.h
+++ b/Engine/Source/Editor/UnrealEd/Classes/UserDefinedStructure/UserDefinedStructEditorData.h
@@ -6,7 +6,7 @@
 #include "UObject/ObjectMacros.h"
 #include "UObject/Object.h"
 #include "Misc/Guid.h"
-#include "UObject/AssetPtr.h"
+#include "UObject/SoftObjectPtr.h"
 #include "EdGraph/EdGraphPin.h"
 #include "UObject/StructOnScope.h"
 #include "EditorUndoClient.h"
@@ -40,7 +40,7 @@ struct FStructVariableDescription
 	FString SubCategory;
 
 	UPROPERTY()
-	TAssetPtr<UObject> SubCategoryObject;
+	TSoftObjectPtr<UObject> SubCategoryObject;
 
 	UPROPERTY()
 	FEdGraphTerminalType PinValueType;
diff --git a/Engine/Source/Editor/UnrealEd/Private/Commandlets/CompileAllBlueprintsCommandlet.cpp b/Engine/Source/Editor/UnrealEd/Private/Commandlets/CompileAllBlueprintsCommandlet.cpp
--- a/Engine/Source/Editor/UnrealEd/Private/Commandlets/CompileAllBlueprintsCommandlet.cpp
+++ b/Engine/Source/Editor/UnrealEd/Private/Commandlets/CompileAllBlueprintsCommandlet.cpp
@@ -297,7 +297,6 @@ void UCompileAllBlueprintsCommandlet::CompileBlueprint(UBlueprint* Blueprint)
 		CompileOptions.bSaveIntermediateProducts = false;
 		CompileOptions.bRegenerateSkelton = false;
 		CompileOptions.bIsDuplicationInstigated = false;
-		CompileOptions.bAddInstrumentation = false;
 		CompileOptions.bReinstanceAndStubOnFailure = false;
 
 		if (bCompileSkeletonOnly)
diff --git a/Engine/Source/Editor/UnrealEd/Private/Commandlets/ContentCommandlets.cpp b/Engine/Source/Editor/UnrealEd/Private/Commandlets/ContentCommandlets.cpp
--- a/Engine/Source/Editor/UnrealEd/Private/Commandlets/ContentCommandlets.cpp
+++ b/Engine/Source/Editor/UnrealEd/Private/Commandlets/ContentCommandlets.cpp
@@ -1226,11 +1226,8 @@ void UResavePackagesCommandlet::PerformAdditionalOperations(class UWorld* World,
 
 			if (bShouldBuildLighting)
 			{
-				// This does not seem to have any use for the texture streaming build but slows down considerably the process.
-			GRedirectCollector.ResolveStringAssetReference();
-
-			FLightingBuildOptions LightingOptions;
- 			LightingOptions.QualityLevel = LightingBuildQuality;
+				FLightingBuildOptions LightingOptions;
+ 				LightingOptions.QualityLevel = LightingBuildQuality;
 
 				auto BuildFailedDelegate = [&bShouldProceedWithRebuild,&World]() {
 				UE_LOG(LogContentCommandlet, Error, TEXT("[REPORT] Failed building lighting for %s"), *World->GetOutermost()->GetName());
diff --git a/Engine/Source/Editor/UnrealEd/Private/Commandlets/DerivedDataCacheCommandlet.cpp b/Engine/Source/Editor/UnrealEd/Private/Commandlets/DerivedDataCacheCommandlet.cpp
--- a/Engine/Source/Editor/UnrealEd/Private/Commandlets/DerivedDataCacheCommandlet.cpp
+++ b/Engine/Source/Editor/UnrealEd/Private/Commandlets/DerivedDataCacheCommandlet.cpp
@@ -222,7 +222,7 @@ int32 UDerivedDataCacheCommandlet::Main( const FString& Params )
 			}
 
 			// even if the load failed this could be the first time through the loop so it might have all the startup packages to resolve
-			GRedirectCollector.ResolveStringAssetReference();
+			GRedirectCollector.ResolveAllSoftObjectPaths();
 
 			// cache all the resources for this platform
 			for (TObjectIterator<UObject> It; It; ++It)
diff --git a/Engine/Source/Editor/UnrealEd/Private/Commandlets/GenerateDistillFileSetsCommandlet.cpp b/Engine/Source/Editor/UnrealEd/Private/Commandlets/GenerateDistillFileSetsCommandlet.cpp
--- a/Engine/Source/Editor/UnrealEd/Private/Commandlets/GenerateDistillFileSetsCommandlet.cpp
+++ b/Engine/Source/Editor/UnrealEd/Private/Commandlets/GenerateDistillFileSetsCommandlet.cpp
@@ -13,6 +13,7 @@
 #include "FileHelpers.h"
 #include "RedirectCollector.h"
 #include "Editor.h"
+#include "Engine/AssetManager.h"
 
 DEFINE_LOG_CATEGORY_STATIC(LogGenerateDistillFileSetsCommandlet, Log, All);
 
@@ -102,6 +103,26 @@ int32 UGenerateDistillFileSetsCommandlet::Main( const FString& InParams )
 			MapList.AddUnique(MapToCook.FilePath);
 		}
 	}
+
+	// Add any assets from the asset manager
+	if (UAssetManager::IsValid())
+	{
+		UAssetManager& Manager = UAssetManager::Get();
+		TArray<FPrimaryAssetTypeInfo> TypeInfos;
+		Manager.GetPrimaryAssetTypeInfoList(TypeInfos);
+
+		for (const FPrimaryAssetTypeInfo& TypeInfo : TypeInfos)
+		{
+			TArray<FAssetData> AssetDataList;
+
+			Manager.GetPrimaryAssetDataList(TypeInfo.PrimaryAssetType, AssetDataList);
+
+			for (const FAssetData& AssetData : AssetDataList)
+			{
+				MapList.AddUnique(AssetData.PackageName.ToString());
+			}
+		}
+	}
 	
 	const FString TemplateFileSwitch = TEXT("Template=");
 	const FString OutputFileSwitch = TEXT("Output=");
@@ -228,7 +249,7 @@ int32 UGenerateDistillFileSetsCommandlet::Main( const FString& InParams )
 			UPackage* Package = LoadPackage( NULL, *MapPackage, LOAD_None );
 			if( Package != NULL )
 			{
-				GRedirectCollector.ResolveStringAssetReference();
+				GRedirectCollector.ResolveAllSoftObjectPaths();
 
 				AllPackageNames.Add(Package->GetName());
 
diff --git a/Engine/Source/Editor/UnrealEd/Private/CookOnTheFlyServer.cpp b/Engine/Source/Editor/UnrealEd/Private/CookOnTheFlyServer.cpp
--- a/Engine/Source/Editor/UnrealEd/Private/CookOnTheFlyServer.cpp
+++ b/Engine/Source/Editor/UnrealEd/Private/CookOnTheFlyServer.cpp
@@ -1067,7 +1067,7 @@ private:
 		return *this;
 	}
 
-	virtual FArchive& operator<< (class FAssetPtr& Value) override
+	virtual FArchive& operator<< (struct FSoftObjectPtr& Value) override
 	{
 		if ( Value.Get() )
 		{
@@ -1075,7 +1075,7 @@ private:
 		}
 		return *this;
 	}
-	virtual FArchive& operator<< (struct FStringAssetReference& Value) override
+	virtual FArchive& operator<< (struct FSoftObjectPath& Value) override
 	{
 		if ( Value.ResolveObject() )
 		{
@@ -1273,7 +1273,7 @@ bool UCookOnTheFlyServer::ContainsMap(const FName& PackageName) const
 	return false;
 }
 
-bool UCookOnTheFlyServer::ContainsRedirector(const FName& PackageName, TMap<FString, FString>& RedirectedPaths) const
+bool UCookOnTheFlyServer::ContainsRedirector(const FName& PackageName, TMap<FName, FName>& RedirectedPaths) const
 {
 	bool bFoundRedirector = false;
 	TArray<FAssetData> Assets;
@@ -1283,22 +1283,25 @@ bool UCookOnTheFlyServer::ContainsRedirector(const FName& PackageName, TMap<FStr
 	{
 		if (Asset.IsRedirector())
 		{
-			FString OriginalPath = Asset.ObjectPath.ToString();
-			FString RedirectedPath;
-			if (Asset.GetTagValue("DestinationObject", RedirectedPath))
+			FName RedirectedPath;
+			FString RedirectedPathString;
+			if (Asset.GetTagValue("DestinationObject", RedirectedPathString))
 			{
-				ConstructorHelpers::StripObjectClass(RedirectedPath);
-				FAssetData DestinationData = AssetRegistry->GetAssetByObjectPath(FName(*RedirectedPath), true);
-				TSet<FString> SeenPaths;
+				ConstructorHelpers::StripObjectClass(RedirectedPathString);
+				RedirectedPath = FName(*RedirectedPathString);
+				FAssetData DestinationData = AssetRegistry->GetAssetByObjectPath(RedirectedPath, true);
+				TSet<FName> SeenPaths;
 
 				SeenPaths.Add(RedirectedPath);
 
 				// Need to follow chain of redirectors
 				while (DestinationData.IsRedirector())
 				{
-					if (DestinationData.GetTagValue("DestinationObject", RedirectedPath))
+					if (DestinationData.GetTagValue("DestinationObject", RedirectedPathString))
 					{
-						ConstructorHelpers::StripObjectClass(RedirectedPath);
+						ConstructorHelpers::StripObjectClass(RedirectedPathString);
+						RedirectedPath = FName(*RedirectedPathString);
+
 						if (SeenPaths.Contains(RedirectedPath))
 						{
 							// Recursive, bail
@@ -1307,7 +1310,7 @@ bool UCookOnTheFlyServer::ContainsRedirector(const FName& PackageName, TMap<FStr
 						else
 						{
 							SeenPaths.Add(RedirectedPath);
-							DestinationData = AssetRegistry->GetAssetByObjectPath(FName(*RedirectedPath), true);
+							DestinationData = AssetRegistry->GetAssetByObjectPath(RedirectedPath, true);
 						}
 					}
 					else
@@ -1322,7 +1325,7 @@ bool UCookOnTheFlyServer::ContainsRedirector(const FName& PackageName, TMap<FStr
 
 				if (!bDestinationValid)
 				{
-					FName StandardPackageName = GetCachedStandardPackageFileFName(FName(*FPackageName::ObjectPathToPackageName(RedirectedPath)));
+					FName StandardPackageName = GetCachedStandardPackageFileFName(FName(*FPackageName::ObjectPathToPackageName(RedirectedPathString)));
 					if (StandardPackageName != NAME_None)
 					{
 						bDestinationValid = true;
@@ -1331,12 +1334,12 @@ bool UCookOnTheFlyServer::ContainsRedirector(const FName& PackageName, TMap<FStr
 
 				if (bDestinationValid)
 				{
-					RedirectedPaths.Add(OriginalPath, RedirectedPath);
+					RedirectedPaths.Add(Asset.ObjectPath, RedirectedPath);
 				}
 				else
 				{
-					RedirectedPaths.Add(OriginalPath, FString());
-					UE_LOG(LogCook, Log, TEXT("Found redirector in package %s pointing to deleted object %s"), *PackageName.ToString(), *RedirectedPath);
+					RedirectedPaths.Add(Asset.ObjectPath, NAME_None);
+					UE_LOG(LogCook, Log, TEXT("Found redirector in package %s pointing to deleted object %s"), *PackageName.ToString(), *RedirectedPathString);
 				}
 
 				bFoundRedirector = true;
@@ -1375,30 +1378,6 @@ FString UCookOnTheFlyServer::GetBaseDirectoryForDLC() const
 
 COREUOBJECT_API extern bool GOutputCookingWarnings;
 
-
-/**
- * Callback for handling string asset references being loaded
- */
-void UCookOnTheFlyServer::OnStringAssetReferenceLoadedPackage(const FName& PackageFName)
-{
-	if (IsCookByTheBookMode())
-	{
-		FName StandardPackageName = GetCachedStandardPackageFileFName(PackageFName);
-		if (StandardPackageName != NAME_None)
-		{
-			RequestPackage(StandardPackageName, true); // force to front of queue because we know this package is now loaded
-		}
-		else
-		{
-			if (!FPackageName::IsScriptPackage(PackageFName.ToString()) && !FPackageName::IsMemoryPackage(PackageFName.ToString()))
-			{
-				UE_LOG(LogCook, Warning, TEXT("Unable to find cached package name for package %s"), *PackageFName.ToString());
-			}
-		}
-	}
-}
-
-
 bool UCookOnTheFlyServer::RequestPackage(const FName& StandardPackageFName, const TArray<FName>& TargetPlatforms, const bool bForceFrontOfQueue)
 {
 	FFilePlatformRequest FileRequest(StandardPackageFName, TargetPlatforms);
@@ -1837,7 +1816,7 @@ uint32 UCookOnTheFlyServer::TickCookOnTheSide( const float TimeSlice, uint32 &Co
 		(!(Result & COSR_WaitingOnChildCookers)))
 	{
 		check(IsCookByTheBookMode());
-		// UE_LOG(LogCook, Display, TEXT("String asset reference resolve tried %d did %d"), bTriedToRunStringAssetReferenceResolve, bHasRunStringAssetReferenceResolve);
+
 		// if we are out of stuff and we are in cook by the book from the editor mode then we finish up
 		CookByTheBookFinished();
 	}
@@ -2086,7 +2065,6 @@ void UCookOnTheFlyServer::GetAllUnsolicitedPackages(TArray<UPackage*>& PackagesT
 		}
 		SCOPE_TIMER(UnsolicitedMarkup);
 		GetUnsolicitedPackages(PackagesToSave, ContainsFullAssetGCClasses, TargetPlatformNames);
-		
 	}
 }
 
@@ -3083,25 +3061,25 @@ void UCookOnTheFlyServer::SaveCookedPackage(UPackage* Package, uint32 SaveFlags,
 	FString Filename(GetCachedPackageFilename(Package));
 
 	// Don't resolve, just add to request list as needed
-	TSet<FName> StringAssetPackages;
+	TSet<FName> SoftObjectPackages;
 
-	GRedirectCollector.ProcessStringAssetReferencePackageList(Package->GetFName(), false, StringAssetPackages);
+	GRedirectCollector.ProcessSoftObjectPathPackageList(Package->GetFName(), false, SoftObjectPackages);
 	
-	for (FName StringAssetPackage : StringAssetPackages)
+	for (FName SoftObjectPackage : SoftObjectPackages)
 	{
-		TMap<FString, FString> RedirectedPaths;
+		TMap<FName, FName> RedirectedPaths;
 
 		// If this is a redirector, extract destination from asset registry
-		if (ContainsRedirector(StringAssetPackage, RedirectedPaths))
+		if (ContainsRedirector(SoftObjectPackage, RedirectedPaths))
 		{
-			for (TPair<FString, FString>& RedirectedPath : RedirectedPaths)
+			for (TPair<FName, FName>& RedirectedPath : RedirectedPaths)
 			{
 				GRedirectCollector.AddAssetPathRedirection(RedirectedPath.Key, RedirectedPath.Value);
 			}
 		}
 
 		// Verify package actually exists
-		FName StandardPackageName = GetCachedStandardPackageFileFName(StringAssetPackage);
+		FName StandardPackageName = GetCachedStandardPackageFileFName(SoftObjectPackage);
 
 		if (StandardPackageName != NAME_None && IsCookByTheBookMode() && !CookByTheBookOptions->bDisableUnsolicitedPackages)
 		{
@@ -3215,7 +3193,7 @@ void UCookOnTheFlyServer::SaveCookedPackage(UPackage* Package, uint32 SaveFlags,
 					if (World)
 					{
 						// Fixup legacy lightmaps before saving
-						// This should be done after loading, but FRedirectCollector::ResolveStringAssetReference in Core loads UWorlds with LoadObject so there's no opportunity to handle this fixup on load
+						// This should be done after loading, but Core loads UWorlds with LoadObject so there's no opportunity to handle this fixup on load
 						World->PersistentLevel->HandleLegacyMapBuildData();
 					}
 
@@ -3283,11 +3261,6 @@ void UCookOnTheFlyServer::Initialize( ECookMode::Type DesiredCookMode, ECookInit
 		FCoreUObjectDelegates::PackageCreatedForLoad.AddUObject(this, &UCookOnTheFlyServer::MaybeMarkPackageAsAlreadyLoaded);
 	}
 
-	if (IsCookByTheBookMode())
-	{
-		FCoreUObjectDelegates::PackageLoadedFromStringAssetReference.AddUObject(this, &UCookOnTheFlyServer::OnStringAssetReferenceLoadedPackage);
-	}
-
 	if (IsCookingInEditor())
 	{
 		FCoreUObjectDelegates::OnObjectPropertyChanged.AddUObject(this, &UCookOnTheFlyServer::OnObjectPropertyChanged);
@@ -4691,13 +4664,6 @@ void UCookOnTheFlyServer::GenerateAssetRegistry()
 		return;
 	}
 
-	if (IsCookingInEditor() == false)
-	{
-		// we want to register the temporary save directory if we are cooking outside the editor.  
-		// If we are cooking inside the editor we never use this directory so don't worry about registring it
-		FPackageName::RegisterMountPoint(TEXT("/TempAutosave/"), FPaths::ProjectSavedDir() / GEngine->PlayOnConsoleSaveDir);
-	}
-
 	double GenerateAssetRegistryTime = 0.0;
 	{
 		SCOPE_TIMER(GenerateAssetRegistryTime);
@@ -5344,7 +5310,6 @@ void UCookOnTheFlyServer::CookByTheBookFinished()
 
 	GetDerivedDataCacheRef().WaitForQuiescence(true);
 	
-	GRedirectCollector.LogTimers();
 	UCookerSettings const* CookerSettings = GetDefault<UCookerSettings>();
 
 	const UProjectPackagingSettings* const PackagingSettings = GetDefault<UProjectPackagingSettings>();
@@ -5984,27 +5949,27 @@ void UCookOnTheFlyServer::StartCookByTheBook( const FCookByTheBookStartupOptions
 	}
 
 	TArray<FName> FilesInPath;
-	TSet<FName> StartupStringAssetPackages;
+	TSet<FName> StartupSoftObjectPackages;
 
 	// Get the list of string asset references, for both empty package and all startup packages
-	GRedirectCollector.ProcessStringAssetReferencePackageList(NAME_None, false, StartupStringAssetPackages);
+	GRedirectCollector.ProcessSoftObjectPathPackageList(NAME_None, false, StartupSoftObjectPackages);
 
 	for (const FName& StartupPackage : CookByTheBookOptions->StartupPackages)
 	{
-		GRedirectCollector.ProcessStringAssetReferencePackageList(StartupPackage, false, StartupStringAssetPackages);
+		GRedirectCollector.ProcessSoftObjectPathPackageList(StartupPackage, false, StartupSoftObjectPackages);
 	}
 
 	CollectFilesToCook(FilesInPath, CookMaps, CookDirectories, CookCultures, IniMapSections, CookOptions);
 
 	// Add string asset packages after collecting files, to avoid accidentally activating the behavior to cook all maps if none are specified
-	for (FName StringAssetPackage : StartupStringAssetPackages)
+	for (FName SoftObjectPackage : StartupSoftObjectPackages)
 	{
-		TMap<FString, FString> RedirectedPaths;
+		TMap<FName, FName> RedirectedPaths;
 
 		// If this is a redirector, extract destination from asset registry
-		if (ContainsRedirector(StringAssetPackage, RedirectedPaths))
+		if (ContainsRedirector(SoftObjectPackage, RedirectedPaths))
 		{
-			for (TPair<FString, FString>& RedirectedPath : RedirectedPaths)
+			for (TPair<FName, FName>& RedirectedPath : RedirectedPaths)
 			{
 				GRedirectCollector.AddAssetPathRedirection(RedirectedPath.Key, RedirectedPath.Value);
 			}
@@ -6012,7 +5977,7 @@ void UCookOnTheFlyServer::StartCookByTheBook( const FCookByTheBookStartupOptions
 
 		if (!CookByTheBookOptions->bDisableUnsolicitedPackages)
 		{
-			AddFileToCook(FilesInPath, StringAssetPackage.ToString());
+			AddFileToCook(FilesInPath, SoftObjectPackage.ToString());
 		}
 	}
 	
diff --git a/Engine/Source/Editor/UnrealEd/Private/EdGraphNode_Comment.cpp b/Engine/Source/Editor/UnrealEd/Private/EdGraphNode_Comment.cpp
--- a/Engine/Source/Editor/UnrealEd/Private/EdGraphNode_Comment.cpp
+++ b/Engine/Source/Editor/UnrealEd/Private/EdGraphNode_Comment.cpp
@@ -16,12 +16,14 @@ UEdGraphNode_Comment::UEdGraphNode_Comment(const FObjectInitializer& ObjectIniti
 {
 	NodeWidth = 400.0f;
 	NodeHeight = 100.0f;
+	FontSize = 18;
 	CommentColor = GetDefault<UGraphEditorSettings>()->DefaultCommentNodeTitleColor;
 	bColorCommentBubble = false;
 	MoveMode = ECommentBoxMode::GroupMovement;
 
 	bCommentBubblePinned = true;
 	bCommentBubbleVisible = true;
+	bCommentBubbleVisible_InDetailsPanel = true;
 	bCanResizeNode = true;
 	bCanRenameNode = true;
 	CommentDepth = -1;
@@ -38,6 +40,13 @@ void UEdGraphNode_Comment::AddReferencedObjects(UObject* InThis, FReferenceColle
 	Super::AddReferencedObjects(InThis, Collector);
 }
 
+void UEdGraphNode_Comment::PostEditChangeProperty(FPropertyChangedEvent& PropertyChangedEvent)
+{
+	bCommentBubbleVisible = bCommentBubbleVisible_InDetailsPanel;
+	bCommentBubblePinned = bCommentBubbleVisible_InDetailsPanel;
+	Super::PostEditChangeProperty(PropertyChangedEvent);
+}
+
 void UEdGraphNode_Comment::PostPlacedNewNode()
 {
 	//@TODO: Consider making the default value we use here a preference
diff --git a/Engine/Source/Editor/UnrealEd/Private/EdGraphToken.cpp b/Engine/Source/Editor/UnrealEd/Private/EdGraphToken.cpp
--- a/Engine/Source/Editor/UnrealEd/Private/EdGraphToken.cpp
+++ b/Engine/Source/Editor/UnrealEd/Private/EdGraphToken.cpp
@@ -2,6 +2,7 @@
 
 #include "EdGraphToken.h"
 #include "Kismet2/CompilerResultsLog.h"
+#include "Kismet2/BlueprintEditorUtils.h"
 
 TSharedRef<IMessageToken> FEdGraphToken::Create(const UObject* InObject, const FCompilerResultsLog* Log, UEdGraphNode*& OutSourceNode)
 {
@@ -23,6 +24,11 @@ TSharedRef<IMessageToken> FEdGraphToken::Create(const UEdGraphPin* InPin, const
 	return MakeShareable(new FEdGraphToken(SourceNode, Log->FindSourcePin(InPin)));
 }
 
+TSharedRef<IMessageToken> FEdGraphToken::Create(const TCHAR* String, const FCompilerResultsLog* Log, UEdGraphNode*& OutSourceNode)
+{
+	return FTextToken::Create(FText::FromString(FString(String)));
+}
+
 const UEdGraphPin* FEdGraphToken::GetPin() const
 {
 	return PinBeingReferenced.Get();
@@ -51,6 +57,15 @@ FEdGraphToken::FEdGraphToken(const UObject* InObject, const UEdGraphPin* InPin)
 		{
 			CachedText = Node->GetNodeTitle(ENodeTitleType::ListView);
 		}
+		else if(const UClass* Class = Cast<UClass>(InObject))
+		{
+			// Remove the trailing C if that is the users preference:
+			CachedText = FBlueprintEditorUtils::GetFriendlyClassDisplayName(Class);
+		}
+		else if(const UField* Field = Cast<UField>(InObject))
+		{
+			CachedText = Field->GetDisplayNameText();
+		}
 		else
 		{
 			CachedText = FText::FromString(InObject->GetName());
diff --git a/Engine/Source/Editor/UnrealEd/Private/EdGraphUtilities.cpp b/Engine/Source/Editor/UnrealEd/Private/EdGraphUtilities.cpp
--- a/Engine/Source/Editor/UnrealEd/Private/EdGraphUtilities.cpp
+++ b/Engine/Source/Editor/UnrealEd/Private/EdGraphUtilities.cpp
@@ -13,7 +13,6 @@
 #include "UnrealExporter.h"
 #include "Framework/Notifications/NotificationManager.h"
 #include "Widgets/Notifications/SNotificationList.h"
-#include "K2Node_TunnelBoundary.h"
 #include "K2Node_Composite.h"
 
 /////////////////////////////////////////////////////
@@ -26,7 +25,7 @@ public:
 	TSet<UEdGraphNode*> SpawnedNodes;
 	TSet<UEdGraphNode*> SubstituteNodes;
 	const UEdGraph* DestinationGraph;
-	TArray<FName> ExtraNamesInUse;
+	TSet<FName> ExtraNamesInUse;
 	TArray<UEdGraphNode*> NodesToDestroy;
 public:
 	FGraphObjectTextFactory(const UEdGraph* InDestinationGraph)
@@ -250,7 +249,7 @@ UEdGraph* FEdGraphUtilities::CloneGraph(UEdGraph* InSource, UObject* NewOuter, F
 
 				if (bCloningForCompile)
 				{
-					DstNode->EnabledState = SrcNode->IsNodeEnabled() ? ENodeEnabledState::Enabled : ENodeEnabledState::Disabled;
+					DstNode->SetEnabledState(SrcNode->IsNodeEnabled() ? ENodeEnabledState::Enabled : ENodeEnabledState::Disabled);
 				}
 			}
 		}
@@ -260,11 +259,11 @@ UEdGraph* FEdGraphUtilities::CloneGraph(UEdGraph* InSource, UObject* NewOuter, F
 }
 
 // Clones the content from SourceGraph and merges it into MergeTarget; including merging/flattening all of the children from the SourceGraph into MergeTarget
-void FEdGraphUtilities::CloneAndMergeGraphIn(UEdGraph* MergeTarget, UEdGraph* SourceGraph, FCompilerResultsLog& MessageLog, bool bRequireSchemaMatch, bool bInIsCompiling/* = false*/, bool bCreateBoundaryNodes/* = false*/, TArray<UEdGraphNode*>* OutClonedNodes)
+void FEdGraphUtilities::CloneAndMergeGraphIn(UEdGraph* MergeTarget, UEdGraph* SourceGraph, FCompilerResultsLog& MessageLog, bool bRequireSchemaMatch, bool bInIsCompiling/* = false*/, TArray<UEdGraphNode*>* OutClonedNodes)
 {
 	// Clone the graph, then move all of it's children
 	UEdGraph* ClonedGraph = CloneGraph(SourceGraph, NULL, &MessageLog, true);
-	MergeChildrenGraphsIn(ClonedGraph, ClonedGraph, bRequireSchemaMatch, false, &MessageLog, bCreateBoundaryNodes);
+	MergeChildrenGraphsIn(ClonedGraph, ClonedGraph, bRequireSchemaMatch, false, &MessageLog);
 
 	// Duplicate the list of cloned nodes
 	if (OutClonedNodes != NULL)
@@ -281,7 +280,7 @@ void FEdGraphUtilities::CloneAndMergeGraphIn(UEdGraph* MergeTarget, UEdGraph* So
 }
 
 // Moves the contents of all of the children graphs (recursively) into the target graph.  This does not clone, it's destructive to the source
-void FEdGraphUtilities::MergeChildrenGraphsIn(UEdGraph* MergeTarget, UEdGraph* ParentGraph, bool bRequireSchemaMatch, bool bInIsCompiling/* = false*/, FCompilerResultsLog* MessageLog/* = nullptr*/, bool bWantBoundaryNodes/* = false*/)
+void FEdGraphUtilities::MergeChildrenGraphsIn(UEdGraph* MergeTarget, UEdGraph* ParentGraph, bool bRequireSchemaMatch, bool bInIsCompiling/* = false*/, FCompilerResultsLog* MessageLog/* = nullptr*/)
 {
 	// Determine if we are regenerating a blueprint on load
 	UBlueprint* Blueprint = FBlueprintEditorUtils::FindBlueprintForGraph(MergeTarget);
@@ -292,46 +291,6 @@ void FEdGraphUtilities::MergeChildrenGraphsIn(UEdGraph* MergeTarget, UEdGraph* P
 	{
 		UEdGraph* ChildGraph = ParentGraph->SubGraphs[Index];
 
-		if (ChildGraph && MessageLog)
-		{
-			UK2Node_Composite* CompositeInstance = ChildGraph->GetTypedOuter<UK2Node_Composite>();
-			if (MessageLog->bTreatCompositeGraphsAsTunnels && CompositeInstance != nullptr)
-			{
-				// Locate or register active tunnels for this child graph.
-				TArray<TWeakObjectPtr<UEdGraphNode>> ActiveTunnels;
-				UEdGraphNode* TrueCompositeInstance = MessageLog->GetSourceTunnelNode(CompositeInstance);
-				MessageLog->GetTunnelsActiveForNode(TrueCompositeInstance, ActiveTunnels);
-				if (!ActiveTunnels.Num())
-				{
-					ActiveTunnels.Add(TrueCompositeInstance);
-					MessageLog->RegisterIntermediateTunnelInstance(TrueCompositeInstance, ActiveTunnels);
-				}
-				// Register composite nodes to the source composite instance.
-				for (auto Node : ChildGraph->Nodes)
-				{
-					MessageLog->RegisterIntermediateTunnelNode(Node, TrueCompositeInstance);
-					if (FBlueprintEditorUtils::IsTunnelInstanceNode(Node))
-					{
-						if (Node->IsA<UK2Node_Composite>())
-						{
-							UEdGraphNode* SourceNode = MessageLog->GetSourceTunnelNode(Node);
-							MessageLog->RegisterIntermediateTunnelInstance(SourceNode, ActiveTunnels);
-						}
-						else
-						{
-							// Register child macro instance nodes using the duplicated instance.
-							MessageLog->RegisterIntermediateTunnelInstance(Node, ActiveTunnels);
-						}
-					}
-				}
-			}
-		}
-		if (bWantBoundaryNodes && MessageLog)
-		{
-			// Create boundary nodes around tunnels for debugging/profiling if requested.
-			UK2Node_TunnelBoundary::CreateBoundaryNodesForGraph(ChildGraph, *MessageLog);
-		}
-
 		auto NodeOwner = Cast<const UEdGraphNode>(ChildGraph ? ChildGraph->GetOuter() : nullptr);
 		const bool bNonVirtualGraph = NodeOwner ? NodeOwner->ShouldMergeChildGraphs() : true;
 
@@ -348,7 +307,7 @@ void FEdGraphUtilities::MergeChildrenGraphsIn(UEdGraph* MergeTarget, UEdGraph* P
 				ChildGraph->MoveNodesToAnotherGraph(MergeTarget, IsAsyncLoading() || bIsLoading, bInIsCompiling);
 			}
 
-			MergeChildrenGraphsIn(MergeTarget, ChildGraph, bRequireSchemaMatch, bInIsCompiling, MessageLog, bWantBoundaryNodes);
+			MergeChildrenGraphsIn(MergeTarget, ChildGraph, bRequireSchemaMatch, bInIsCompiling, MessageLog);
 		}
 	}
 }
diff --git a/Engine/Source/Editor/UnrealEd/Private/EditorActor.cpp b/Engine/Source/Editor/UnrealEd/Private/EditorActor.cpp
--- a/Engine/Source/Editor/UnrealEd/Private/EditorActor.cpp
+++ b/Engine/Source/Editor/UnrealEd/Private/EditorActor.cpp
@@ -60,6 +60,8 @@
 #include "Settings/LevelEditorMiscSettings.h"
 #include "ActorGroupingUtils.h"
 #include "HAL/PlatformApplicationMisc.h"
+#include "IAssetTools.h"
+#include "AssetToolsModule.h"
 
 #define LOCTEXT_NAMESPACE "UnrealEd.EditorActor"
 
@@ -664,7 +666,7 @@ bool UUnrealEdEngine::CanDeleteSelectedActors( const UWorld* InWorld, const bool
 	return bContainsDeletable;
 }
 
-bool UUnrealEdEngine::edactDeleteSelected( UWorld* InWorld, bool bVerifyDeletionCanHappen, bool bWarnAboutReferences)
+bool UUnrealEdEngine::edactDeleteSelected( UWorld* InWorld, bool bVerifyDeletionCanHappen, bool bWarnAboutReferences, bool bWarnAboutSoftReferences)
 {
 	if ( bVerifyDeletionCanHappen )
 	{
@@ -743,8 +745,7 @@ bool UUnrealEdEngine::edactDeleteSelected( UWorld* InWorld, bool bVerifyDeletion
 
 	bool	bRequestedDeleteAllByLevel = false;
 	bool	bRequestedDeleteAllByActor = false;
-	//  @todo remove me : this is to be removed once HLOD remove actor delegate gets moved to safer place
-	bool	bRequestedDeleteAllByLODActor = false;
+	bool	bRequestedDeleteAllBySoftReference = false;
 	EAppMsgType::Type MessageType = ActorsToDelete.Num() > 1 ? EAppMsgType::YesNoYesAllNoAll : EAppMsgType::YesNo;
 	int32		DeleteCount = 0;
 
@@ -781,20 +782,22 @@ bool UUnrealEdEngine::edactDeleteSelected( UWorld* InWorld, bool bVerifyDeletion
 
 		bool bReferencedByLevelScript = bWarnAboutReferences && (nullptr != LSB && ReferencedToActorsFromLevelScriptArray.Num() > 0);
 		bool bReferencedByActor = false;
-		bool bReferencedByLODActor = false;
-		FString LODActorName;
+		bool bReferencedBySoftReference = false;
+		TArray<UObject*> SoftReferencingObjects;
+
+		if (bWarnAboutSoftReferences)
+		{
+			FAssetToolsModule& AssetToolsModule = FModuleManager::GetModuleChecked<FAssetToolsModule>(TEXT("AssetTools"));
+			AssetToolsModule.Get().FindSoftReferencesToObject(Actor, SoftReferencingObjects);
+
+			bReferencedBySoftReference = (SoftReferencingObjects.Num() > 0);
+		}
 
 		// If there are any referencing actors, make sure that they are reference types that we care about.
 		if (ReferencingActors != nullptr)
 		{
 			for (AActor* ReferencingActor : (*ReferencingActors))
 			{
-				if (ReferencingActor->IsA(ALODActor::StaticClass()))
-				{
-					bReferencedByLODActor = true;
-					LODActorName = ReferencingActor->GetActorLabel();
-					break;
-				}
 				// If the referencing actor is a child actor that is referencing us, do not treat it
 				// as referencing for the purposes of warning about deletion
 				UChildActorComponent* ParentComponent = ReferencingActor->GetParentComponent();
@@ -807,11 +810,11 @@ bool UUnrealEdEngine::edactDeleteSelected( UWorld* InWorld, bool bVerifyDeletion
 		}
 
 		// We have references from one or more sources, prompt the user for feedback.
-		if (bReferencedByLevelScript || bReferencedByActor || bReferencedByLODActor)
+		if (bReferencedByLevelScript || bReferencedByActor || bReferencedBySoftReference)
 		{
 			if ((bReferencedByLevelScript && !bRequestedDeleteAllByLevel) ||
 				(bReferencedByActor && !bRequestedDeleteAllByActor) ||
-				(bReferencedByLODActor && !bRequestedDeleteAllByLODActor))
+				(bReferencedBySoftReference && !bRequestedDeleteAllBySoftReference))
 			{
 				FText ConfirmDelete;
 
@@ -826,12 +829,11 @@ bool UUnrealEdEngine::edactDeleteSelected( UWorld* InWorld, bool bVerifyDeletion
 						LevelScriptReferenceString += TEXT(" (Level Blueprint)");
 					}
 
-					if (i < ReferencedToActorsFromLevelScriptArray.Num() - 1)
-					{
-						LevelScriptReferenceString += TEXT("\n");
-					}
+					LevelScriptReferenceString += TEXT("\n");
 				}
 
+				LevelScriptReferenceString.TrimTrailing();
+
 				FString ActorReferenceString;
 
 				if (ReferencingActors != nullptr)
@@ -845,37 +847,43 @@ bool UUnrealEdEngine::edactDeleteSelected( UWorld* InWorld, bool bVerifyDeletion
 							ActorReferenceString += TEXT(" (Other Actor)");
 						}
 
-						if (i < ReferencingActors->Num() - 1)
-						{
-							ActorReferenceString += TEXT("\n");
-						}
+						ActorReferenceString += TEXT("\n");
 					}
 				}
 
-				// check LODActor outside of normal check
-				// you might like to know other actor is referencing it
-				if (bReferencedByLODActor)
+				if (bReferencedBySoftReference)
 				{
-					ConfirmDelete = FText::Format(LOCTEXT("ConfirmDeleteActorReferencedByHLOD",
-						"Actor {0} is referenced by LODActor {1}, do you really want to delete it?"),
-						FText::FromString(Actor->GetActorLabel()),  FText::FromString(LODActorName));
+					for (UObject* ReferencingObject : SoftReferencingObjects)
+					{
+						if (AActor* ReferencingActor = Cast<AActor>(ReferencingObject))
+						{
+							ActorReferenceString += FString::Printf(TEXT("(Soft) Actor %s in %s\n"), *ReferencingActor->GetActorLabel(), *FPackageName::GetLongPackageAssetName(ReferencingActor->GetOutermost()->GetName()));
+						}
+						else
+						{
+							ActorReferenceString += FString::Printf(TEXT("(Soft) Object %s\n"), *ReferencingObject->GetPathName());
+						}
+					}
 				}
-				else if (bReferencedByLevelScript && bReferencedByActor)
+
+				ActorReferenceString.TrimTrailing();
+
+				if (bReferencedByLevelScript && (bReferencedByActor || bReferencedBySoftReference))
 				{
 					ConfirmDelete = FText::Format(LOCTEXT("ConfirmDeleteActorReferenceByScriptAndActor",
-						"Actor {0} is referenced by the level blueprint and other Actors.\nDo you really want to delete it?\n\nReference List:\n\t{1}\n\t{2}"),
+						"Actor {0} is referenced by the level blueprint and other Actors/Objects.\nDo you really want to delete it? This will break references.\n\nReference List:\n\n{1}\n{2}"),
 						FText::FromString(Actor->GetActorLabel()), FText::FromString(LevelScriptReferenceString), FText::FromString(ActorReferenceString));
 				}
 				else if (bReferencedByLevelScript)
 				{
 					ConfirmDelete = FText::Format(LOCTEXT("ConfirmDeleteActorReferencedByScript",
-						"Actor {0} is referenced by the level blueprint.\nDo you really want to delete it?\n\nReference List:\n\t{1}"),
+						"Actor {0} is referenced by the level blueprint.\nDo you really want to delete it? This will break references.\n\nReference List:\n\n{1}"),
 						FText::FromString(Actor->GetActorLabel()), FText::FromString(LevelScriptReferenceString));
 				}
 				else
 				{
 					ConfirmDelete = FText::Format(LOCTEXT("ConfirmDeleteActorReferencedByActor",
-						"Actor {0} is referenced by another Actor(s).\nDo you really want to delete it?\n\nReference List:\n\t{1}"),
+						"Actor {0} is referenced by other Actors/Objects.\nDo you really want to delete it? This will break references.\n\nReference List:\n\n{1}"),
 						FText::FromString(Actor->GetActorLabel()), FText::FromString(ActorReferenceString));
 				}
 
@@ -884,7 +892,7 @@ bool UUnrealEdEngine::edactDeleteSelected( UWorld* InWorld, bool bVerifyDeletion
 				{
 					bRequestedDeleteAllByLevel |= bReferencedByLevelScript;
 					bRequestedDeleteAllByActor |= bReferencedByActor;
-					bRequestedDeleteAllByLODActor |= bReferencedByLODActor;
+					bRequestedDeleteAllBySoftReference |= bReferencedBySoftReference;
 				}
 				else if (Result == EAppReturnType::NoAll)
 				{
@@ -904,14 +912,10 @@ bool UUnrealEdEngine::edactDeleteSelected( UWorld* InWorld, bool bVerifyDeletion
 			{
 				for (int32 ReferencingActorIndex = 0; ReferencingActorIndex < ReferencingActors->Num(); ReferencingActorIndex++)
 				{
-					(*ReferencingActors)[ReferencingActorIndex]->Modify();
-				}
-			}
-			if (bReferencedByLODActor && ReferencingActors != nullptr)
-			{
-				for (int32 ReferencingActorIndex = 0; ReferencingActorIndex < ReferencingActors->Num(); ReferencingActorIndex++)
-				{
-					ALODActor* LODActor = Cast<ALODActor>((*ReferencingActors)[ReferencingActorIndex]);
+					AActor* ReferencingActor = (*ReferencingActors)[ReferencingActorIndex];
+					ReferencingActor->Modify();
+
+					ALODActor* LODActor = Cast<ALODActor>(ReferencingActor);
 					// it's possible other actor is referencing this
 					if (LODActor)
 					{
@@ -921,7 +925,6 @@ bool UUnrealEdEngine::edactDeleteSelected( UWorld* InWorld, bool bVerifyDeletion
 			}
 		}
 	
-	
 		bool bRebuildNavigation = false;
 
 		ABrush* Brush = Cast< ABrush >( Actor );
@@ -1014,8 +1017,6 @@ bool UUnrealEdEngine::edactDeleteSelected( UWorld* InWorld, bool bVerifyDeletion
 	return true;
 }
 
-
-
 bool UUnrealEdEngine::ShouldAbortActorDeletion() const
 {
 	bool bResult = false;
diff --git a/Engine/Source/Editor/UnrealEd/Private/EditorEngine.cpp b/Engine/Source/Editor/UnrealEd/Private/EditorEngine.cpp
--- a/Engine/Source/Editor/UnrealEd/Private/EditorEngine.cpp
+++ b/Engine/Source/Editor/UnrealEd/Private/EditorEngine.cpp
@@ -55,7 +55,6 @@
 #include "Misc/ConfigCacheIni.h"
 #include "UObject/UObjectIterator.h"
 #include "Serialization/ArchiveReplaceObjectRef.h"
-#include "Misc/RedirectCollector.h"
 #include "GameFramework/WorldSettings.h"
 #include "Engine/Light.h"
 #include "Engine/StaticMeshActor.h"
@@ -2153,7 +2152,7 @@ void UEditorEngine::Cleanse( bool ClearSelection, bool Redraw, const FText& Tran
 		// the fact that the loading routing will check that already
 		// existed, but the object was missing in cache.
 		const EObjectFlags FlagsToClear = RF_Standalone | RF_Transactional;
-		TArray<UPackage*> PackagesToUnload;
+		TSet<UPackage*> PackagesToUnload;
 		for (TObjectIterator<UObjectRedirector> RedirIt; RedirIt; ++RedirIt)
 		{
 			UPackage* RedirectorPackage = RedirIt->GetOutermost();
@@ -2169,7 +2168,8 @@ void UEditorEngine::Cleanse( bool ClearSelection, bool Redraw, const FText& Tran
 			if (!PackageObjects.ContainsByPredicate(
 					[](UObject* Object)
 					{
-						return !Object->IsA<UMetaData>() && !Object->IsA<UObjectRedirector>();
+						// Look for any standalone objects that are not a redirector or metadata, if found this is not a redirector-only package
+						return !Object->IsA<UMetaData>() && !Object->IsA<UObjectRedirector>() && Object->HasAnyFlags(RF_Standalone);
 					})
 				)
 			{
@@ -2183,11 +2183,11 @@ void UEditorEngine::Cleanse( bool ClearSelection, bool Redraw, const FText& Tran
 			}
 		}
 
-		for (auto* PackageToUnload : PackagesToUnload)
+		for (UPackage* PackageToUnload : PackagesToUnload)
 		{
 			TArray<UObject*> PackageObjects;
 			GetObjectsWithOuter(PackageToUnload, PackageObjects);
-			for (auto* Object : PackageObjects)
+			for (UObject* Object : PackageObjects)
 			{
 				Object->ClearFlags(FlagsToClear);
 				Object->RemoveFromRoot();
@@ -4042,11 +4042,10 @@ bool UEditorEngine::IsPackageValidForAutoAdding(UPackage* InPackage, const FStri
 		// and that the editor is not auto-saving.
 		if ( bPackageIsValid )
 		{
-			const bool bIsPlayOnConsolePackage = CleanFilename.StartsWith( PLAYWORLD_CONSOLE_BASE_PACKAGE_PREFIX );
 			const bool bIsPIEOrScriptPackage = InPackage->RootPackageHasAnyFlags( PKG_ContainsScript | PKG_PlayInEditor );
 			const bool bIsAutosave = GUnrealEd->GetPackageAutoSaver().IsAutoSaving();
 
-			if ( bIsPlayOnConsolePackage || bIsPIEOrScriptPackage || bIsAutosave || GIsAutomationTesting )
+			if ( bIsPIEOrScriptPackage || bIsAutosave || GIsAutomationTesting )
 			{
 				bPackageIsValid = false;
 			}
@@ -4122,15 +4121,6 @@ FSavePackageResultStruct UEditorEngine::Save( UPackage* InOuter, UObject* InBase
 {
 	FScopedSlowTask SlowTask(100, FText(), bSlowTask);
 
-	// If we we need to fixup string asset references, load any referenced by this package before it is saved
-	static bool bForceLoadStringAssetReferences = FParse::Param(FCommandLine::Get(), TEXT("FixupStringAssetReferences"));
-
-	if (bForceLoadStringAssetReferences)
-	{
-		const FString PackageName = FPackageName::FilenameToLongPackageName(Filename);
-		GRedirectCollector.ResolveStringAssetReference(FName(*PackageName));
-	}
-
 	UObject* Base = InBase;
 	if ( !Base && InOuter && InOuter->HasAnyPackageFlags(PKG_ContainsMap) )
 	{
@@ -4252,6 +4242,9 @@ FSavePackageResultStruct UEditorEngine::Save( UPackage* InOuter, UObject* InBase
 			
 			// Update components again in case it was a world without a physics scene but did have rendered components.
 			World->UpdateWorldComponents(true, true);
+
+			// Rerunning construction scripts may have made it dirty again
+			InOuter->SetDirtyFlag(false);
 		}
 	}
 
@@ -6319,7 +6312,7 @@ bool UEditorEngine::LoadPreviewMesh( int32 Index )
 	const ULevelEditorViewportSettings& ViewportSettings = *GetDefault<ULevelEditorViewportSettings>();
 	if( ViewportSettings.PreviewMeshes.IsValidIndex(Index) )
 	{
-		const FStringAssetReference& MeshName = ViewportSettings.PreviewMeshes[Index];
+		const FSoftObjectPath& MeshName = ViewportSettings.PreviewMeshes[Index];
 
 		// If we don't have a preview mesh component in the world yet, create one. 
 		if( !PreviewMeshComp )
@@ -6591,46 +6584,56 @@ void UEditorEngine::UpdateAutoLoadProject()
 	IFileManager::Get().Delete(*AutoLoadInProgressFilename, bRequireExists, bEvenIfReadOnly, bQuiet);
 }
 
-FORCEINLINE bool NetworkRemapPath_local(FWorldContext &Context, FString &Str, bool reading)
+FORCEINLINE bool NetworkRemapPath_local(FWorldContext& Context, FString& Str, bool bReading)
 {
-	if (reading)
+	if (bReading)
 	{
-		return (Str.ReplaceInline( *Context.PIERemapPrefix, *Context.PIEPrefix, ESearchCase::IgnoreCase) > 0);
+		if (FPackageName::IsShortPackageName(Str))
+		{
+			return false;
+		}
+
+		// First strip any source prefix, then add the appropriate prefix for this context
+		FSoftObjectPath Path = UWorld::RemovePIEPrefix(Str);
+			
+		Path.FixupForPIE();
+		FString Remapped = Path.ToString();
+		if (!Remapped.Equals(Str, ESearchCase::CaseSensitive))
+		{
+			Str = Remapped;
+			return true;
+		}
 	}
 	else
 	{
-		return (Str.ReplaceInline( *Context.PIEPrefix, *Context.PIERemapPrefix, ESearchCase::IgnoreCase) > 0);
+		// When sending, strip prefix
+		FString Remapped = UWorld::RemovePIEPrefix(Str);
+		if (!Remapped.Equals(Str, ESearchCase::CaseSensitive))
+		{
+			Str = Remapped;
+			return true;
+		}
 	}
+	return false;
 }
 
-bool UEditorEngine::NetworkRemapPath(UNetDriver* Driver, FString &Str, bool reading)
+bool UEditorEngine::NetworkRemapPath(UNetDriver* Driver, FString& Str, bool bReading)
 {
 	if (Driver == nullptr)
 	{
 		return false;
 	}
 
-	FWorldContext &Context = GetWorldContextFromWorldChecked(Driver->GetWorld());
-	if (Context.PIEPrefix.IsEmpty() || Context.PIERemapPrefix.IsEmpty())
-	{
-		return false;
-	}
-
-	return NetworkRemapPath_local(Context, Str, reading);
+	FWorldContext& Context = GetWorldContextFromWorldChecked(Driver->GetWorld());
+	return NetworkRemapPath_local(Context, Str, bReading);
 }
 
-bool UEditorEngine::NetworkRemapPath( UPendingNetGame *PendingNetGame, FString &Str, bool reading)
+bool UEditorEngine::NetworkRemapPath( UPendingNetGame *PendingNetGame, FString& Str, bool bReading)
 {
-	FWorldContext &Context = GetWorldContextFromPendingNetGameChecked(PendingNetGame);
-	if (Context.PIEPrefix.IsEmpty() || Context.PIERemapPrefix.IsEmpty())
-	{
-		return false;
-	}
-
-	return NetworkRemapPath_local(Context, Str, reading);
+	FWorldContext& Context = GetWorldContextFromPendingNetGameChecked(PendingNetGame);
+	return NetworkRemapPath_local(Context, Str, bReading);
 }
 
-
 void UEditorEngine::VerifyLoadMapWorldCleanup()
 {
 	// This does the same as UEngine::VerifyLoadMapWorldCleanup except it also allows Editor Worlds as a valid world.
@@ -6897,6 +6900,29 @@ void FActorLabelUtilities::SetActorLabelUnique(AActor* Actor, const FString& New
 	Actor->SetActorLabel(ModifiedActorLabel);
 }
 
+void FActorLabelUtilities::RenameExistingActor(AActor* Actor, const FString& NewActorLabel, bool bMakeUnique)
+{
+	FAssetToolsModule& AssetToolsModule = FModuleManager::GetModuleChecked<FAssetToolsModule>(TEXT("AssetTools"));
+
+	FSoftObjectPath OldPath = FSoftObjectPath(Actor);
+	if (bMakeUnique)
+	{
+		SetActorLabelUnique(Actor, NewActorLabel, nullptr);
+	}
+	else
+	{
+		Actor->SetActorLabel(NewActorLabel);
+	}
+	FSoftObjectPath NewPath = FSoftObjectPath(Actor);
+
+	if (OldPath != NewPath)
+	{
+		TArray<FAssetRenameData> RenameData;
+		RenameData.Add(FAssetRenameData(OldPath, NewPath, true));
+		AssetToolsModule.Get().RenameAssets(RenameData);
+	}
+}
+
 void UEditorEngine::HandleTravelFailure(UWorld* InWorld, ETravelFailure::Type FailureType, const FString& ErrorString)
 {
 	if (InWorld && InWorld->IsPlayInEditor())
diff --git a/Engine/Source/Editor/UnrealEd/Private/EditorLevelUtils.cpp b/Engine/Source/Editor/UnrealEd/Private/EditorLevelUtils.cpp
--- a/Engine/Source/Editor/UnrealEd/Private/EditorLevelUtils.cpp
+++ b/Engine/Source/Editor/UnrealEd/Private/EditorLevelUtils.cpp
@@ -47,6 +47,8 @@ EditorLevelUtils.cpp: Editor-specific level management routines
 #include "Components/ModelComponent.h"
 #include "Misc/RuntimeErrors.h"
 #include "HAL/PlatformApplicationMisc.h"
+#include "IAssetTools.h"
+#include "AssetToolsModule.h"
 
 DEFINE_LOG_CATEGORY(LogLevelTools);
 
@@ -112,12 +114,14 @@ int32 UEditorLevelUtils::MoveActorsToLevel(const TArray<AActor*>& ActorsToMove,
 
 		if (FinalMoveList.Num() > 0)
 		{
+			TMap<FSoftObjectPath, FSoftObjectPath> ActorPathMapping;
 			GEditor->SelectNone(false, true, false);
 
 			USelection* ActorSelection = GEditor->GetSelectedActors();
 			ActorSelection->BeginBatchSelectOperation();
 			for (AActor* Actor : FinalMoveList)
 			{
+				ActorPathMapping.Add(FSoftObjectPath(Actor), FSoftObjectPath());
 				GEditor->SelectActor(Actor, true, false);
 			}
 			ActorSelection->EndBatchSelectOperation(false);
@@ -131,7 +135,7 @@ int32 UEditorLevelUtils::MoveActorsToLevel(const TArray<AActor*>& ActorsToMove,
 				ULevel* OldCurrentLevel = OwningWorld->GetCurrentLevel();
 
 				// Copy the actors we have selected to the clipboard
-				GEditor->CopySelectedActorsToClipboard(OwningWorld, true);
+				GEditor->CopySelectedActorsToClipboard(OwningWorld, true, true);
 
 				// Set the new level and force it visible while we do the paste
 				OwningWorld->SetCurrentLevel(DestLevel);
@@ -144,6 +148,70 @@ int32 UEditorLevelUtils::MoveActorsToLevel(const TArray<AActor*>& ActorsToMove,
 				// Paste the actors into the new level
 				GEditor->edactPasteSelected(OwningWorld, false, false, false);
 
+				// Build a remapping of old to new names so we can do a fixup
+				for (FSelectionIterator It(GEditor->GetSelectedActorIterator()); It; ++It)
+				{
+					AActor* Actor = static_cast<AActor*>(*It);
+					FSoftObjectPath NewPath = FSoftObjectPath(Actor);
+
+					bool bFoundMatch = false;
+
+					// First try exact match
+					for (TPair<FSoftObjectPath, FSoftObjectPath>& Pair : ActorPathMapping)
+					{
+						if (Pair.Value.IsNull() && NewPath.GetSubPathString() == Pair.Key.GetSubPathString())
+						{
+							bFoundMatch = true;
+							Pair.Value = NewPath;
+							break;
+						}
+					}
+
+					if (!bFoundMatch)
+					{
+						// Remove numbers from end as it may have had to add some to disambiguate
+						FString PartialPath = NewPath.GetSubPathString();
+						int32 IgnoreNumber;
+						FActorLabelUtilities::SplitActorLabel(PartialPath, IgnoreNumber);
+
+						for (TPair<FSoftObjectPath, FSoftObjectPath>& Pair : ActorPathMapping)
+						{
+							if (Pair.Value.IsNull())
+							{
+								FString KeyPartialPath = Pair.Key.GetSubPathString();
+								FActorLabelUtilities::SplitActorLabel(KeyPartialPath, IgnoreNumber);
+								if (PartialPath == KeyPartialPath)
+								{
+									bFoundMatch = true;
+									Pair.Value = NewPath;
+									break;
+								}
+							}
+						}
+					}
+
+					if (!bFoundMatch)
+					{
+						UE_LOG(LogLevelTools, Error, TEXT("Cannot find remapping for moved actor ID %s, any soft references pointing to it will be broken!"), *Actor->GetPathName());
+					}
+				}
+
+				FAssetToolsModule& AssetToolsModule = FModuleManager::GetModuleChecked<FAssetToolsModule>(TEXT("AssetTools"));
+				TArray<FAssetRenameData> RenameData;
+
+				for (TPair<FSoftObjectPath, FSoftObjectPath>& Pair : ActorPathMapping)
+				{
+					if (Pair.Value.IsValid())
+					{
+						RenameData.Add(FAssetRenameData(Pair.Key, Pair.Value, true));
+					}
+				}
+					
+				if (RenameData.Num() > 0)
+				{
+					AssetToolsModule.Get().RenameAssets(RenameData);
+				}
+
 				// Restore new level visibility to previous state
 				if (!bLevelVisible)
 				{
@@ -210,17 +278,18 @@ ULevel* UEditorLevelUtils::AddLevelsToWorld(UWorld* InWorld, const TArray<FStrin
 
 	// Try to add the levels that were specified in the dialog.
 	ULevel* NewLevel = nullptr;
-	for (const auto& PackageName : PackageNames)
+	for (const FString& PackageName : PackageNames)
 	{
 		SlowTask.EnterProgressFrame();
 
-		ULevelStreaming* NewStreamingLevel = AddLevelToWorld(InWorld, *PackageName, LevelStreamingClass);
-		NewLevel = NewStreamingLevel->GetLoadedLevel();
-		if (NewLevel)
+		if (ULevelStreaming* NewStreamingLevel = AddLevelToWorld(InWorld, *PackageName, LevelStreamingClass))
 		{
-			LevelDirtyCallback.Request();
+			NewLevel = NewStreamingLevel->GetLoadedLevel();
+			if (NewLevel)
+			{
+				LevelDirtyCallback.Request();
+			}
 		}
-
 	} // for each file
 
 	  // Set the last loaded level to be the current level
diff --git a/Engine/Source/Editor/UnrealEd/Private/EditorServer.cpp b/Engine/Source/Editor/UnrealEd/Private/EditorServer.cpp
--- a/Engine/Source/Editor/UnrealEd/Private/EditorServer.cpp
+++ b/Engine/Source/Editor/UnrealEd/Private/EditorServer.cpp
@@ -2873,7 +2873,7 @@ public:
 	* @param	OutNewActors			[out] Newly created actors are appended to this list.
 	* @param	DestLevel				The level to duplicate the actors in this job to.
 	*/
-	void MoveActorsToLevel(TArray<AActor*>& OutNewActors, ULevel* DestLevel, ULevel* BufferLevel, bool bCopyOnly, FString* OutClipboardContents )
+	void MoveActorsToLevel(TArray<AActor*>& OutNewActors, ULevel* DestLevel, ULevel* BufferLevel, bool bCopyOnly, bool bIsMove, FString* OutClipboardContents )
 	{
 		UWorld* World = SrcLevel->OwningWorld;
 		ULevel* OldCurrentLevel = World->GetCurrentLevel();
@@ -2903,7 +2903,7 @@ public:
 
 		if( !bCopyOnly )
 		{
-			const bool bSuccess = GEditor->edactDeleteSelected( World, false );
+			const bool bSuccess = GEditor->edactDeleteSelected( World, false, true, !bIsMove);
 			if ( !bSuccess )
 			{
 				// The deletion was aborted.
@@ -2952,7 +2952,7 @@ public:
 				*OutClipboardContents = *ScratchData;
 			}
 
-			GEditor->edactDeleteSelected( World, false );
+			GEditor->edactDeleteSelected( World, false, false, false );
 		}
 
 		if( DestLevel )
@@ -3140,7 +3140,7 @@ bool UEditorEngine::CanCopySelectedActorsToClipboard( UWorld* InWorld, FCopySele
 	return false;
 }
 
-void UEditorEngine::CopySelectedActorsToClipboard( UWorld* InWorld, bool bShouldCut )
+void UEditorEngine::CopySelectedActorsToClipboard( UWorld* InWorld, bool bShouldCut, const bool bIsMove )
 {
 	FCopySelectedInfo CopySelected;
 	if ( !CanCopySelectedActorsToClipboard( InWorld, &CopySelected ) )
@@ -3177,7 +3177,7 @@ void UEditorEngine::CopySelectedActorsToClipboard( UWorld* InWorld, bool bShould
 			// Cut!
 			const FScopedTransaction Transaction( NSLOCTEXT("UnrealEd", "Cut", "Cut") );
 			edactCopySelected( World );
-			edactDeleteSelected( World );
+			edactDeleteSelected( World, true, true, !bIsMove );
 		}
 		else
 		{
@@ -3285,7 +3285,7 @@ void UEditorEngine::CopySelectedActorsToClipboard( UWorld* InWorld, bool bShould
 
 					FString CopiedActorsString;
 					const bool bCopyOnly = !bShouldCut;
-					Job->MoveActorsToLevel( NewActors, NULL, BufferLevel, bCopyOnly, &CopiedActorsString );
+					Job->MoveActorsToLevel( NewActors, NULL, BufferLevel, bCopyOnly, bIsMove, &CopiedActorsString );
 
 					// Append our copied actors to our final clipboard string
 					ClipboardString += CopiedActorsString;
diff --git a/Engine/Source/Editor/UnrealEd/Private/FileHelpers.cpp b/Engine/Source/Editor/UnrealEd/Private/FileHelpers.cpp
--- a/Engine/Source/Editor/UnrealEd/Private/FileHelpers.cpp
+++ b/Engine/Source/Editor/UnrealEd/Private/FileHelpers.cpp
@@ -3529,17 +3529,6 @@ void FEditorFileUtils::FindAllPackageFiles(TArray<FString>& OutPackages)
 	TArray<FString> Paths;
 	GConfig->GetArray( TEXT("Core.System"), *Key, Paths, GEngineIni );
 
-	// If doing a 'Play on XXX' from the editor, add the auto-save directory to the package search path, so streamed sub-levels can be found
-	if ( !GIsEditor && FParse::Param(FCommandLine::Get(), TEXT("PIEVIACONSOLE")) )
-	{
-		FString AutoSave;
-		GConfig->GetString( TEXT("/Script/UnrealEd.EditorEngine"), TEXT("AutoSaveDir"), AutoSave, GEngineIni );
-		if (AutoSave.Len())
-		{
-			Paths.AddUnique(AutoSave);
-		}
-	}
-
 	for (int32 PathIndex = 0; PathIndex < Paths.Num(); PathIndex++)
 	{
 		FPackageName::FindPackagesInDirectory(OutPackages, *Paths[PathIndex]);
