Another const FString workaround for linux compiler
#rnx

diff --git a/Engine/Source/Editor/PropertyEditor/Private/PropertyHandleImpl.cpp b/Engine/Source/Editor/PropertyEditor/Private/PropertyHandleImpl.cpp
--- a/Engine/Source/Editor/PropertyEditor/Private/PropertyHandleImpl.cpp
+++ b/Engine/Source/Editor/PropertyEditor/Private/PropertyHandleImpl.cpp
@@ -3760,7 +3760,7 @@ FPropertyAccess::Result FPropertyHandleObject::SetValue(const FAssetData& NewVal
 FPropertyAccess::Result FPropertyHandleObject::SetValueFromFormattedString(const FString& InValue, EPropertyValueSetFlags::Type Flags)
 {
 	UProperty* Property = GetProperty();
-	const FString EmptyString;
+	const FString EmptyString = FString();
 	const FString& AllowedClassesString = Property ? Property->GetMetaData("AllowedClasses") : EmptyString;
 
 	if (Property && !AllowedClassesString.IsEmpty())
