Prereq for details panel right click fix

diff --git a/Engine/Source/Editor/PropertyEditor/Private/SDetailSingleItemRow.h b/Engine/Source/Editor/PropertyEditor/Private/SDetailSingleItemRow.h
--- a/Engine/Source/Editor/PropertyEditor/Private/SDetailSingleItemRow.h
+++ b/Engine/Source/Editor/PropertyEditor/Private/SDetailSingleItemRow.h
@@ -97,11 +97,6 @@ private:
 	FReply OnFavoriteToggle();
 	void AllowShowFavorite();
 
-// For array drag drop:
-	FReply OnMouseButtonUp(const FGeometry& MyGeometry, const FPointerEvent& MouseEvent) override
-	{
-		return FReply::Handled().EndDragDrop();
-	};
 	void OnArrayDragEnter(const FDragDropEvent& DragDropEvent);
 	void OnArrayDragLeave(const FDragDropEvent& DragDropEvent);
 	FReply OnArrayDrop(const FDragDropEvent& DragDropEvent);
