Fix for not being able to right-click details panel items

diff --git a/Engine/Source/Editor/PropertyEditor/Private/IDetailsViewPrivate.h b/Engine/Source/Editor/PropertyEditor/Private/IDetailsViewPrivate.h
--- a/Engine/Source/Editor/PropertyEditor/Private/IDetailsViewPrivate.h
+++ b/Engine/Source/Editor/PropertyEditor/Private/IDetailsViewPrivate.h
@@ -122,4 +122,7 @@ public:
 	virtual void UpdateSinglePropertyMap(TSharedPtr<FComplexPropertyNode> InRootPropertyNode, struct FDetailLayoutData& LayoutData) = 0;
 
 	virtual const FCustomPropertyTypeLayoutMap& GetCustomPropertyTypeLayoutMap() const = 0;
+
+	/** Saves the expansion state of property nodes for the selected object set */
+	virtual void SaveExpandedItems(TSharedRef<FPropertyNode> StartNode) = 0;
 };
diff --git a/Engine/Source/Editor/PropertyEditor/Private/SDetailSingleItemRow.cpp b/Engine/Source/Editor/PropertyEditor/Private/SDetailSingleItemRow.cpp
--- a/Engine/Source/Editor/PropertyEditor/Private/SDetailSingleItemRow.cpp
+++ b/Engine/Source/Editor/PropertyEditor/Private/SDetailSingleItemRow.cpp
@@ -184,6 +184,12 @@ FReply SDetailSingleItemRow::OnArrayDrop(const FDragDropEvent& DragDropEvent)
 			TSharedPtr<IPropertyHandleArray> ParentHandle = SwappingHandle->GetParentHandle()->AsArray();
 			if (ParentHandle.IsValid())
 			{
+				// Need to swap the moving and target expansion states before saving
+				bool bOriginalSwappableExpansion = SwappablePropertyNode->HasNodeFlags(EPropertyNodeFlags::Expanded) != 0;
+				bool bOriginalSwappingExpansion = SwappingPropertyNode->HasNodeFlags(EPropertyNodeFlags::Expanded) != 0;
+				SwappablePropertyNode->SetNodeFlags(EPropertyNodeFlags::Expanded, bOriginalSwappingExpansion);
+				SwappingPropertyNode->SetNodeFlags(EPropertyNodeFlags::Expanded, bOriginalSwappableExpansion);
+				OwnerTreeNode.Pin()->GetDetailsView()->SaveExpandedItems(SwappablePropertyNode->GetParentNodeSharedPtr().ToSharedRef());
 				ParentHandle->MoveElementTo(OriginalIndex, NewIndex);
 			}
 		}
diff --git a/Engine/Source/Editor/PropertyEditor/Private/SDetailsViewBase.h b/Engine/Source/Editor/PropertyEditor/Private/SDetailsViewBase.h
--- a/Engine/Source/Editor/PropertyEditor/Private/SDetailsViewBase.h
+++ b/Engine/Source/Editor/PropertyEditor/Private/SDetailsViewBase.h
@@ -160,6 +160,7 @@ public:
 	virtual void UpdateSinglePropertyMap(TSharedPtr<FComplexPropertyNode> InRootPropertyNode, FDetailLayoutData& LayoutData) override;
 	virtual FNotifyHook* GetNotifyHook() const override { return DetailsViewArgs.NotifyHook; }
 	virtual const FCustomPropertyTypeLayoutMap& GetCustomPropertyTypeLayoutMap() const { return InstancedTypeToLayoutMap; }
+	void SaveExpandedItems( TSharedRef<FPropertyNode> StartNode ) override;
 
 	virtual bool IsConnected() const = 0;
 	virtual FRootPropertyNodeList& GetRootNodes() = 0;
@@ -202,8 +203,6 @@ public:
 	virtual void Tick( const FGeometry& AllottedGeometry, const double InCurrentTime, const float InDeltaTime ) override;
 	// End of SWidget interface
 
-	/** Saves the expansion state of property nodes for the selected object set */
-	void SaveExpandedItems( TSharedRef<FPropertyNode> StartNode );
 
 	/**
 	* Restores the expansion state of property nodes for the selected object set
