Fix for array drag and drop trying to drag and drop between different arrays

diff --git a/Engine/Source/Editor/PropertyEditor/Private/SDetailSingleItemRow.cpp b/Engine/Source/Editor/PropertyEditor/Private/SDetailSingleItemRow.cpp
--- a/Engine/Source/Editor/PropertyEditor/Private/SDetailSingleItemRow.cpp
+++ b/Engine/Source/Editor/PropertyEditor/Private/SDetailSingleItemRow.cpp
@@ -182,7 +182,7 @@ FReply SDetailSingleItemRow::OnArrayDrop(const FDragDropEvent& DragDropEvent)
 			}
 			TSharedPtr<IPropertyHandle> SwappingHandle = PropertyEditorHelpers::GetPropertyHandle(SwappingPropertyNode.ToSharedRef(), OwnerTreeNode.Pin()->GetDetailsView()->GetNotifyHook(), OwnerTreeNode.Pin()->GetDetailsView()->GetPropertyUtilities());
 			TSharedPtr<IPropertyHandleArray> ParentHandle = SwappingHandle->GetParentHandle()->AsArray();
-			if (ParentHandle.IsValid())
+			if (ParentHandle.IsValid() && SwappablePropertyNode->GetParentNode() == SwappingPropertyNode->GetParentNode())
 			{
 				// Need to swap the moving and target expansion states before saving
 				bool bOriginalSwappableExpansion = SwappablePropertyNode->HasNodeFlags(EPropertyNodeFlags::Expanded) != 0;
