Fixed crash when construction scripts are rerun during the details panel calling posteditchange

diff --git a/Engine/Source/Editor/PropertyEditor/Private/DetailItemNode.cpp b/Engine/Source/Editor/PropertyEditor/Private/DetailItemNode.cpp
--- a/Engine/Source/Editor/PropertyEditor/Private/DetailItemNode.cpp
+++ b/Engine/Source/Editor/PropertyEditor/Private/DetailItemNode.cpp
@@ -342,7 +342,7 @@ void FDetailItemNode::GenerateChildren( bool bUpdateFilteredNodes )
 	Children.Empty();
 
 	TSharedPtr<FDetailCategoryImpl> ParentCategoryPinned = ParentCategory.Pin();
-	if (ParentCategoryPinned.IsValid() == false)
+	if (ParentCategoryPinned.IsValid() == false || !ParentCategoryPinned->IsParentLayoutValid())
 	{
 		return;
 	}
