PropertyEditor - Make the property row generator utilities safer when it's owning generator is deleted.  This is neccesary becaause it's possible for it to outlive it's owner since it's returned by shared reference.  This fixes a niagara specific crash where undo causes the property row generator to be deleted, but windows mesages are still being processed and if a control created by the generator was focused the messages may be routed to that control which may call the utilities IsPropertyEditingEnabled function.

diff --git a/Engine/Source/Editor/PropertyEditor/Private/PropertyRowGenerator.cpp b/Engine/Source/Editor/PropertyEditor/Private/PropertyRowGenerator.cpp
--- a/Engine/Source/Editor/PropertyEditor/Private/PropertyRowGenerator.cpp
+++ b/Engine/Source/Editor/PropertyEditor/Private/PropertyRowGenerator.cpp
@@ -16,13 +16,18 @@ class FPropertyRowGeneratorUtilities : public IPropertyUtilities
 {
 public:
 	FPropertyRowGeneratorUtilities(FPropertyRowGenerator& InGenerator)
-		: Generator(InGenerator)
+		: Generator(&InGenerator)
 	{}
 
+	void ResetGenerator()
+	{
+		Generator = nullptr;
+	}
+
 	/** IPropertyUtilities interface */
 	virtual class FNotifyHook* GetNotifyHook() const override
 	{
-		return Generator.GetNotifyHook();
+		return Generator != nullptr ? Generator->GetNotifyHook() : nullptr;
 	}
 	virtual bool AreFavoritesEnabled() const override
 	{
@@ -33,22 +38,28 @@ public:
 	virtual void CreateColorPickerWindow(const TSharedRef< class FPropertyEditor >& PropertyEditor, bool bUseAlpha) const override {}
 	virtual void EnqueueDeferredAction(FSimpleDelegate DeferredAction) override
 	{
-		Generator.EnqueueDeferredAction(DeferredAction);
+		checkf(Generator != nullptr, TEXT("Can not enqueue action, generator is no longer valid"));
+		Generator->EnqueueDeferredAction(DeferredAction);
 	}
 	virtual bool IsPropertyEditingEnabled() const override
 	{
-		return Generator.IsPropertyEditingEnabled();
+		return Generator != nullptr && Generator->IsPropertyEditingEnabled();
 	}
 
 	virtual void ForceRefresh() override
 	{
-		Generator.ForceRefresh();
+		if (Generator != nullptr)
+		{
+			Generator->ForceRefresh();
+		}
 	}
 	virtual void RequestRefresh() override {}
 
 	virtual TSharedPtr<class FAssetThumbnailPool> GetThumbnailPool() const override
 	{
-		return Generator.GetThumbnailPool();
+		return Generator != nullptr
+			? Generator->GetThumbnailPool()
+			: TSharedPtr<class FAssetThumbnailPool>();
 	}
 
 	virtual void NotifyFinishedChangingProperties(const FPropertyChangedEvent& PropertyChangedEvent) override {}
@@ -57,15 +68,23 @@ public:
 
 	const TArray<TWeakObjectPtr<UObject>>& GetSelectedObjects() const override
 	{
-		return Generator.GetSelectedObjects();
+		if (Generator != nullptr)
+		{
+			return Generator->GetSelectedObjects();
+		}
+		else
+		{
+			static TArray<TWeakObjectPtr<UObject>> NullSelectedObjects;
+			return NullSelectedObjects;
+		}
 	}
 
 	virtual bool HasClassDefaultObject() const override
 	{
-		return Generator.HasClassDefaultObject();
+		return Generator != nullptr && Generator->HasClassDefaultObject();
 	}
 private:
-	FPropertyRowGenerator& Generator;
+	FPropertyRowGenerator* Generator;
 };
 
 
@@ -78,7 +97,7 @@ FPropertyRowGenerator::FPropertyRowGenerator(const FPropertyRowGeneratorArgs& In
 
 FPropertyRowGenerator::~FPropertyRowGenerator()
 {
-
+	StaticCastSharedRef<FPropertyRowGeneratorUtilities>(PropertyUtilities)->ResetGenerator();
 }
 
 void FPropertyRowGenerator::SetObjects(const TArray<UObject*>& InObjects)
