Fixed various reset to default issues
- Fixed reset to default not working for enum properties anywere.
- Fixed texure details not refreshing when certian properties are refreshed

diff --git a/Engine/Source/Editor/PropertyEditor/Private/PropertyHandleImpl.cpp b/Engine/Source/Editor/PropertyEditor/Private/PropertyHandleImpl.cpp
--- a/Engine/Source/Editor/PropertyEditor/Private/PropertyHandleImpl.cpp
+++ b/Engine/Source/Editor/PropertyEditor/Private/PropertyHandleImpl.cpp
@@ -646,6 +646,14 @@ void FPropertyValueImpl::SetOnChildPropertyValuePreChange(const FSimpleDelegate&
 	}
 }
 
+void FPropertyValueImpl::SetOnPropertyResetToDefault(const FSimpleDelegate& InOnPropertyResetToDefault)
+{
+	if (PropertyNode.IsValid())
+	{
+		PropertyNode.Pin()->OnPropertyResetToDefault().Add(InOnPropertyResetToDefault);
+	}
+}
+
 void FPropertyValueImpl::SetOnRebuildChildren( const FSimpleDelegate& InOnRebuildChildren )
 {
 	if( PropertyNode.IsValid() )
@@ -976,14 +984,17 @@ void FPropertyValueImpl::ResetToDefault()
 	TSharedPtr<FPropertyNode> PropertyNodePin = PropertyNode.Pin();
 	if( PropertyNodePin.IsValid() && !PropertyNodePin->IsEditConst() && PropertyNodePin->GetDiffersFromDefault() )
 	{
+		const bool bUseDisplayName = false;
 		FScopedTransaction Transaction( NSLOCTEXT("UnrealEd", "PropertyWindowResetToDefault", "Reset to Default") );
 		TSharedPtr<FPropertyNode>& KeyNode = PropertyNodePin->GetPropertyKeyNode();
 		if(KeyNode.IsValid())
 		{
 			FPropertyValueImpl(KeyNode, NotifyHook, PropertyUtilities.Pin())
-				.ImportText(KeyNode->GetDefaultValueAsString(), EPropertyValueSetFlags::DefaultFlags);
+				.ImportText(KeyNode->GetDefaultValueAsString(bUseDisplayName), EPropertyValueSetFlags::DefaultFlags);
 		}
-		ImportText(PropertyNodePin->GetDefaultValueAsString(), EPropertyValueSetFlags::DefaultFlags);
+		ImportText(PropertyNodePin->GetDefaultValueAsString(bUseDisplayName), EPropertyValueSetFlags::DefaultFlags);
+
+		PropertyNodePin->BroadcastPropertyResetToDefault();
 	}
 }
 
@@ -2449,6 +2460,11 @@ void FPropertyHandleBase::SetOnChildPropertyValuePreChange(const FSimpleDelegate
 	Implementation->SetOnChildPropertyValuePreChange(InOnChildPropertyValuePreChange);
 }
 
+void FPropertyHandleBase::SetOnPropertyResetToDefault(const FSimpleDelegate& InOnPropertyResetToDefault)
+{
+	Implementation->SetOnPropertyResetToDefault(InOnPropertyResetToDefault);
+}
+
 TSharedPtr<FPropertyNode> FPropertyHandleBase::GetPropertyNode() const
 {
 	return Implementation->GetPropertyNode();
diff --git a/Engine/Source/Editor/PropertyEditor/Private/PropertyHandleImpl.h b/Engine/Source/Editor/PropertyEditor/Private/PropertyHandleImpl.h
--- a/Engine/Source/Editor/PropertyEditor/Private/PropertyHandleImpl.h
+++ b/Engine/Source/Editor/PropertyEditor/Private/PropertyHandleImpl.h
@@ -147,15 +147,20 @@ public:
 	void SetOnChildPropertyValueChanged( const FSimpleDelegate& InOnChildPropertyValueChanged );
 
 	/**
-	* Sets a delegate to call when the property value is about to change
-	*/
+	 * Sets a delegate to call when the property value is about to change
+	 */
 	void SetOnPropertyValuePreChange(const FSimpleDelegate& InOnPropertyValuePreChange);
 
 	/**
-	* Sets a delegate to call when the propery value of a child is about to change
-	*/
+	 * Sets a delegate to call when the propery value of a child is about to change
+	 */
 	void SetOnChildPropertyValuePreChange(const FSimpleDelegate& InOnChildPropertyValuePreChange);
 
+	/**
+	 * Sets a delegate to call when a property is reset to default
+	 */
+	void SetOnPropertyResetToDefault(const FSimpleDelegate& InOnPropertyResetToDefault);
+
 	/**
 	 * Sets a delegate to call when children of the property node must be rebuilt
 	 */
@@ -431,6 +436,8 @@ public:
 	virtual void SetOnChildPropertyValueChanged( const FSimpleDelegate& InOnPropertyValueChanged ) override;
 	virtual void SetOnPropertyValuePreChange(const FSimpleDelegate& InOnPropertyValuePreChange) override;
 	virtual void SetOnChildPropertyValuePreChange(const FSimpleDelegate& InOnPropertyValuePreChange) override;
+	virtual void SetOnPropertyResetToDefault(const FSimpleDelegate& InOnPropertyResetToDefault) override;
+
 	virtual int32 GetIndexInArray() const override;
 	virtual FPropertyAccess::Result GetValueAsFormattedString( FString& OutValue, EPropertyPortFlags PortFlags = PPF_PropertyWindow ) const override;
 	virtual FPropertyAccess::Result GetValueAsDisplayString( FString& OutValue, EPropertyPortFlags PortFlags = PPF_PropertyWindow) const override;
diff --git a/Engine/Source/Editor/PropertyEditor/Private/PropertyNode.cpp b/Engine/Source/Editor/PropertyEditor/Private/PropertyNode.cpp
--- a/Engine/Source/Editor/PropertyEditor/Private/PropertyNode.cpp
+++ b/Engine/Source/Editor/PropertyEditor/Private/PropertyNode.cpp
@@ -1658,7 +1658,7 @@ bool FPropertyNode::GetDiffersFromDefault()
 	return bDiffersFromDefault;
 }
 
-FString FPropertyNode::GetDefaultValueAsStringForObject( FPropertyItemValueDataTrackerSlate& ValueTracker, UObject* InObject, UProperty* InProperty )
+FString FPropertyNode::GetDefaultValueAsStringForObject( FPropertyItemValueDataTrackerSlate& ValueTracker, UObject* InObject, UProperty* InProperty, bool bUseDisplayName)
 {
 	check( InObject );
 	check( InProperty );
@@ -1681,7 +1681,7 @@ FString FPropertyNode::GetDefaultValueAsStringForObject( FPropertyItemValueDataT
 			// The property is a simple field.  Compare it against the enclosing object's default for that property.
 			if ( !bDiffersFromDefaultForObject)
 			{
-				uint32 PortFlags = PPF_PropertyWindow;
+				uint32 PortFlags = bUseDisplayName ? PPF_PropertyWindow : PPF_None;
 				UObjectPropertyBase* ObjectProperty = Cast<UObjectPropertyBase>(InProperty);
 				if (InProperty->ContainsInstancedObjectProperty())
 				{
@@ -1731,7 +1731,7 @@ FString FPropertyNode::GetDefaultValueAsStringForObject( FPropertyItemValueDataT
 	return DefaultValue;
 }
 
-FString FPropertyNode::GetDefaultValueAsString()
+FString FPropertyNode::GetDefaultValueAsString(bool bUseDisplayName)
 {
 	FObjectPropertyNode* ObjectNode = FindObjectItemParent();
 	FString DefaultValue;
@@ -1745,7 +1745,7 @@ FString FPropertyNode::GetDefaultValueAsString()
 
 			if( Object && ValueTracker.IsValid() )
 			{
-				FString NodeDefaultValue = GetDefaultValueAsStringForObject( *ValueTracker, Object, Property.Get() );
+				FString NodeDefaultValue = GetDefaultValueAsStringForObject( *ValueTracker, Object, Property.Get(), bUseDisplayName );
 				if ( DefaultValue.Len() > 0 && NodeDefaultValue.Len() > 0)
 				{
 					DefaultValue += TEXT(", ");
@@ -2245,6 +2245,11 @@ void FPropertyNode::BroadcastPropertyPreChangeDelegates()
 
 }
 
+void FPropertyNode::BroadcastPropertyResetToDefault()
+{
+	PropertyResetToDefaultEvent.Broadcast();
+}
+
 void FPropertyNode::SetOnRebuildChildren( FSimpleDelegate InOnRebuildChildren )
 {
 	OnRebuildChildren = InOnRebuildChildren;
diff --git a/Engine/Source/Editor/PropertyEditor/Private/PropertyNode.h b/Engine/Source/Editor/PropertyEditor/Private/PropertyNode.h
--- a/Engine/Source/Editor/PropertyEditor/Private/PropertyNode.h
+++ b/Engine/Source/Editor/PropertyEditor/Private/PropertyNode.h
@@ -549,6 +549,10 @@ public:
 	/** Broadcasts when a child of this property changes */
 	FPropertyValuePreChangeEvent& OnChildPropertyValuePreChange() { return ChildPropertyValuePreChangeEvent; }
 
+	/** Broadcasts when this property is reset to default */
+	DECLARE_EVENT(FPropertyNode, FPropertyResetToDefaultEvent);
+	FPropertyResetToDefaultEvent& OnPropertyResetToDefault() { return PropertyResetToDefaultEvent; }
+
 	/**
 	 * Marks window's seem due to filtering flags
 	 * @param InFilterStrings	- List of strings that must be in the property name in order to display
@@ -801,7 +805,12 @@ public:
 	/**
 	* Gets the default value of the property as string.
 	*/
-	FString GetDefaultValueAsString();
+	FString GetDefaultValueAsString(bool bUseDisplayName = true);
+
+	/**
+	 * Broadcasts reset to default property changes
+	 */
+	void BroadcastPropertyResetToDefault();
 protected:
 
 	TSharedRef<FEditPropertyChain> BuildPropertyChain( UProperty* PropertyAboutToChange );
@@ -840,7 +849,7 @@ protected:
 
 	bool GetDiffersFromDefaultForObject( FPropertyItemValueDataTrackerSlate& ValueTracker, UProperty* InProperty );
 
-	FString GetDefaultValueAsStringForObject( FPropertyItemValueDataTrackerSlate& ValueTracker, UObject* InObject, UProperty* InProperty );
+	FString GetDefaultValueAsStringForObject( FPropertyItemValueDataTrackerSlate& ValueTracker, UObject* InObject, UProperty* InProperty, bool bUseDisplayName );
 
 	/**
 	 * Helper function to obtain the display name for an enum property
@@ -919,6 +928,9 @@ protected:
 	/** Called when a child's property value has changed */
 	FPropertyValueChangedEvent ChildPropertyValueChangedEvent;
 
+	/** Called when the property is reset to default */
+	FPropertyResetToDefaultEvent PropertyResetToDefaultEvent;
+
 	/** The property being displayed/edited. */
 	TWeakObjectPtr<UProperty> Property;
 
diff --git a/Engine/Source/Editor/PropertyEditor/Public/PropertyHandle.h b/Engine/Source/Editor/PropertyEditor/Public/PropertyHandle.h
--- a/Engine/Source/Editor/PropertyEditor/Public/PropertyHandle.h
+++ b/Engine/Source/Editor/PropertyEditor/Public/PropertyHandle.h
@@ -211,19 +211,26 @@ public:
 	virtual void SetOnChildPropertyValueChanged( const FSimpleDelegate& InOnChildPropertyValueChanged ) = 0;
 
 	/**
-	* Sets a delegate to call when the value of the property is about to be changed
-	*
-	* @param InOnPropertyValuePreChange	The delegate to call
-	*/
+	 * Sets a delegate to call when the value of the property is about to be changed
+	 *
+	 * @param InOnPropertyValuePreChange	The delegate to call
+	 */
 	virtual void SetOnPropertyValuePreChange(const FSimpleDelegate& InOnPropertyValuePreChange) = 0;
 
 	/**
-	* Sets a delegate to call when the value of the property of a child is about to be changed
-	*
-	* @param InOnChildPropertyValuePreChange	The delegate to call
-	*/
+	 * Sets a delegate to call when the value of the property of a child is about to be changed
+	 *
+	 * @param InOnChildPropertyValuePreChange	The delegate to call
+	 */
 	virtual void SetOnChildPropertyValuePreChange(const FSimpleDelegate& InOnChildPropertyValuePreChange) = 0;
 
+	/**
+	 * Sets a delegate to call when the property is reset to default
+	 *
+	 * @param InOnPropertyResetToDefault	The delegate to call
+	 */
+	virtual void SetOnPropertyResetToDefault(const FSimpleDelegate& InOnPropertyResetToDefault) = 0;
+
 	/**
 	 * Gets the typed value of a property.  
 	 * If the property does not support the value type FPropertyAccess::Fail is returned
@@ -598,18 +605,6 @@ public:
 	 * Sets an override for this property's reset to default behavior
 	 */
 	virtual void ExecuteCustomResetToDefault(const class FResetToDefaultOverride& OnCustomResetToDefault) = 0;
-
-	DEPRECATED(4.17, "IsResetToDefaultAvailable has been deprecated.  Use CanResetToDefault instead")
-	bool IsResetToDefaultAvailable() const
-	{
-		return CanResetToDefault();
-	}
-
-	DEPRECATED(4.17, "CustomResetToDefault has been deprecated.  Use ExecuteCustomResetToDefault instead")
-	void CustomResetToDefault(const class FResetToDefaultOverride& OnCustomResetToDefault)
-	{
-		ExecuteCustomResetToDefault(OnCustomResetToDefault);
-	}
 };
 
 /**
diff --git a/Engine/Source/Editor/TextureEditor/Private/Customizations/TextureDetailsCustomization.cpp b/Engine/Source/Editor/TextureEditor/Private/Customizations/TextureDetailsCustomization.cpp
--- a/Engine/Source/Editor/TextureEditor/Private/Customizations/TextureDetailsCustomization.cpp
+++ b/Engine/Source/Editor/TextureEditor/Private/Customizations/TextureDetailsCustomization.cpp
@@ -76,6 +76,8 @@ void FTextureDetails::CustomizeDetails(IDetailLayoutBuilder& DetailBuilder)
 		IDetailCategoryBuilder& TextureCategory = DetailBuilder.EditCategory("Texture");
 		IDetailPropertyRow& PowerOfTwoModePropertyRow = TextureCategory.AddProperty(GET_MEMBER_NAME_CHECKED(UTexture, PowerOfTwoMode));
 
+		PowerOfTwoModePropertyHandle->SetOnPropertyResetToDefault(FSimpleDelegate::CreateSP(this, &FTextureDetails::OnPropertyResetToDefault));
+
 		// Generate a list of enum values for the combo box
 		TArray<FText> PowerOfTwoModeComboBoxToolTips;
 		TArray<bool> RestrictedList;
@@ -231,6 +233,13 @@ void FTextureDetails::CreatePowerOfTwoModeMessage() const
 						 LOCTEXT("CannotEditPowerOfTwoMode", "Power of Two Mode cannot be changed to None for this texture as it is a non power of two size and has a Maximum Texture Size override. Change the Maximum Texture Size to 0 before attempting to change the Power of Two Mode."));
 }
 
+void FTextureDetails::OnPropertyResetToDefault() const
+{
+	uint8 CurrentPowerOfTwoMode;
+	ensure(PowerOfTwoModePropertyHandle->GetValue(CurrentPowerOfTwoMode) == FPropertyAccess::Success);
+	TextComboBox->SetSelectedItem(PowerOfTwoModeComboBoxList[CurrentPowerOfTwoMode]);
+}
+
 void FTextureDetails::OnPowerOfTwoModeChanged(TSharedPtr<FString> NewValue, ESelectInfo::Type SelectInfo)
 {
 	int32 NewPowerOfTwoMode = PowerOfTwoModeComboBoxList.Find(NewValue);
diff --git a/Engine/Source/Editor/TextureEditor/Private/Customizations/TextureDetailsCustomization.h b/Engine/Source/Editor/TextureEditor/Private/Customizations/TextureDetailsCustomization.h
--- a/Engine/Source/Editor/TextureEditor/Private/Customizations/TextureDetailsCustomization.h
+++ b/Engine/Source/Editor/TextureEditor/Private/Customizations/TextureDetailsCustomization.h
@@ -35,6 +35,7 @@ private:
 	void OnPowerOfTwoModeChanged(TSharedPtr<FString> NewValue, ESelectInfo::Type SelectInfo);
 	bool CanEditPowerOfTwoMode(int32 NewPowerOfTwoMode) const;
 	void CreatePowerOfTwoModeMessage() const;
+	void OnPropertyResetToDefault() const;
 
 	TSharedPtr<IPropertyHandle> MaxTextureSizePropertyHandle;
 	TSharedPtr<IPropertyHandle> PowerOfTwoModePropertyHandle;
