Fix warning about setting an empty value to the enum when selectnig the default value.

diff --git a/Engine/Source/Editor/PropertyEditor/Private/PropertyNode.cpp b/Engine/Source/Editor/PropertyEditor/Private/PropertyNode.cpp
--- a/Engine/Source/Editor/PropertyEditor/Private/PropertyNode.cpp
+++ b/Engine/Source/Editor/PropertyEditor/Private/PropertyNode.cpp
@@ -2717,7 +2717,7 @@ void FPropertyNode::PropagatePropertyChange( UObject* ModifiedObject, const TCHA
 					// before we do this, so that we can restore it once we've checked whether the two properties are identical
 					// This ensures that shared state keeps the correct value, even if the destination property itself isn't imported (or only partly imported, as is the case with arrays/maps/sets)
 					FString CurrentValue;
-					ComplexProperty->ExportText_Direct(CurrentValue, ModifiedComplexPropAddr, nullptr, ModifiedObject, PPF_None);
+					ComplexProperty->ExportText_Direct(CurrentValue, ModifiedComplexPropAddr, ModifiedComplexPropAddr, ModifiedObject, PPF_None);
 					ComplexProperty->ImportText(*PreviousValue, TempComplexPropAddr, PPF_None, ModifiedObject);
 					bShouldImport = ComplexProperty->Identical(DestComplexPropAddr, TempComplexPropAddr, PPF_None);
 					ComplexProperty->ImportText(*CurrentValue, TempComplexPropAddr, PPF_None, ModifiedObject);
