GitHub 4077 : Hide SDetailView Filterbox when no actor selected

diff --git a/Engine/Source/Editor/Kismet/Private/SSCSEditor.cpp b/Engine/Source/Editor/Kismet/Private/SSCSEditor.cpp
--- a/Engine/Source/Editor/Kismet/Private/SSCSEditor.cpp
+++ b/Engine/Source/Editor/Kismet/Private/SSCSEditor.cpp
@@ -3522,8 +3522,11 @@ void SSCSEditor::Construct( const FArguments& InArgs )
 
 	TSharedPtr<SHorizontalBox> ButtonBox;
 	TSharedPtr<SVerticalBox>   HeaderBox;
-	TSharedPtr<SWidget> SearchBar = SAssignNew(FilterBox, SSearchBox)
-		.OnTextChanged(this, &SSCSEditor::OnFilterTextChanged);
+	TSharedPtr<SWidget> SearchBar =
+		SAssignNew(FilterBox, SSearchBox)
+			.HintText(EditorMode == EComponentEditorMode::ActorInstance ? LOCTEXT("SearchComponentsHint", "Search Components") : LOCTEXT("SearchHint", "Search"))
+			.OnTextChanged(this, &SSCSEditor::OnFilterTextChanged);
+
 	const bool  bInlineSearchBarWithButtons = (EditorMode == EComponentEditorMode::BlueprintSCS);
 
 	bool bHideComponentClassCombo = InArgs._HideComponentClassCombo.Get();
diff --git a/Engine/Source/Editor/PropertyEditor/Private/SDetailsView.cpp b/Engine/Source/Editor/PropertyEditor/Private/SDetailsView.cpp
--- a/Engine/Source/Editor/PropertyEditor/Private/SDetailsView.cpp
+++ b/Engine/Source/Editor/PropertyEditor/Private/SDetailsView.cpp
@@ -190,6 +190,7 @@ void SDetailsView::Construct(const FArguments& InArgs)
 			[
 				// Create the search box
 				SAssignNew(SearchBox, SSearchBox)
+				.HintText(LOCTEXT("SearchDetailsHint", "Search Details"))
 				.OnTextChanged(this, &SDetailsView::OnFilterTextChanged)
 				.AddMetaData<FTagMetaData>(TEXT("Details.Search"))
 			]
@@ -892,7 +893,7 @@ void SDetailsView::SetOnObjectArrayChanged(FOnObjectArrayChanged OnObjectArrayCh
 
 bool SDetailsView::IsConnected() const
 {
-	return RootPropertyNodes.Num() > 0;
+	return GetNumObjects() > 0;
 }
 
 const FSlateBrush* SDetailsView::OnGetLockButtonImageResource() const
diff --git a/Engine/Source/Editor/PropertyEditor/Private/SDetailsViewBase.cpp b/Engine/Source/Editor/PropertyEditor/Private/SDetailsViewBase.cpp
--- a/Engine/Source/Editor/PropertyEditor/Private/SDetailsViewBase.cpp
+++ b/Engine/Source/Editor/PropertyEditor/Private/SDetailsViewBase.cpp
@@ -624,8 +624,17 @@ void SDetailsViewBase::FilterView(const FString& InFilterText)
 
 EVisibility SDetailsViewBase::GetFilterBoxVisibility() const
 {
-	// Visible if we allow search and we have anything to search otherwise collapsed so it doesn't take up room
-	return (DetailsViewArgs.bAllowSearch && IsConnected() && RootTreeNodes.Num() > 0) || HasActiveSearch() || CurrentFilter.bShowOnlyModifiedProperties || CurrentFilter.bShowAnimated || CurrentFilter.bShowKeyable || CurrentFilter.bShowOnlyDiffering ? EVisibility::Visible : EVisibility::Collapsed;
+	EVisibility Result = EVisibility::Collapsed;
+	// Visible if we allow search and we have anything to search otherwise collapsed so it doesn't take up room	
+	if (DetailsViewArgs.bAllowSearch && IsConnected())
+	{
+		if (RootTreeNodes.Num() > 0 || HasActiveSearch() || CurrentFilter.bShowOnlyModifiedProperties || CurrentFilter.bShowOnlyDiffering)
+		{
+			Result = EVisibility::Visible;
+		}
+	}
+
+	return Result;
 }
 
 bool SDetailsViewBase::SupportsKeyboardFocus() const
