Fixed insert/delete/duplicate children  calling PostEditChange on the existing child node not the array

diff --git a/Engine/Source/Editor/PropertyEditor/Private/PropertyHandleImpl.cpp b/Engine/Source/Editor/PropertyEditor/Private/PropertyHandleImpl.cpp
--- a/Engine/Source/Editor/PropertyEditor/Private/PropertyHandleImpl.cpp
+++ b/Engine/Source/Editor/PropertyEditor/Private/PropertyHandleImpl.cpp
@@ -1374,7 +1374,7 @@ void FPropertyValueImpl::InsertChild( TSharedPtr<FPropertyNode> ChildNodeToInser
 		ChangeEvent.SetArrayIndexPerObject(ArrayIndicesPerObject);
 		ChangeEvent.SetInstancesChangedResultPerArchetype(PropagationResultPerObject);
 
-		ChildNodePtr->NotifyPostChange(ChangeEvent, NotifyHook);
+		PropertyNode.Pin()->NotifyPostChange(ChangeEvent, NotifyHook);
 
 		if (PropertyUtilities.IsValid())
 		{
@@ -1520,7 +1520,7 @@ void FPropertyValueImpl::DeleteChild( TSharedPtr<FPropertyNode> ChildNodeToDelet
 		ChangeEvent.SetArrayIndexPerObject(ArrayIndicesPerObject);
 		ChangeEvent.SetInstancesChangedResultPerArchetype(PropagationResultPerObject);
 
-		ChildNodePtr->NotifyPostChange(ChangeEvent, NotifyHook);
+		PropertyNode.Pin()->NotifyPostChange(ChangeEvent, NotifyHook);
 
 		if (PropertyUtilities.IsValid())
 		{
@@ -2044,7 +2044,7 @@ void FPropertyValueImpl::DuplicateChild( TSharedPtr<FPropertyNode> ChildNodeToDu
 		ChangeEvent.SetArrayIndexPerObject(ArrayIndicesPerObject);
 		ChangeEvent.SetInstancesChangedResultPerArchetype(PropagationResultPerObject);
 
-		ChildNodePtr->NotifyPostChange(ChangeEvent, NotifyHook);
+		PropertyNode.Pin()->NotifyPostChange(ChangeEvent, NotifyHook);
 
 		if (PropertyUtilities.IsValid())
 		{
