Fix hidden asset properties in struct details panels.
We consider all object properties with "allowedclasses" metadata to be asset properties since they only show an asset picker.

diff --git a/Engine/Source/Editor/PropertyEditor/Private/PropertyEditorModule.cpp b/Engine/Source/Editor/PropertyEditor/Private/PropertyEditorModule.cpp
--- a/Engine/Source/Editor/PropertyEditor/Private/PropertyEditorModule.cpp
+++ b/Engine/Source/Editor/PropertyEditor/Private/PropertyEditorModule.cpp
@@ -756,9 +756,16 @@ TSharedRef<class IStructureDetailsView> FPropertyEditorModule::CreateStructureDe
 					// Not an "asset" property, but it may still be a property using an asset class type (such as a raw pointer)
 					if( ObjectProperty->PropertyClass )
 					{
-						// We can use the asset tools module to see whether this type has asset actions (which likely means it's an asset class type)
-						FAssetToolsModule& AssetToolsModule = FAssetToolsModule::GetModule();
-						return AssetToolsModule.Get().GetAssetTypeActionsForClass(ObjectProperty->PropertyClass).IsValid();
+						if (ObjectProperty->HasMetaData(TEXT("AllowedClasses")))
+						{
+							return true;
+						}
+						else
+						{
+							// We can use the asset tools module to see whether this type has asset actions (which likely means it's an asset class type)
+							FAssetToolsModule& AssetToolsModule = FAssetToolsModule::GetModule();
+							return AssetToolsModule.Get().GetAssetTypeActionsForClass(ObjectProperty->PropertyClass).IsValid();
+						}
 					}
 				}
 
