Work around for sync to content browser from details panels not working for interface properties

diff --git a/Engine/Source/Editor/PropertyEditor/Private/Presentation/PropertyEditor/PropertyEditor.cpp b/Engine/Source/Editor/PropertyEditor/Private/Presentation/PropertyEditor/PropertyEditor.cpp
--- a/Engine/Source/Editor/PropertyEditor/Private/Presentation/PropertyEditor/PropertyEditor.cpp
+++ b/Engine/Source/Editor/PropertyEditor/Private/Presentation/PropertyEditor/PropertyEditor.cpp
@@ -724,7 +724,8 @@ void FPropertyEditor::SyncToObjectsInNode( const TWeakPtr< FPropertyNode >& Weak
 		}
 		else if( IntProp )
 		{
-			PropertyClass = IntProp->InterfaceClass;
+			// Note: this should be IntProp->InterfaceClass but we're using UObject as the class  to work around InterfaceClass not working with FindObject
+			PropertyClass = UObject::StaticClass();
 		}
 
 		// Get a list of addresses for objects handled by the property window.
