Fix for case where the the property noded arn't rebuilt in time and custom property ui can be using stale data.

diff --git a/Engine/Source/Editor/PropertyEditor/Private/PropertyHandleImpl.cpp b/Engine/Source/Editor/PropertyEditor/Private/PropertyHandleImpl.cpp
--- a/Engine/Source/Editor/PropertyEditor/Private/PropertyHandleImpl.cpp
+++ b/Engine/Source/Editor/PropertyEditor/Private/PropertyHandleImpl.cpp
@@ -1392,6 +1392,8 @@ void FPropertyValueImpl::InsertChild( TSharedPtr<FPropertyNode> ChildNodeToInser
 			ChildNodePtr->FixPropertiesInEvent(ChangeEvent);
 			PropertyUtilities.Pin()->NotifyFinishedChangingProperties(ChangeEvent);
 		}
+
+		PropertyNode.Pin()->RebuildChildren();
 	}
 }
 
@@ -1538,6 +1540,8 @@ void FPropertyValueImpl::DeleteChild( TSharedPtr<FPropertyNode> ChildNodeToDelet
 			ChildNodePtr->FixPropertiesInEvent(ChangeEvent);
 			PropertyUtilities.Pin()->NotifyFinishedChangingProperties(ChangeEvent);
 		}
+
+		PropertyNode.Pin()->RebuildChildren();
 	}
 }
 
@@ -1638,6 +1642,8 @@ void FPropertyValueImpl::SwapChildren( TSharedPtr<FPropertyNode> FirstChildNode,
 			SecondChildNodePtr->FixPropertiesInEvent(ChangeEvent);
 			PropertyUtilities.Pin()->NotifyFinishedChangingProperties(ChangeEvent);
 		}
+
+		PropertyNode.Pin()->RebuildChildren();
 	}
 }
 
