PropertyEditor - Fix refresh for custom node builders by forcing a rebuild of the tree nodes whenever a custom node builders children are rebuilt, also create a new utilities interface for functions that need to be handled separately for the details view and the property row generator.

diff --git a/Engine/Source/Editor/PropertyEditor/Private/DetailCategoryBuilderImpl.cpp b/Engine/Source/Editor/PropertyEditor/Private/DetailCategoryBuilderImpl.cpp
--- a/Engine/Source/Editor/PropertyEditor/Private/DetailCategoryBuilderImpl.cpp
+++ b/Engine/Source/Editor/PropertyEditor/Private/DetailCategoryBuilderImpl.cpp
@@ -4,8 +4,6 @@
 #include "ObjectPropertyNode.h"
 #include "Misc/ConfigCacheIni.h"
 
-
-
 #include "DetailCategoryGroupNode.h"
 #include "DetailItemNode.h"
 #include "DetailAdvancedDropdownNode.h"
@@ -13,6 +11,7 @@
 #include "DetailGroup.h"
 #include "StructurePropertyNode.h"
 #include "ItemPropertyNode.h"
+#include "IPropertyGenerationUtilities.h"
 
 namespace DetailLayoutConstants
 {
@@ -462,15 +461,21 @@ void FDetailCategoryImpl::RequestItemExpanded(TSharedRef<FDetailTreeNode> TreeNo
 
 void FDetailCategoryImpl::RefreshTree(bool bRefilterCategory)
 {
-	TSharedPtr<FDetailLayoutBuilderImpl> DetailLayoutBuilderPtr = DetailLayoutBuilder.Pin();
-	if (DetailLayoutBuilderPtr.IsValid() && GetDetailsView())
+	if (bRefilterCategory)
 	{
-		if (bRefilterCategory)
+		TSharedPtr<FDetailLayoutBuilderImpl> DetailLayoutBuilderPtr = DetailLayoutBuilder.Pin();
+		if (DetailLayoutBuilderPtr.IsValid())
 		{
 			FilterNode(DetailLayoutBuilderPtr->GetCurrentFilter());
+			DetailLayoutBuilderPtr->GetPropertyGenerationUtilities().RebuildTreeNodes();
+		}
+	}
+	else
+	{
+		if (GetDetailsView())
+		{
+			GetDetailsView()->RefreshTree();
 		}
-
-		GetDetailsView()->RefreshTree();
 	}
 }
 
@@ -869,14 +874,9 @@ void FDetailCategoryImpl::FilterNode(const FDetailFilter& InFilter)
 
 FCustomPropertyTypeLayoutMap FDetailCategoryImpl::GetCustomPropertyTypeLayoutMap() const
 {
-	if (GetDetailsView())
-	{
-		return GetDetailsView()->GetCustomPropertyTypeLayoutMap();
-	}
-
 	if (DetailLayoutBuilder.IsValid())
 	{
-		return DetailLayoutBuilder.Pin()->GetInstancedPropertyTypeLayoutMap();
+		return DetailLayoutBuilder.Pin()->GetPropertyGenerationUtilities().GetInstancedPropertyTypeLayoutMap();
 	}
 
 	return FCustomPropertyTypeLayoutMap();
diff --git a/Engine/Source/Editor/PropertyEditor/Private/DetailLayoutBuilderImpl.cpp b/Engine/Source/Editor/PropertyEditor/Private/DetailLayoutBuilderImpl.cpp
--- a/Engine/Source/Editor/PropertyEditor/Private/DetailLayoutBuilderImpl.cpp
+++ b/Engine/Source/Editor/PropertyEditor/Private/DetailLayoutBuilderImpl.cpp
@@ -9,10 +9,11 @@
 #include "DetailMultiTopLevelObjectRootNode.h"
 #include "ObjectEditorUtils.h"
 
-FDetailLayoutBuilderImpl::FDetailLayoutBuilderImpl(TSharedPtr<FComplexPropertyNode>& InRootNode, FClassToPropertyMap& InPropertyMap, const TSharedRef< class IPropertyUtilities >& InPropertyUtilities, const TSharedPtr< IDetailsViewPrivate >& InDetailsView, bool bIsExternal)
+FDetailLayoutBuilderImpl::FDetailLayoutBuilderImpl(TSharedPtr<FComplexPropertyNode>& InRootNode, FClassToPropertyMap& InPropertyMap, const TSharedRef<IPropertyUtilities>& InPropertyUtilities, const TSharedRef<IPropertyGenerationUtilities>& InPropertyGenerationUtilities, const TSharedPtr< IDetailsViewPrivate >& InDetailsView, bool bIsExternal)
 	: RootNode( InRootNode )
 	, PropertyMap( InPropertyMap )
 	, PropertyDetailsUtilities( InPropertyUtilities )
+	, PropertyGenerationUtilities( InPropertyGenerationUtilities )
 	, DetailsView( InDetailsView.Get() )
 	, CurrentCustomizationClass( nullptr )
 {
@@ -508,14 +509,11 @@ void FDetailLayoutBuilderImpl::NotifyNodeVisibilityChanged()
 	OnNodeVisibilityChanged.Broadcast();
 }
 
-FCustomPropertyTypeLayoutMap&  FDetailLayoutBuilderImpl::GetInstancedPropertyTypeLayoutMap()
+IPropertyGenerationUtilities& FDetailLayoutBuilderImpl::GetPropertyGenerationUtilities() const
 {
-	return InstancedPropertyTypeLayoutMap;
-}
-
-void FDetailLayoutBuilderImpl::SetInstancedPropertyTypeLayoutMap(FCustomPropertyTypeLayoutMap& InInstancedPropertyTypeLayoutMap)
-{
-	InstancedPropertyTypeLayoutMap = InInstancedPropertyTypeLayoutMap;
+	TSharedPtr<IPropertyGenerationUtilities> PropertyGenerationUtilitiesPinned = PropertyGenerationUtilities.Pin();
+	checkf(PropertyGenerationUtilitiesPinned.IsValid(), TEXT("Property generation utilities were destroyed while the layout builder was still in use."));
+	return *PropertyGenerationUtilitiesPinned.Get();
 }
 
 TSharedPtr<FAssetThumbnailPool> FDetailLayoutBuilderImpl::GetThumbnailPool() const
diff --git a/Engine/Source/Editor/PropertyEditor/Private/DetailLayoutBuilderImpl.h b/Engine/Source/Editor/PropertyEditor/Private/DetailLayoutBuilderImpl.h
--- a/Engine/Source/Editor/PropertyEditor/Private/DetailLayoutBuilderImpl.h
+++ b/Engine/Source/Editor/PropertyEditor/Private/DetailLayoutBuilderImpl.h
@@ -12,6 +12,7 @@
 class FDetailCategoryImpl;
 class IDetailCategoryBuilder;
 class IPropertyUtilities;
+class IPropertyGenerationUtilities;
 
 class FDetailLayoutBuilderImpl : public IDetailLayoutBuilder, public TSharedFromThis<FDetailLayoutBuilderImpl>
 {
@@ -20,6 +21,7 @@ public:
 		TSharedPtr<FComplexPropertyNode>& InRootNode,
 		FClassToPropertyMap& InPropertyMap,
 		const TSharedRef<IPropertyUtilities>& InPropertyUtilities,
+		const TSharedRef<IPropertyGenerationUtilities>& InPropertyGenerationUtilities,
 		const TSharedPtr<IDetailsViewPrivate>& InDetailsView,
 		bool bIsExternal);
 
@@ -194,9 +196,8 @@ public:
 	/** Notifies this detail layout builder that a node it owns had it's visibility forcibly changed. */
 	void NotifyNodeVisibilityChanged();
 
-	FCustomPropertyTypeLayoutMap& GetInstancedPropertyTypeLayoutMap();
-
-	void SetInstancedPropertyTypeLayoutMap(FCustomPropertyTypeLayoutMap& InInstancedPropertyTypeLayoutMap);
+	/** Gets internal utilities for generating property layouts. */
+	IPropertyGenerationUtilities& GetPropertyGenerationUtilities() const;
 
 private:
 	/**
@@ -238,7 +239,9 @@ private:
 	/** The current variable name of the class being customized (inner class instances)*/
 	FName CurrentCustomizationVariableName;
 	/** The global property utilties.  This is weak to avoid circular ref but it should always be valid if this class exists*/
-	const TWeakPtr<class IPropertyUtilities> PropertyDetailsUtilities;
+	const TWeakPtr<IPropertyUtilities> PropertyDetailsUtilities;
+	/** Internal utilities for generating property layouts. */
+	const TWeakPtr<IPropertyGenerationUtilities> PropertyGenerationUtilities;
 	/** The view where this detail customizer resides */
 	class IDetailsViewPrivate* DetailsView;
 	/** The current class being customized */
@@ -247,7 +250,5 @@ private:
 	bool bLayoutForExternalRoot;
 	/** A delegate which is called whenever a node owned by this layout builder has it's visibility forcibly changed. */
 	FSimpleMulticastDelegate OnNodeVisibilityChanged;
-
-	FCustomPropertyTypeLayoutMap InstancedPropertyTypeLayoutMap;
 };
 
diff --git a/Engine/Source/Editor/PropertyEditor/Private/DetailsViewPropertyGenerationUtilities.h b/Engine/Source/Editor/PropertyEditor/Private/DetailsViewPropertyGenerationUtilities.h
new file mode 100644
--- /dev/null
+++ b/Engine/Source/Editor/PropertyEditor/Private/DetailsViewPropertyGenerationUtilities.h
@@ -0,0 +1,28 @@
+// Copyright 1998-2018 Epic Games, Inc. All Rights Reserved.
+
+#pragma once
+
+#include "IPropertyGenerationUtilities.h"
+
+/** Property generation utilities for widgets derived from SDetailsViewBase. */
+class FDetailsViewPropertyGenerationUtilities : public IPropertyGenerationUtilities
+{
+public:
+	FDetailsViewPropertyGenerationUtilities(SDetailsViewBase& InDetailsView)
+		: DetailsView(&InDetailsView)
+	{
+	}
+
+	virtual const FCustomPropertyTypeLayoutMap& GetInstancedPropertyTypeLayoutMap() const override
+	{
+		return DetailsView->GetCustomPropertyTypeLayoutMap();
+	}
+
+	virtual void RebuildTreeNodes() override
+	{
+		DetailsView->RerunCurrentFilter();
+	}
+
+private:
+	SDetailsViewBase * DetailsView;
+};
\ No newline at end of file
diff --git a/Engine/Source/Editor/PropertyEditor/Private/IPropertyGenerationUtilities.h b/Engine/Source/Editor/PropertyEditor/Private/IPropertyGenerationUtilities.h
new file mode 100644
--- /dev/null
+++ b/Engine/Source/Editor/PropertyEditor/Private/IPropertyGenerationUtilities.h
@@ -0,0 +1,14 @@
+// Copyright 1998-2018 Epic Games, Inc. All Rights Reserved.
+
+#pragma once
+
+/** Defines internal utilities for generating property layouts. */
+class IPropertyGenerationUtilities
+{
+public:
+	/** Gets the instance type customization map for the details implementation providing this utilities object. */
+	virtual const FCustomPropertyTypeLayoutMap& GetInstancedPropertyTypeLayoutMap() const = 0;
+
+	/** Rebuilds the details tree nodes for the details implementation providing this utilities object. */
+	virtual void RebuildTreeNodes() = 0;
+};
\ No newline at end of file
diff --git a/Engine/Source/Editor/PropertyEditor/Private/PropertyRowGenerator.cpp b/Engine/Source/Editor/PropertyEditor/Private/PropertyRowGenerator.cpp
--- a/Engine/Source/Editor/PropertyEditor/Private/PropertyRowGenerator.cpp
+++ b/Engine/Source/Editor/PropertyEditor/Private/PropertyRowGenerator.cpp
@@ -11,6 +11,7 @@
 #include "Modules/ModuleManager.h"
 #include "DetailLayoutHelpers.h"
 #include "PropertyHandleImpl.h"
+#include "IPropertyGenerationUtilities.h"
 
 class FPropertyRowGeneratorUtilities : public IPropertyUtilities
 {
@@ -88,16 +89,48 @@ private:
 };
 
 
+class FPropertyRowGeneratorGenerationUtilities : public IPropertyGenerationUtilities
+{
+public:
+	FPropertyRowGeneratorGenerationUtilities(FPropertyRowGenerator& InGenerator)
+		: Generator(&InGenerator)
+	{
+	}
+
+	void ResetGenerator()
+	{
+		Generator = nullptr;
+	}
+
+	virtual const FCustomPropertyTypeLayoutMap& GetInstancedPropertyTypeLayoutMap() const override
+	{
+		return Generator->GetInstancedPropertyTypeLayoutMap();
+	}
+
+	virtual void RebuildTreeNodes() override
+	{
+		if (Generator != nullptr)
+		{
+			Generator->UpdateDetailRows();
+		}
+	}
+
+private:
+	FPropertyRowGenerator* Generator;
+};
+
 FPropertyRowGenerator::FPropertyRowGenerator(const FPropertyRowGeneratorArgs& InArgs, TSharedPtr<FAssetThumbnailPool> InThumbnailPool)
 	: Args(InArgs)
 	, ThumbnailPool(InThumbnailPool)
 	, PropertyUtilities(new FPropertyRowGeneratorUtilities(*this))
+	, PropertyGenerationUtilities(new FPropertyRowGeneratorGenerationUtilities(*this))
 {
 }
 
 FPropertyRowGenerator::~FPropertyRowGenerator()
 {
 	StaticCastSharedRef<FPropertyRowGeneratorUtilities>(PropertyUtilities)->ResetGenerator();
+	StaticCastSharedRef<FPropertyRowGeneratorGenerationUtilities>(PropertyGenerationUtilities)->ResetGenerator();
 }
 
 void FPropertyRowGenerator::SetObjects(const TArray<UObject*>& InObjects)
@@ -362,6 +395,11 @@ void FPropertyRowGenerator::PostSetObject()
 	UpdateDetailRows();
 }
 
+const FCustomPropertyTypeLayoutMap& FPropertyRowGenerator::GetInstancedPropertyTypeLayoutMap() const
+{
+	return InstancedTypeToLayoutMap;
+}
+
 void FPropertyRowGenerator::UpdateDetailRows()
 {
 	RootTreeNodes.Reset();
@@ -477,9 +515,8 @@ void FPropertyRowGenerator::UpdateSinglePropertyMap(TSharedPtr<FComplexPropertyN
 	// Reset everything
 	LayoutData.ClassToPropertyMap.Empty();
 
-	TSharedPtr<FDetailLayoutBuilderImpl> DetailLayout = MakeShareable(new FDetailLayoutBuilderImpl(InRootPropertyNode, LayoutData.ClassToPropertyMap, PropertyUtilities, nullptr, false));
+	TSharedPtr<FDetailLayoutBuilderImpl> DetailLayout = MakeShareable(new FDetailLayoutBuilderImpl(InRootPropertyNode, LayoutData.ClassToPropertyMap, PropertyUtilities, PropertyGenerationUtilities, nullptr, false));
 	DetailLayout->AddNodeVisibilityChangedHandler(FSimpleMulticastDelegate::FDelegate::CreateRaw(this, &FPropertyRowGenerator::LayoutNodeVisibilityChanged));
-	DetailLayout->SetInstancedPropertyTypeLayoutMap(InstancedTypeToLayoutMap);
 	LayoutData.DetailLayout = DetailLayout;
 
 	TSharedPtr<FComplexPropertyNode> RootPropertyNode = InRootPropertyNode;
diff --git a/Engine/Source/Editor/PropertyEditor/Private/PropertyRowGenerator.h b/Engine/Source/Editor/PropertyEditor/Private/PropertyRowGenerator.h
--- a/Engine/Source/Editor/PropertyEditor/Private/PropertyRowGenerator.h
+++ b/Engine/Source/Editor/PropertyEditor/Private/PropertyRowGenerator.h
@@ -18,7 +18,7 @@ class IDetailCustomization;
 class IDetailNodeTree;
 class FComplexPropertyNode;
 class FDetailTreeNode;
-
+class IPropertyGenerationUtilities;
 
 struct FPropertyNodeMap
 {
@@ -107,10 +107,13 @@ public:
 	virtual TSharedPtr<class FAssetThumbnailPool> GetThumbnailPool() const;
 	virtual bool HasClassDefaultObject() const { return bViewingClassDefaultObject; }
 	virtual const TArray<TWeakObjectPtr<UObject>>& GetSelectedObjects() const { return SelectedObjects;  }
+
+	const FCustomPropertyTypeLayoutMap& GetInstancedPropertyTypeLayoutMap() const;
+	void UpdateDetailRows();
+
 private:
 	void PreSetObject(int32 NumNewObjects, bool bHasStructRoots);
 	void PostSetObject();
-	void UpdateDetailRows();
 	void UpdatePropertyMaps();
 	void UpdateSinglePropertyMap(TSharedPtr<FComplexPropertyNode> InRootPropertyNode, FDetailLayoutData& LayoutData);
 	bool ValidatePropertyNodes(const FRootPropertyNodeList& PropertyNodeList);
@@ -143,5 +146,8 @@ private:
 	TSharedPtr<FAssetThumbnailPool> ThumbnailPool;
 	/** Utility class for accessing commonly used helper methods from customizations */
 	TSharedRef<IPropertyUtilities> PropertyUtilities;
+	/** Utility class for accessing internal helper methods */
+	TSharedRef<IPropertyGenerationUtilities> PropertyGenerationUtilities;
+
 	bool bViewingClassDefaultObject;
 };
diff --git a/Engine/Source/Editor/PropertyEditor/Private/SDetailsView.cpp b/Engine/Source/Editor/PropertyEditor/Private/SDetailsView.cpp
--- a/Engine/Source/Editor/PropertyEditor/Private/SDetailsView.cpp
+++ b/Engine/Source/Editor/PropertyEditor/Private/SDetailsView.cpp
@@ -17,11 +17,11 @@
 #include "Widgets/Input/SSearchBox.h"
 #include "Classes/EditorStyleSettings.h"
 #include "DetailLayoutBuilderImpl.h"
+#include "DetailsViewPropertyGenerationUtilities.h"
 
 
 #define LOCTEXT_NAMESPACE "SDetailsView"
 
-
 SDetailsView::~SDetailsView()
 {
 	const FRootPropertyNodeList& RootNodes = GetRootNodes();
@@ -42,6 +42,7 @@ void SDetailsView::Construct(const FArguments& InArgs)
 	bViewingClassDefaultObject = false;
 
 	PropertyUtilities = MakeShareable( new FPropertyDetailsUtilities( *this ) );
+	PropertyGenerationUtilities = MakeShareable( new FDetailsViewPropertyGenerationUtilities(*this) );
 	
 	ColumnWidth = DetailsViewArgs.ColumnWidth;
 
diff --git a/Engine/Source/Editor/PropertyEditor/Private/SDetailsViewBase.cpp b/Engine/Source/Editor/PropertyEditor/Private/SDetailsViewBase.cpp
--- a/Engine/Source/Editor/PropertyEditor/Private/SDetailsViewBase.cpp
+++ b/Engine/Source/Editor/PropertyEditor/Private/SDetailsViewBase.cpp
@@ -347,7 +347,7 @@ void SDetailsViewBase::UpdateSinglePropertyMap(TSharedPtr<FComplexPropertyNode>
 	// Reset everything
 	LayoutData.ClassToPropertyMap.Empty();
 
-	TSharedPtr<FDetailLayoutBuilderImpl> DetailLayout = MakeShareable(new FDetailLayoutBuilderImpl(InRootPropertyNode, LayoutData.ClassToPropertyMap, PropertyUtilities.ToSharedRef(), SharedThis(this), bIsExternal));
+	TSharedPtr<FDetailLayoutBuilderImpl> DetailLayout = MakeShareable(new FDetailLayoutBuilderImpl(InRootPropertyNode, LayoutData.ClassToPropertyMap, PropertyUtilities.ToSharedRef(), PropertyGenerationUtilities.ToSharedRef(), SharedThis(this), bIsExternal));
 	LayoutData.DetailLayout = DetailLayout;
 
 	TSharedPtr<FComplexPropertyNode> RootPropertyNode = InRootPropertyNode;
diff --git a/Engine/Source/Editor/PropertyEditor/Private/SDetailsViewBase.h b/Engine/Source/Editor/PropertyEditor/Private/SDetailsViewBase.h
--- a/Engine/Source/Editor/PropertyEditor/Private/SDetailsViewBase.h
+++ b/Engine/Source/Editor/PropertyEditor/Private/SDetailsViewBase.h
@@ -29,6 +29,7 @@ class IDetailCustomization;
 class IDetailKeyframeHandler;
 class IDetailPropertyExtensionHandler;
 class SDetailNameArea;
+class IPropertyGenerationUtilities;
 
 
 
@@ -373,6 +374,8 @@ protected:
 	bool bHasOpenColorPicker;
 	/** Settings for this view */
 	TSharedPtr<class IPropertyUtilities> PropertyUtilities;
+	/** Gets internal utilities for generating property layouts. */
+	TSharedPtr<IPropertyGenerationUtilities> PropertyGenerationUtilities;
 	/** The name area which is not recreated when selection changes */
 	TSharedPtr<class SDetailNameArea> NameArea;
 	/** Asset pool for rendering and managing asset thumbnails visible in this view*/
diff --git a/Engine/Source/Editor/PropertyEditor/Private/SStructureDetailsView.cpp b/Engine/Source/Editor/PropertyEditor/Private/SStructureDetailsView.cpp
--- a/Engine/Source/Editor/PropertyEditor/Private/SStructureDetailsView.cpp
+++ b/Engine/Source/Editor/PropertyEditor/Private/SStructureDetailsView.cpp
@@ -10,6 +10,7 @@
 #include "StructurePropertyNode.h"
 #include "Widgets/Colors/SColorPicker.h"
 #include "Widgets/Input/SSearchBox.h"
+#include "DetailsViewPropertyGenerationUtilities.h"
 
 
 #define LOCTEXT_NAMESPACE "SStructureDetailsView"
@@ -43,6 +44,7 @@ void SStructureDetailsView::Construct(const FArguments& InArgs)
 	RootNodes.Add(MakeShareable(new FStructurePropertyNode));
 		
 	PropertyUtilities = MakeShareable( new FPropertyDetailsUtilities( *this ) );
+	PropertyGenerationUtilities = MakeShareable(new FDetailsViewPropertyGenerationUtilities(*this));
 	
 	ColumnSizeData.LeftColumnWidth = TAttribute<float>(this, &SStructureDetailsView::OnGetLeftColumnWidth);
 	ColumnSizeData.RightColumnWidth = TAttribute<float>(this, &SStructureDetailsView::OnGetRightColumnWidth);
