PropertyEditor - Add a virtual destructor for IPropertyGenerationUtilities.

diff --git a/Engine/Source/Editor/PropertyEditor/Private/IPropertyGenerationUtilities.h b/Engine/Source/Editor/PropertyEditor/Private/IPropertyGenerationUtilities.h
--- a/Engine/Source/Editor/PropertyEditor/Private/IPropertyGenerationUtilities.h
+++ b/Engine/Source/Editor/PropertyEditor/Private/IPropertyGenerationUtilities.h
@@ -6,6 +6,8 @@
 class IPropertyGenerationUtilities
 {
 public:
+	virtual ~IPropertyGenerationUtilities() { }
+
 	/** Gets the instance type customization map for the details implementation providing this utilities object. */
 	virtual const FCustomPropertyTypeLayoutMap& GetInstancedPropertyTypeLayoutMap() const = 0;
 
