Fixed crash reordering a preview mesh collection from inside an animation editor

diff --git a/Engine/Source/Editor/PropertyEditor/Private/SDetailSingleItemRow.cpp b/Engine/Source/Editor/PropertyEditor/Private/SDetailSingleItemRow.cpp
--- a/Engine/Source/Editor/PropertyEditor/Private/SDetailSingleItemRow.cpp
+++ b/Engine/Source/Editor/PropertyEditor/Private/SDetailSingleItemRow.cpp
@@ -199,7 +199,9 @@ FReply SDetailSingleItemRow::OnArrayDrop(const FDragDropEvent& DragDropEvent)
 				bool bOriginalSwappingExpansion = SwappingPropertyNode->HasNodeFlags(EPropertyNodeFlags::Expanded) != 0;
 				SwappablePropertyNode->SetNodeFlags(EPropertyNodeFlags::Expanded, bOriginalSwappingExpansion);
 				SwappingPropertyNode->SetNodeFlags(EPropertyNodeFlags::Expanded, bOriginalSwappableExpansion);
-				OwnerTreeNode.Pin()->GetDetailsView()->SaveExpandedItems(SwappablePropertyNode->GetParentNodeSharedPtr().ToSharedRef());
+
+				IDetailsViewPrivate* DetailsView = OwnerTreeNode.Pin()->GetDetailsView();
+				DetailsView->SaveExpandedItems(SwappablePropertyNode->GetParentNodeSharedPtr().ToSharedRef());
 				FScopedTransaction Transaction(NSLOCTEXT("UnrealEd", "MoveRow", "Move Row"));
 
 				SwappingHandle->GetParentHandle()->NotifyPreChange();
@@ -209,9 +211,9 @@ FReply SDetailSingleItemRow::OnArrayDrop(const FDragDropEvent& DragDropEvent)
 
 				FPropertyChangedEvent MoveEvent(SwappingHandle->GetParentHandle()->GetProperty(), EPropertyChangeType::Unspecified);
 				SwappingHandle->GetParentHandle()->NotifyPostChange();
-				if (OwnerTreeNode.Pin()->GetDetailsView()->GetPropertyUtilities().IsValid())
+				if (DetailsView->GetPropertyUtilities().IsValid())
 				{
-					OwnerTreeNode.Pin()->GetDetailsView()->GetPropertyUtilities()->NotifyFinishedChangingProperties(MoveEvent);
+					DetailsView->GetPropertyUtilities()->NotifyFinishedChangingProperties(MoveEvent);
 				}
 			}
 		}
