Never hide the filter bar if there is a filter active

diff --git a/Engine/Source/Editor/PropertyEditor/Private/SDetailsViewBase.cpp b/Engine/Source/Editor/PropertyEditor/Private/SDetailsViewBase.cpp
--- a/Engine/Source/Editor/PropertyEditor/Private/SDetailsViewBase.cpp
+++ b/Engine/Source/Editor/PropertyEditor/Private/SDetailsViewBase.cpp
@@ -649,7 +649,7 @@ EVisibility SDetailsViewBase::GetFilterBoxVisibility() const
 	// Visible if we allow search and we have anything to search otherwise collapsed so it doesn't take up room	
 	if (DetailsViewArgs.bAllowSearch && IsConnected())
 	{
-		if (RootTreeNodes.Num() > 0 || HasActiveSearch() || CurrentFilter.bShowOnlyModifiedProperties || CurrentFilter.bShowOnlyDiffering)
+		if (RootTreeNodes.Num() > 0 || HasActiveSearch() || !CurrentFilter.IsEmptyFilter())
 		{
 			Result = EVisibility::Visible;
 		}
