Fix editor crash related input binding panel

diff --git a/Engine/Source/Editor/InputBindingEditor/Private/InputBindingEditorModule.cpp b/Engine/Source/Editor/InputBindingEditor/Private/InputBindingEditorModule.cpp
--- a/Engine/Source/Editor/InputBindingEditor/Private/InputBindingEditorModule.cpp
+++ b/Engine/Source/Editor/InputBindingEditor/Private/InputBindingEditorModule.cpp
@@ -30,6 +30,15 @@ public:
 		return MakeShareable( new FEditorKeyboardShortcutSettings );
 	}
 
+	virtual void PendingDelete()
+	{
+		if (EditorPanel.IsValid())
+		{
+			ensure(EditorPanel.IsUnique());
+			EditorPanel.Reset();
+		}
+	}
+
 	virtual void CustomizeDetails( IDetailLayoutBuilder& DetailBuilder )
 	{
 		EditorPanel = MakeShareable( new FInputBindingEditorPanel );
diff --git a/Engine/Source/Editor/PropertyEditor/Private/SDetailsViewBase.cpp b/Engine/Source/Editor/PropertyEditor/Private/SDetailsViewBase.cpp
--- a/Engine/Source/Editor/PropertyEditor/Private/SDetailsViewBase.cpp
+++ b/Engine/Source/Editor/PropertyEditor/Private/SDetailsViewBase.cpp
@@ -327,8 +327,20 @@ void SDetailsViewBase::UpdatePropertyMaps()
 		// We need to be able to create a new detail layout and properly clean up the old one in the process
 		check(!LayoutData.DetailLayout.IsValid() || LayoutData.DetailLayout.IsUnique());
 
+		// Allow customizations to perform cleanup as the delete occurs later on
+		for (TSharedPtr<IDetailCustomization>& DetailCustomization : LayoutData.CustomizationClassInstances)
+		{
+			if (DetailCustomization.IsValid())
+			{
+				DetailCustomization->PendingDelete();
+			}
+		}
+
 		// All the current customization instances need to be deleted when it is safe
 		CustomizationClassInstancesPendingDelete.Append(LayoutData.CustomizationClassInstances);
+
+		// All the current detail layouts need to be deleted when it is safe
+		DetailLayoutsPendingDelete.Add(LayoutData.DetailLayout);
 	}
 
 	FRootPropertyNodeList& RootPropertyNodes = GetRootNodes();
@@ -697,6 +709,9 @@ void SDetailsViewBase::Tick( const FGeometry& AllottedGeometry, const double InC
 	// Empty all the customization instances that need to be deleted
 	CustomizationClassInstancesPendingDelete.Empty();
 
+	// Empty all the detail layouts that need to be deleted
+	DetailLayoutsPendingDelete.Empty();
+
 	FRootPropertyNodeList& RootPropertyNodes = GetRootNodes();
 
 	for(TSharedPtr<FComplexPropertyNode>& RootPropertyNode : RootPropertyNodes)
diff --git a/Engine/Source/Editor/PropertyEditor/Private/SDetailsViewBase.h b/Engine/Source/Editor/PropertyEditor/Private/SDetailsViewBase.h
--- a/Engine/Source/Editor/PropertyEditor/Private/SDetailsViewBase.h
+++ b/Engine/Source/Editor/PropertyEditor/Private/SDetailsViewBase.h
@@ -363,6 +363,8 @@ protected:
 	TSharedPtr<SSearchBox> SearchBox;
 	/** Customization instances that need to be destroyed when safe to do so */
 	TArray< TSharedPtr<IDetailCustomization> > CustomizationClassInstancesPendingDelete;
+	/** Detail layouts that need to be destroyed when safe to do so */
+	TArray< TSharedPtr<FDetailLayoutBuilderImpl> > DetailLayoutsPendingDelete;
 	/** Map of nodes that are requesting an automatic expansion/collapse due to being filtered */
 	TMap< TSharedRef<FDetailTreeNode>, bool > FilteredNodesRequestingExpansionState;
 	/** Current set of expanded detail nodes (by path) that should be saved when the details panel closes */
diff --git a/Engine/Source/Editor/PropertyEditor/Public/IDetailCustomization.h b/Engine/Source/Editor/PropertyEditor/Public/IDetailCustomization.h
--- a/Engine/Source/Editor/PropertyEditor/Public/IDetailCustomization.h
+++ b/Engine/Source/Editor/PropertyEditor/Public/IDetailCustomization.h
@@ -13,6 +13,9 @@ class IDetailCustomization : public TSharedFromThis<IDetailCustomization>
 public:
 	virtual ~IDetailCustomization() {}
 
+	/** Called when no longer used and will be deleted */
+	virtual void PendingDelete() {}
+
 	/** Called when details should be customized */
 	virtual void CustomizeDetails( IDetailLayoutBuilder& DetailBuilder ) = 0;
 
