Fix for crash while holding a window tab during auto save

diff --git a/Engine/Source/Editor/PropertyEditor/Private/UserInterface/PropertyEditor/SPropertyEditorAsset.cpp b/Engine/Source/Editor/PropertyEditor/Private/UserInterface/PropertyEditor/SPropertyEditorAsset.cpp
--- a/Engine/Source/Editor/PropertyEditor/Private/UserInterface/PropertyEditor/SPropertyEditorAsset.cpp
+++ b/Engine/Source/Editor/PropertyEditor/Private/UserInterface/PropertyEditor/SPropertyEditorAsset.cpp
@@ -603,7 +603,7 @@ SPropertyEditorAsset::EActorReferenceState SPropertyEditorAsset::GetActorReferen
 
 void SPropertyEditorAsset::Tick( const FGeometry& AllottedGeometry, const double InCurrentTime, const float InDeltaTime )
 {
-	if( AssetThumbnail.IsValid() )
+	if( AssetThumbnail.IsValid() && !GIsSavingPackage && !IsGarbageCollecting() )
 	{
 		// Ensure the thumbnail is up to date
 		FObjectOrAssetData Value;
