[NO NON-PII COMMIT MESSAGE]

diff --git a/Engine/Source/Editor/PropertyEditor/Private/ObjectPropertyNode.cpp b/Engine/Source/Editor/PropertyEditor/Private/ObjectPropertyNode.cpp
--- a/Engine/Source/Editor/PropertyEditor/Private/ObjectPropertyNode.cpp
+++ b/Engine/Source/Editor/PropertyEditor/Private/ObjectPropertyNode.cpp
@@ -136,6 +136,13 @@ bool FObjectPropertyNode::GetReadAddressUncached(FPropertyNode& InNode,
 		return false;
 	}
 
+	UStruct* OwnerStruct = InItemProperty->GetOwnerStruct();
+	if (!OwnerStruct || OwnerStruct->Children == nullptr)
+	{
+		// Verify that the property is not part of an invalid trash class
+		return false;
+	}
+
 	// Requesting a single selection?
 	if( InRequiresSingleSelection && GetNumObjects() > 1)
 	{
diff --git a/Engine/Source/Editor/PropertyEditor/Private/PropertyNode.cpp b/Engine/Source/Editor/PropertyEditor/Private/PropertyNode.cpp
--- a/Engine/Source/Editor/PropertyEditor/Private/PropertyNode.cpp
+++ b/Engine/Source/Editor/PropertyEditor/Private/PropertyNode.cpp
@@ -418,6 +418,13 @@ EPropertyDataValidationResult FPropertyNode::EnsureDataIsValid()
 		if (Property.IsValid())
 		{
 			UProperty* MyProperty = Property.Get();
+			UStruct* OwnerStruct = MyProperty->GetOwnerStruct();
+
+			if (!OwnerStruct || OwnerStruct->Children == nullptr)
+			{
+				//verify that the property is not part of an invalid trash class, treat it as an invalid object if it is which will cause a refresh
+				return EPropertyDataValidationResult::ObjectInvalid;
+			}
 
 			//verify that the number of container children is correct
 			UArrayProperty* ArrayProperty = Cast<UArrayProperty>(MyProperty);
diff --git a/Engine/Source/Editor/PropertyEditor/Private/StructurePropertyNode.h b/Engine/Source/Editor/PropertyEditor/Private/StructurePropertyNode.h
--- a/Engine/Source/Editor/PropertyEditor/Private/StructurePropertyNode.h
+++ b/Engine/Source/Editor/PropertyEditor/Private/StructurePropertyNode.h
@@ -56,6 +56,13 @@ public:
 			return false;
 		}
 
+		UStruct* OwnerStruct = InItemProperty->GetOwnerStruct();
+		if (!OwnerStruct || OwnerStruct->Children == nullptr)
+		{
+			// Verify that the property is not part of an invalid trash class
+			return false;
+		}
+
 		uint8* ReadAddress = StructData->GetStructMemory();
 		check(ReadAddress);
 		OutAddresses.Add(NULL, InPropertyNode.GetValueBaseAddress(ReadAddress), true);
