Remove asset load attempt when clicking "browse to asset in content browser"

diff --git a/Engine/Source/Editor/PropertyEditor/Private/UserInterface/PropertyEditor/SPropertyEditorAsset.cpp b/Engine/Source/Editor/PropertyEditor/Private/UserInterface/PropertyEditor/SPropertyEditorAsset.cpp
--- a/Engine/Source/Editor/PropertyEditor/Private/UserInterface/PropertyEditor/SPropertyEditorAsset.cpp
+++ b/Engine/Source/Editor/PropertyEditor/Private/UserInterface/PropertyEditor/SPropertyEditorAsset.cpp
@@ -1042,13 +1042,6 @@ void SPropertyEditorAsset::OnBrowse()
 	FObjectOrAssetData Value;
 	GetValue( Value );
 
-	// Try loading owning object
-	if (Value.Object == nullptr && Value.ObjectPath.IsValid())
-	{
-		FSoftObjectPath MapObjectPath = FSoftObjectPath(Value.ObjectPath.GetAssetPathName(), FString());
-		MapObjectPath.TryLoad();
-	}
-
 	if(PropertyEditor.IsValid() && Value.Object)
 	{
 		// This code only works on loaded objects
