Fixed crash when using Merge Actors with Replace Source Actors enabled.

Edit conditions are now correctly applied to fields of nested structs.

diff --git a/Engine/Source/Editor/PropertyEditor/Private/EditConditionContext.cpp b/Engine/Source/Editor/PropertyEditor/Private/EditConditionContext.cpp
--- a/Engine/Source/Editor/PropertyEditor/Private/EditConditionContext.cpp
+++ b/Engine/Source/Editor/PropertyEditor/Private/EditConditionContext.cpp
@@ -12,30 +12,13 @@ FEditConditionContext::FEditConditionContext(FPropertyNode& InPropertyNode)
 {
 	PropertyNode = InPropertyNode.AsShared();
 
-	FComplexPropertyNode* ComplexParentNode = FindComplexParent();
+	FComplexPropertyNode* ComplexParentNode = InPropertyNode.FindComplexParent();
 	check(ComplexParentNode);
 
 	const UProperty* Property = InPropertyNode.GetProperty();
 	check(Property);
-}
-
-FComplexPropertyNode* FEditConditionContext::FindComplexParent() const
-{
-	if (!PropertyNode.IsValid())
-	{
-		return nullptr;
-	}
 
-	TSharedPtr<FPropertyNode> PinnedNode = PropertyNode.Pin();
-
-	FPropertyNode* SearchStart = PinnedNode.Get();
-	if (PropertyEditorHelpers::IsStaticArray(*SearchStart))
-	{
-		//in the case of conditional static arrays, we have to go up one more level to get the proper parent struct.
-		SearchStart = SearchStart->GetParentNode();
-	}
-
-	return SearchStart->FindComplexParent();
+	check(!PropertyEditorHelpers::IsStaticArray(*InPropertyNode.GetParentNode()));
 }
 
 const UBoolProperty* FEditConditionContext::GetSingleBoolProperty(const TSharedPtr<FEditConditionExpression>& Expression) const
@@ -109,14 +92,22 @@ TOptional<bool> FEditConditionContext::GetBoolValue(const FString& PropertyName)
 	}
 
 	TSharedPtr<FPropertyNode> PinnedNode = PropertyNode.Pin();
+	FPropertyNode* ParentNode = PinnedNode->GetParentNode();
 	
 	TOptional<bool> Result;
 
-	FComplexPropertyNode* ComplexParentNode = FindComplexParent();
+	FComplexPropertyNode* ComplexParentNode = PinnedNode->FindComplexParent();
 	for (int32 Index = 0; Index < ComplexParentNode->GetInstancesNum(); ++Index)
 	{
 		uint8* BasePtr = ComplexParentNode->GetMemoryOfInstance(Index);
-		uint8* ValuePtr = BoolProperty->ContainerPtrToValuePtr<uint8>(BasePtr);
+		if (BasePtr == nullptr)
+		{
+			// deleted or not fully loaded yet
+			return TOptional<bool>();
+		}
+
+		uint8* ParentPtr = ParentNode->GetValueAddress(BasePtr);
+		uint8* ValuePtr = BoolProperty->ContainerPtrToValuePtr<uint8>(ParentPtr);
 
 		bool bValue = BoolProperty->GetPropertyValue(ValuePtr);
 		if (!Result.IsSet())
@@ -142,14 +133,22 @@ TOptional<double> FEditConditionContext::GetNumericValue(const FString& Property
 	}
 
 	TSharedPtr<FPropertyNode> PinnedNode = PropertyNode.Pin();
+	FPropertyNode* ParentNode = PinnedNode->GetParentNode();
 
 	TOptional<double> Result;
 
-	FComplexPropertyNode* ComplexParentNode = FindComplexParent();
+	FComplexPropertyNode* ComplexParentNode = PinnedNode->FindComplexParent();
 	for (int32 Index = 0; Index < ComplexParentNode->GetInstancesNum(); ++Index)
 	{
 		uint8* BasePtr = ComplexParentNode->GetMemoryOfInstance(Index);
-		uint8* ValuePtr = NumericProperty->ContainerPtrToValuePtr<uint8>(BasePtr);
+		if (BasePtr == nullptr)
+		{
+			// deleted or not fully loaded yet
+			return TOptional<double>();
+		}
+
+		uint8* ParentPtr = ParentNode->GetValueAddress(BasePtr);
+		uint8* ValuePtr = NumericProperty->ContainerPtrToValuePtr<uint8>(ParentPtr);
 
 		double Value = 0;
 
@@ -202,14 +201,22 @@ TOptional<FString> FEditConditionContext::GetEnumValue(const FString& PropertyNa
 	}
 	
 	TSharedPtr<FPropertyNode> PinnedNode = PropertyNode.Pin();
-
+	FPropertyNode* ParentNode = PinnedNode->GetParentNode();
+	
 	TOptional<int64> Result;
 
-	FComplexPropertyNode* ComplexParentNode = FindComplexParent();
+	FComplexPropertyNode* ComplexParentNode = PinnedNode->FindComplexParent();
 	for (int32 Index = 0; Index < ComplexParentNode->GetInstancesNum(); ++Index)
 	{
 		uint8* BasePtr = ComplexParentNode->GetMemoryOfInstance(Index);
-		uint8* ValuePtr = Property->ContainerPtrToValuePtr<uint8>(BasePtr);
+		if (BasePtr == nullptr)
+		{
+			// deleted or not fully loaded yet
+			return TOptional<FString>();
+		}
+
+		uint8* ParentPtr = ParentNode->GetValueAddress(BasePtr);
+		uint8* ValuePtr = Property->ContainerPtrToValuePtr<uint8>(ParentPtr);
 
 		int64 Value = NumericProperty->GetSignedIntPropertyValue(ValuePtr);
 		if (!Result.IsSet())
diff --git a/Engine/Source/Editor/PropertyEditor/Private/EditConditionContext.h b/Engine/Source/Editor/PropertyEditor/Private/EditConditionContext.h
--- a/Engine/Source/Editor/PropertyEditor/Private/EditConditionContext.h
+++ b/Engine/Source/Editor/PropertyEditor/Private/EditConditionContext.h
@@ -37,6 +37,4 @@ public:
 
 private:
 	TWeakPtr<FPropertyNode> PropertyNode;
-
-	FComplexPropertyNode* FindComplexParent() const;
 };
\ No newline at end of file
diff --git a/Engine/Source/Editor/PropertyEditor/Private/EditConditionParserTests.cpp b/Engine/Source/Editor/PropertyEditor/Private/EditConditionParserTests.cpp
--- a/Engine/Source/Editor/PropertyEditor/Private/EditConditionParserTests.cpp
+++ b/Engine/Source/Editor/PropertyEditor/Private/EditConditionParserTests.cpp
@@ -404,4 +404,37 @@ bool FEditConditionParser_EvaluateUObject::RunTest(const FString& Parameters)
 	return bAllResults;
 }
 
+IMPLEMENT_SIMPLE_AUTOMATION_TEST(FEditConditionParser_SingleBool, "EditConditionParser.SingleBool", EAutomationTestFlags::EditorContext | EAutomationTestFlags::EngineFilter)
+bool FEditConditionParser_SingleBool::RunTest(const FString& Parameters)
+{
+	UEditConditionTestObject* TestObject = NewObject<UEditConditionTestObject>();
+	TestObject->AddToRoot();
+
+	TSharedPtr<FObjectPropertyNode> ObjectNode(new FObjectPropertyNode);
+	ObjectNode->AddObject(TestObject);
+
+	FPropertyNodeInitParams InitParams;
+	ObjectNode->InitNode(InitParams);
+
+	static const FName BoolPropertyName(TEXT("BoolProperty"));
+	TSharedPtr<FPropertyNode> PropertyNode = ObjectNode->FindChildPropertyNode(BoolPropertyName, true);
+	FEditConditionContext Context(*PropertyNode.Get());
+
+	FEditConditionParser Parser;
+
+	{
+		TSharedPtr<FEditConditionExpression> Expression = Parser.Parse(FString(TEXT("BoolProperty")));
+		const UBoolProperty* Property = Context.GetSingleBoolProperty(Expression);
+		TestNotNull(TEXT("Bool"), Property);
+	}
+
+	{
+		TSharedPtr<FEditConditionExpression> Expression = Parser.Parse(FString(TEXT("UintBitfieldProperty")));
+		const UBoolProperty* Property = Context.GetSingleBoolProperty(Expression);
+		TestNotNull(TEXT("Uint Bitfield"), Property);
+	}
+
+	return true;
+}
+
 #endif // WITH_DEV_AUTOMATION_TESTS
\ No newline at end of file
diff --git a/Engine/Source/Editor/PropertyEditor/Private/EditConditionParserTests.h b/Engine/Source/Editor/PropertyEditor/Private/EditConditionParserTests.h
--- a/Engine/Source/Editor/PropertyEditor/Private/EditConditionParserTests.h
+++ b/Engine/Source/Editor/PropertyEditor/Private/EditConditionParserTests.h
@@ -38,4 +38,7 @@ class UEditConditionTestObject : public UObject
 
 	UPROPERTY(EditAnywhere, Category=Test)
 	int32 IntegerProperty;
+
+	UPROPERTY(EditAnywhere, Category=Test)
+	uint8 UintBitfieldProperty : 1;
 };
diff --git a/Engine/Source/Editor/PropertyEditor/Private/Presentation/PropertyEditor/PropertyEditor.cpp b/Engine/Source/Editor/PropertyEditor/Private/Presentation/PropertyEditor/PropertyEditor.cpp
--- a/Engine/Source/Editor/PropertyEditor/Private/Presentation/PropertyEditor/PropertyEditor.cpp
+++ b/Engine/Source/Editor/PropertyEditor/Private/Presentation/PropertyEditor/PropertyEditor.cpp
@@ -449,56 +449,53 @@ void FPropertyEditor::ToggleEditConditionState()
 {
 	const FScopedTransaction Transaction(FText::Format(LOCTEXT("SetEditConditionState", "Set {0} edit condition state "), PropertyNode->GetDisplayName()));
 
-	// Propagate the value change to any instances if we're editing a template object.
-	FObjectPropertyNode* ObjectNode = PropertyNode->FindObjectItemParent();
-
-	FPropertyNode* ParentNode = PropertyNode->GetParentNode();
-	check(ParentNode != nullptr);
-
 	PropertyNode->NotifyPreChange( PropertyNode->GetProperty(), PropertyUtilities->GetNotifyHook() );
 
 	const UBoolProperty* EditConditionProperty = EditConditionContext->GetSingleBoolProperty(EditConditionExpression);
 	check(EditConditionProperty != nullptr);
 
-	FComplexPropertyNode* ComplexParentNode = ParentNode->FindComplexParent();
+	FPropertyNode* ParentNode = PropertyNode->GetParentNode();
+	check(ParentNode != nullptr);
+
+	bool OldValue = true;
+
+	FComplexPropertyNode* ComplexParentNode = PropertyNode->FindComplexParent();
 	for (int32 Index = 0; Index < ComplexParentNode->GetInstancesNum(); ++Index)
 	{
+		// ComplexParentNode points to the top-level object 
+		// ParentNode can point to a struct inside that object (which is stored as an FItemPropertyNode)
+		// We need all three pointers to get the value pointer
 		uint8* BaseAddress = ComplexParentNode->GetMemoryOfInstance(Index);
-		check(BaseAddress != nullptr);
-
-		// Get the address corresponding to the base of this property (i.e. if a struct property, set BaseOffset to the address of value for the whole struct)
-		uint8* BaseOffset = ParentNode->GetValueAddress(BaseAddress);
-		check(BaseOffset != NULL);
+		uint8* ParentOffset = ParentNode->GetValueAddress(BaseAddress);
+		uint8* ValuePtr = EditConditionProperty->ContainerPtrToValuePtr<uint8>(ParentOffset);
 
-		uint8* ValueAddr = EditConditionProperty->ContainerPtrToValuePtr<uint8>(BaseOffset);
-
-		const bool OldValue = EditConditionProperty->GetPropertyValue(ValueAddr);
-		const bool NewValue = !OldValue;
-		EditConditionProperty->SetPropertyValue(ValueAddr, NewValue);
+		OldValue &= EditConditionProperty->GetPropertyValue(ValuePtr);
+		EditConditionProperty->SetPropertyValue(ValuePtr, !OldValue);
+	}
 
-		if (ObjectNode != nullptr)
+	// Propagate the value change to any instances if we're editing a template object
+	FObjectPropertyNode* ObjectNode = PropertyNode->FindObjectItemParent();
+	if (ObjectNode != nullptr)
+	{
+		for (int32 ObjIndex = 0; ObjIndex < ObjectNode->GetNumObjects(); ++ObjIndex)
 		{
-			for (int32 ObjIndex = 0; ObjIndex < ObjectNode->GetNumObjects(); ++ObjIndex)
+			TWeakObjectPtr<UObject> ObjectWeakPtr = ObjectNode->GetUObject(ObjIndex);
+			UObject* Object = ObjectWeakPtr.Get();
+			if (Object != nullptr && Object->IsTemplate())
 			{
-				TWeakObjectPtr<UObject> ObjectWeakPtr = ObjectNode->GetUObject(ObjIndex);
-				UObject* Object = ObjectWeakPtr.Get();
-				if (Object != nullptr && Object->IsTemplate())
+				TArray<UObject*> ArchetypeInstances;
+				Object->GetArchetypeInstances(ArchetypeInstances);
+				for (int32 InstanceIndex = 0; InstanceIndex < ArchetypeInstances.Num(); ++InstanceIndex)
 				{
-					TArray<UObject*> ArchetypeInstances;
-					Object->GetArchetypeInstances(ArchetypeInstances);
-					for (int32 InstanceIndex = 0; InstanceIndex < ArchetypeInstances.Num(); ++InstanceIndex)
+					uint8* ArchetypeBaseOffset = ComplexParentNode->GetValueAddress((uint8*) ArchetypeInstances[InstanceIndex]);
+					uint8* ArchetypeParentOffset = ParentNode->GetValueAddress(ArchetypeBaseOffset);
+					uint8* ArchetypeValueAddr = EditConditionProperty->ContainerPtrToValuePtr<uint8>(ArchetypeParentOffset);
+
+					// Only propagate if the current value on the instance matches the previous value on the template.
+					const bool CurValue = EditConditionProperty->GetPropertyValue(ArchetypeValueAddr);
+					if (OldValue == CurValue)
 					{
-						// Only propagate if the current value on the instance matches the previous value on the template.
-						uint8* ArchetypeBaseOffset = ParentNode->GetValueAddress((uint8*)ArchetypeInstances[InstanceIndex]);
-						if(ArchetypeBaseOffset)
-						{
-							uint8* ArchetypeValueAddr = EditConditionProperty->ContainerPtrToValuePtr<uint8>(ArchetypeBaseOffset);
-							const bool CurValue = EditConditionProperty->GetPropertyValue(ArchetypeValueAddr);
-							if(OldValue == CurValue)
-							{
-								EditConditionProperty->SetPropertyValue(ArchetypeValueAddr, NewValue);
-							}
-						}
+						EditConditionProperty->SetPropertyValue(ArchetypeValueAddr, !OldValue);
 					}
 				}
 			}
