Removed unnecessary check.

diff --git a/Engine/Source/Editor/PropertyEditor/Private/EditConditionContext.cpp b/Engine/Source/Editor/PropertyEditor/Private/EditConditionContext.cpp
--- a/Engine/Source/Editor/PropertyEditor/Private/EditConditionContext.cpp
+++ b/Engine/Source/Editor/PropertyEditor/Private/EditConditionContext.cpp
@@ -17,8 +17,6 @@ FEditConditionContext::FEditConditionContext(FPropertyNode& InPropertyNode)
 
 	const UProperty* Property = InPropertyNode.GetProperty();
 	check(Property);
-
-	check(!PropertyEditorHelpers::IsStaticArray(*InPropertyNode.GetParentNode()));
 }
 
 const UBoolProperty* FEditConditionContext::GetSingleBoolProperty(const TSharedPtr<FEditConditionExpression>& Expression) const
