Partial merge. Will redo SClassViewer.cpp implementation since the structure is different in UE4/Main. All other changes came over cleanly

#UE4 Added AssetReferenceFilter. An editor feature to allow projects to specify a filter to disallow editor systems to form references between certain assets. Add a listener for UUnrealEdEngine::OnMakeAssetReferenceFilter which returns a subclass of IAssetReferenceFilter that was constructed with the supplied context. PassesFilter will be called on this filter object to determine if the reference is allowed. Currently respected by asset picker, class picker, details panels, and edgraph drag/drop operations.

diff --git a/Engine/Source/Editor/ClassViewer/Public/ClassViewerModule.h b/Engine/Source/Editor/ClassViewer/Public/ClassViewerModule.h
--- a/Engine/Source/Editor/ClassViewer/Public/ClassViewerModule.h
+++ b/Engine/Source/Editor/ClassViewer/Public/ClassViewerModule.h
@@ -102,6 +102,9 @@ public:
 	/** The property this class viewer be working on. */
 	TSharedPtr<class IPropertyHandle> PropertyHandle;
 
+	/** The passed in property handle will be used to gather referencing assets. If additional referencing assets should be reported, supply them here. */
+	TArray<FAssetData> AdditionalReferencingAssets;
+
 	/** true (the default) shows the view options at the bottom of the class picker */
 	bool bAllowViewOptions;
 
diff --git a/Engine/Source/Editor/ContentBrowser/Private/SAssetPicker.cpp b/Engine/Source/Editor/ContentBrowser/Private/SAssetPicker.cpp
--- a/Engine/Source/Editor/ContentBrowser/Private/SAssetPicker.cpp
+++ b/Engine/Source/Editor/ContentBrowser/Private/SAssetPicker.cpp
@@ -20,6 +20,9 @@
 #include "SAssetView.h"
 #include "SContentBrowser.h"
 #include "Framework/Commands/GenericCommands.h"
+#include "UnrealEdGlobals.h"
+#include "Editor/UnrealEdEngine.h"
+#include "PropertyHandle.h"
 
 #define LOCTEXT_NAMESPACE "ContentBrowser"
 
@@ -220,12 +223,48 @@ void SAssetPicker::Construct( const FArguments& InArgs )
 		];
 	}
 
+	// Make game-specific filter
+	FOnShouldFilterAsset ShouldFilterAssetDelegate;
+	{
+		FAssetReferenceFilterContext AssetReferenceFilterContext;
+		AssetReferenceFilterContext.ReferencingAssets = InArgs._AssetPickerConfig.AdditionalReferencingAssets;
+		if (InArgs._AssetPickerConfig.PropertyHandle.IsValid())
+		{
+			TArray<UObject*> ReferencingObjects;
+			InArgs._AssetPickerConfig.PropertyHandle->GetOuterObjects(ReferencingObjects);
+			for (UObject* ReferencingObject : ReferencingObjects)
+			{
+				AssetReferenceFilterContext.ReferencingAssets.Add(FAssetData(ReferencingObject));
+			}
+		}
+		TSharedPtr<IAssetReferenceFilter> AssetReferenceFilter = GUnrealEd ? GUnrealEd->MakeAssetReferenceFilter(AssetReferenceFilterContext) : nullptr;
+		if (AssetReferenceFilter.IsValid())
+		{
+			FOnShouldFilterAsset ConfigFilter = InArgs._AssetPickerConfig.OnShouldFilterAsset;
+			ShouldFilterAssetDelegate = FOnShouldFilterAsset::CreateLambda([ConfigFilter, AssetReferenceFilter](const FAssetData& AssetData) -> bool {
+				if (!AssetReferenceFilter->PassesFilter(AssetData))
+				{
+					return true;
+				}
+				if (ConfigFilter.IsBound())
+				{
+					return ConfigFilter.Execute(AssetData);
+				}
+				return false;
+			});
+		}
+		else
+		{
+			ShouldFilterAssetDelegate = InArgs._AssetPickerConfig.OnShouldFilterAsset;
+		}
+	}
+
 	VerticalBox->AddSlot()
 	.FillHeight(1.f)
 	[
 		SAssignNew(AssetViewPtr, SAssetView)
 		.SelectionMode( InArgs._AssetPickerConfig.SelectionMode )
-		.OnShouldFilterAsset(InArgs._AssetPickerConfig.OnShouldFilterAsset)
+		.OnShouldFilterAsset(ShouldFilterAssetDelegate)
 		.OnAssetSelectionChanged(this, &SAssetPicker::HandleAssetSelectionChanged)
 		.OnAssetsActivated(this, &SAssetPicker::HandleAssetsActivated)
 		.OnGetAssetContextMenu(InArgs._AssetPickerConfig.OnGetAssetContextMenu)
diff --git a/Engine/Source/Editor/ContentBrowser/Public/IContentBrowserSingleton.h b/Engine/Source/Editor/ContentBrowser/Public/IContentBrowserSingleton.h
--- a/Engine/Source/Editor/ContentBrowser/Public/IContentBrowserSingleton.h
+++ b/Engine/Source/Editor/ContentBrowser/Public/IContentBrowserSingleton.h
@@ -184,6 +184,12 @@ struct FAssetPickerConfig
 	/** The asset that should be initially selected */
 	FAssetData InitialAssetSelection;
 
+	/** The handle to the property that opened this picker. Needed for contextual filtering. */
+	TSharedPtr<class IPropertyHandle> PropertyHandle;
+
+	/** The passed in property handle will be used to gather referencing assets. If additional referencing assets should be reported, supply them here. */
+	TArray<FAssetData> AdditionalReferencingAssets;
+
 	/** The delegate that fires when an asset was selected */
 	FOnAssetSelected OnAssetSelected;
 
diff --git a/Engine/Source/Editor/GraphEditor/Private/KismetNodes/SGraphNodeSpawnActorFromClass.cpp b/Engine/Source/Editor/GraphEditor/Private/KismetNodes/SGraphNodeSpawnActorFromClass.cpp
--- a/Engine/Source/Editor/GraphEditor/Private/KismetNodes/SGraphNodeSpawnActorFromClass.cpp
+++ b/Engine/Source/Editor/GraphEditor/Private/KismetNodes/SGraphNodeSpawnActorFromClass.cpp
@@ -66,6 +66,21 @@ protected:
 		Options.bShowObjectRootClass = true;
 		TSharedPtr< FActorBasedClassFilter > Filter = MakeShareable(new FActorBasedClassFilter);
 		Options.ClassFilter = Filter;
+		// Populate the referencing asset, if possible
+		if (UEdGraphPin* LocalGraphPin = GetPinObj())
+		{
+			if (UEdGraphNode* LocalNode = LocalGraphPin->GetOwningNode())
+			{
+				if (UEdGraph* LocalGraph = LocalNode->GetGraph())
+				{
+					UObject* GraphOuter = LocalGraph->GetOuter();
+					if (GraphOuter)
+					{
+						Options.AdditionalReferencingAssets.Add(FAssetData(GraphOuter));
+					}
+				}
+			}
+		}
 
 		return 
 			SNew(SBox)
diff --git a/Engine/Source/Editor/GraphEditor/Private/SGraphPanel.cpp b/Engine/Source/Editor/GraphEditor/Private/SGraphPanel.cpp
--- a/Engine/Source/Editor/GraphEditor/Private/SGraphPanel.cpp
+++ b/Engine/Source/Editor/GraphEditor/Private/SGraphPanel.cpp
@@ -29,6 +29,9 @@
 #include "KismetNodes/KismetNodeInfoContext.h"
 #include "GraphDiffControl.h"
 
+#include "Editor/UnrealEdEngine.h"
+#include "UnrealEdGlobals.h"
+
 DEFINE_LOG_CATEGORY_STATIC(LogGraphPanel, Log, All);
 
 //////////////////////////////////////////////////////////////////////////
@@ -909,13 +912,25 @@ FReply SGraphPanel::OnDragOver( const FGeometry& MyGeometry, const FDragDropEven
 		{
 			TSharedPtr<FAssetDragDropOp> AssetOp = StaticCastSharedPtr<FAssetDragDropOp>(Operation);
 			bool bOkIcon = false;
-			FString TooltipText;
+			FText TooltipText;
 			if (AssetOp->HasAssets())
 			{
-				GraphObj->GetSchema()->GetAssetsGraphHoverMessage(AssetOp->GetAssets(), GraphObj, /*out*/ TooltipText, /*out*/ bOkIcon);
+				const TArray<FAssetData>& HoveredAssetData = AssetOp->GetAssets();
+				FText AssetReferenceFilterFailureReason;
+				if (PassesAssetReferenceFilter(HoveredAssetData, &AssetReferenceFilterFailureReason))
+				{
+					FString TooltipTextString;
+					GraphObj->GetSchema()->GetAssetsGraphHoverMessage(HoveredAssetData, GraphObj, /*out*/ TooltipTextString, /*out*/ bOkIcon);
+					TooltipText = FText::FromString(TooltipTextString);
+				}
+				else
+				{
+					TooltipText = AssetReferenceFilterFailureReason;
+					bOkIcon = false;
+				}
 			}
 			const FSlateBrush* TooltipIcon = bOkIcon ? FEditorStyle::GetBrush(TEXT("Graph.ConnectorFeedback.OK")) : FEditorStyle::GetBrush(TEXT("Graph.ConnectorFeedback.Error"));
-			AssetOp->SetToolTip(FText::FromString(TooltipText), TooltipIcon);
+			AssetOp->SetToolTip(TooltipText, TooltipIcon);
 		}
 		return FReply::Handled();
 	} 
@@ -976,7 +991,10 @@ FReply SGraphPanel::OnDrop( const FGeometry& MyGeometry, const FDragDropEvent& D
 
 			if ( DroppedAssetData.Num() > 0 )
 			{
-				GraphObj->GetSchema()->DroppedAssetsOnGraph( DroppedAssetData, NodeAddPosition, GraphObj );
+				if (PassesAssetReferenceFilter(DroppedAssetData))
+				{
+					GraphObj->GetSchema()->DroppedAssetsOnGraph( DroppedAssetData, NodeAddPosition, GraphObj );
+				}
 				return FReply::Handled();
 			}
 		}
@@ -985,6 +1003,32 @@ FReply SGraphPanel::OnDrop( const FGeometry& MyGeometry, const FDragDropEvent& D
 	}
 }
 
+bool SGraphPanel::PassesAssetReferenceFilter(const TArray<FAssetData>& ReferencedAssets, FText* OutFailureReason) const
+{
+	if (GUnrealEd)
+	{
+		FAssetReferenceFilterContext AssetReferenceFilterContext;
+		UObject* GraphOuter = GraphObj ? GraphObj->GetOuter() : nullptr;
+		if (GraphOuter)
+		{
+			AssetReferenceFilterContext.ReferencingAssets.Add(FAssetData(GraphOuter));
+		}
+		TSharedPtr<IAssetReferenceFilter> AssetReferenceFilter = GUnrealEd->MakeAssetReferenceFilter(AssetReferenceFilterContext);
+		if (AssetReferenceFilter.IsValid())
+		{
+			for (const FAssetData& Asset : ReferencedAssets)
+			{
+				if (!AssetReferenceFilter->PassesFilter(Asset, OutFailureReason))
+				{
+					return false;
+				}
+			}
+		}
+	}
+
+	return true;
+}
+
 void SGraphPanel::OnBeginMakingConnection(UEdGraphPin* InOriginatingPin)
 {
 	OnBeginMakingConnection(FGraphPinHandle(InOriginatingPin));
diff --git a/Engine/Source/Editor/GraphEditor/Public/SGraphPanel.h b/Engine/Source/Editor/GraphEditor/Public/SGraphPanel.h
--- a/Engine/Source/Editor/GraphEditor/Public/SGraphPanel.h
+++ b/Engine/Source/Editor/GraphEditor/Public/SGraphPanel.h
@@ -313,6 +313,9 @@ private:
 	/** Returns the pin that we're considering as hovered if we are hovering over a spline; may be null */
 	class SGraphPin* GetBestPinFromHoveredSpline() const;
 
+	/** Returns true if all assets in the ReferencedAssets array are allowed to be referenced by this graph according to the AssetReferenceFilter. OutFailureReason if supplied if it is not.  */
+	bool PassesAssetReferenceFilter(const TArray<FAssetData>& ReferencedAssets, FText* OutFailureReason = nullptr) const;
+
 	/** Handle to timer callback that allows the UI to refresh it's arrangement each tick, allows animations to occur within the UI */
 	TWeakPtr<FActiveTimerHandle> ActiveTimerHandleInvalidatePerTick;
 
diff --git a/Engine/Source/Editor/PropertyEditor/Private/PropertyCustomizationHelpers.cpp b/Engine/Source/Editor/PropertyEditor/Private/PropertyCustomizationHelpers.cpp
--- a/Engine/Source/Editor/PropertyEditor/Private/PropertyCustomizationHelpers.cpp
+++ b/Engine/Source/Editor/PropertyEditor/Private/PropertyCustomizationHelpers.cpp
@@ -259,36 +259,38 @@ namespace PropertyCustomizationHelpers
 			];
 	}
 
-	TSharedRef<SWidget> MakeAssetPickerAnchorButton( FOnGetAllowedClasses OnGetAllowedClasses, FOnAssetSelected OnAssetSelectedFromPicker )
+	TSharedRef<SWidget> MakeAssetPickerAnchorButton( FOnGetAllowedClasses OnGetAllowedClasses, FOnAssetSelected OnAssetSelectedFromPicker, const TSharedPtr<IPropertyHandle>& PropertyHandle)
 	{
 		return 
 			SNew( SPropertyAssetPicker )
 			.OnGetAllowedClasses( OnGetAllowedClasses )
-			.OnAssetSelected( OnAssetSelectedFromPicker );
+			.OnAssetSelected( OnAssetSelectedFromPicker )
+			.PropertyHandle( PropertyHandle );
 	}
 
 	TArray<const UClass*> EmptyClassArray;
 
-	TSharedRef<SWidget> MakeAssetPickerWithMenu(const FAssetData& InitialObject, const bool AllowClear, const TArray<const UClass*>& AllowedClasses, const TArray<UFactory*>& NewAssetFactories, FOnShouldFilterAsset OnShouldFilterAsset, FOnAssetSelected OnSet, FSimpleDelegate OnClose)
+	TSharedRef<SWidget> MakeAssetPickerWithMenu(const FAssetData& InitialObject, const bool AllowClear, const TArray<const UClass*>& AllowedClasses, const TArray<UFactory*>& NewAssetFactories, FOnShouldFilterAsset OnShouldFilterAsset, FOnAssetSelected OnSet, FSimpleDelegate OnClose, const TSharedPtr<IPropertyHandle>& PropertyHandle)
 	{
-		return MakeAssetPickerWithMenu(InitialObject, AllowClear, true, AllowedClasses, EmptyClassArray, NewAssetFactories, OnShouldFilterAsset, OnSet, OnClose);
+		return MakeAssetPickerWithMenu(InitialObject, AllowClear, true, AllowedClasses, EmptyClassArray, NewAssetFactories, OnShouldFilterAsset, OnSet, OnClose, PropertyHandle);
 	}
 
-	TSharedRef<SWidget> MakeAssetPickerWithMenu(const FAssetData& InitialObject, const bool AllowClear, const TArray<const UClass*>& AllowedClasses, const TArray<const UClass*>& DisallowedClasses, const TArray<UFactory*>& NewAssetFactories, FOnShouldFilterAsset OnShouldFilterAsset, FOnAssetSelected OnSet, FSimpleDelegate OnClose)
+	TSharedRef<SWidget> MakeAssetPickerWithMenu(const FAssetData& InitialObject, const bool AllowClear, const TArray<const UClass*>& AllowedClasses, const TArray<const UClass*>& DisallowedClasses, const TArray<UFactory*>& NewAssetFactories, FOnShouldFilterAsset OnShouldFilterAsset, FOnAssetSelected OnSet, FSimpleDelegate OnClose, const TSharedPtr<IPropertyHandle>& PropertyHandle)
 	{
-		return MakeAssetPickerWithMenu(InitialObject, AllowClear, true, AllowedClasses, DisallowedClasses, NewAssetFactories, OnShouldFilterAsset, OnSet, OnClose);
+		return MakeAssetPickerWithMenu(InitialObject, AllowClear, true, AllowedClasses, DisallowedClasses, NewAssetFactories, OnShouldFilterAsset, OnSet, OnClose, PropertyHandle);
 	}
 
-	TSharedRef<SWidget> MakeAssetPickerWithMenu(const FAssetData& InitialObject, const bool AllowClear, const bool AllowCopyPaste, const TArray<const UClass*>& AllowedClasses, const TArray<UFactory*>& NewAssetFactories, FOnShouldFilterAsset OnShouldFilterAsset, FOnAssetSelected OnSet, FSimpleDelegate OnClose)
+	TSharedRef<SWidget> MakeAssetPickerWithMenu(const FAssetData& InitialObject, const bool AllowClear, const bool AllowCopyPaste, const TArray<const UClass*>& AllowedClasses, const TArray<UFactory*>& NewAssetFactories, FOnShouldFilterAsset OnShouldFilterAsset, FOnAssetSelected OnSet, FSimpleDelegate OnClose, const TSharedPtr<IPropertyHandle>& PropertyHandle)
 	{
-		return MakeAssetPickerWithMenu(InitialObject, AllowClear, AllowCopyPaste, AllowedClasses, EmptyClassArray, NewAssetFactories, OnShouldFilterAsset, OnSet, OnClose);
+		return MakeAssetPickerWithMenu(InitialObject, AllowClear, AllowCopyPaste, AllowedClasses, EmptyClassArray, NewAssetFactories, OnShouldFilterAsset, OnSet, OnClose, PropertyHandle);
 	}
 
-	TSharedRef<SWidget> MakeAssetPickerWithMenu( const FAssetData& InitialObject, const bool AllowClear, const bool AllowCopyPaste, const TArray<const UClass*>& AllowedClasses, const TArray<const UClass*>& DisallowedClasses, const TArray<UFactory*>& NewAssetFactories, FOnShouldFilterAsset OnShouldFilterAsset, FOnAssetSelected OnSet, FSimpleDelegate OnClose)
+	TSharedRef<SWidget> MakeAssetPickerWithMenu( const FAssetData& InitialObject, const bool AllowClear, const bool AllowCopyPaste, const TArray<const UClass*>& AllowedClasses, const TArray<const UClass*>& DisallowedClasses, const TArray<UFactory*>& NewAssetFactories, FOnShouldFilterAsset OnShouldFilterAsset, FOnAssetSelected OnSet, FSimpleDelegate OnClose, const TSharedPtr<IPropertyHandle>& PropertyHandle)
 	{
 		return
 			SNew(SPropertyMenuAssetPicker)
 			.InitialObject(InitialObject)
+			.PropertyHandle(PropertyHandle)
 			.AllowClear(AllowClear)
 			.AllowCopyPaste(AllowCopyPaste)
 			.AllowedClasses(AllowedClasses)
diff --git a/Engine/Source/Editor/PropertyEditor/Private/PropertyEditorHelpers.cpp b/Engine/Source/Editor/PropertyEditor/Private/PropertyEditorHelpers.cpp
--- a/Engine/Source/Editor/PropertyEditor/Private/PropertyEditorHelpers.cpp
+++ b/Engine/Source/Editor/PropertyEditor/Private/PropertyEditorHelpers.cpp
@@ -1031,7 +1031,7 @@ namespace PropertyEditorHelpers
 			}
 
 		case EPropertyButton::PickAsset:
-			NewButton = PropertyCustomizationHelpers::MakeAssetPickerAnchorButton( FOnGetAllowedClasses::CreateSP( PropertyEditor, &FPropertyEditor::OnGetClassesForAssetPicker ), FOnAssetSelected::CreateSP( PropertyEditor, &FPropertyEditor::OnAssetSelected ) );
+			NewButton = PropertyCustomizationHelpers::MakeAssetPickerAnchorButton( FOnGetAllowedClasses::CreateSP( PropertyEditor, &FPropertyEditor::OnGetClassesForAssetPicker ), FOnAssetSelected::CreateSP( PropertyEditor, &FPropertyEditor::OnAssetSelected ), PropertyEditor->GetPropertyHandle());
 			break;
 
 		case EPropertyButton::PickActor:
diff --git a/Engine/Source/Editor/PropertyEditor/Private/UserInterface/PropertyEditor/SPropertyAssetPicker.cpp b/Engine/Source/Editor/PropertyEditor/Private/UserInterface/PropertyEditor/SPropertyAssetPicker.cpp
--- a/Engine/Source/Editor/PropertyEditor/Private/UserInterface/PropertyEditor/SPropertyAssetPicker.cpp
+++ b/Engine/Source/Editor/PropertyEditor/Private/UserInterface/PropertyEditor/SPropertyAssetPicker.cpp
@@ -14,6 +14,7 @@ void SPropertyAssetPicker::Construct( const FArguments& InArgs )
 {
 	OnAssetSelected = InArgs._OnAssetSelected;
 	OnGetAllowedClasses = InArgs._OnGetAllowedClasses;
+	PropertyHandle = InArgs._PropertyHandle;
 
 	ChildSlot
 	[
@@ -77,6 +78,9 @@ TSharedRef<SWidget> SPropertyAssetPicker::OnGenerateAssetPicker()
 	AssetPickerConfig.bAllowDragging = false;
 	// Use the list view by default
 	AssetPickerConfig.InitialAssetViewType = EAssetViewType::List;
+	// Populate the referencing assets via property handle
+	AssetPickerConfig.PropertyHandle = PropertyHandle;
+	
 	TSharedRef<SWidget> MenuContent = 
 		SNew(SBox)
 		.HeightOverride(300)
diff --git a/Engine/Source/Editor/PropertyEditor/Private/UserInterface/PropertyEditor/SPropertyAssetPicker.h b/Engine/Source/Editor/PropertyEditor/Private/UserInterface/PropertyEditor/SPropertyAssetPicker.h
--- a/Engine/Source/Editor/PropertyEditor/Private/UserInterface/PropertyEditor/SPropertyAssetPicker.h
+++ b/Engine/Source/Editor/PropertyEditor/Private/UserInterface/PropertyEditor/SPropertyAssetPicker.h
@@ -15,6 +15,7 @@ public:
 	SLATE_BEGIN_ARGS( SPropertyAssetPicker ) {}
 		SLATE_EVENT( FOnGetAllowedClasses, OnGetAllowedClasses )
 		SLATE_EVENT( FOnAssetSelected, OnAssetSelected )
+		SLATE_ARGUMENT( TSharedPtr<IPropertyHandle>, PropertyHandle )
 	SLATE_END_ARGS()
 
 	void Construct( const FArguments& InArgs );
@@ -32,4 +33,5 @@ private:
 
 	FOnGetAllowedClasses OnGetAllowedClasses;
 	FOnAssetSelected OnAssetSelected;
+	TSharedPtr<IPropertyHandle> PropertyHandle;
 };
diff --git a/Engine/Source/Editor/PropertyEditor/Private/UserInterface/PropertyEditor/SPropertyEditorAsset.cpp b/Engine/Source/Editor/PropertyEditor/Private/UserInterface/PropertyEditor/SPropertyEditorAsset.cpp
--- a/Engine/Source/Editor/PropertyEditor/Private/UserInterface/PropertyEditor/SPropertyEditorAsset.cpp
+++ b/Engine/Source/Editor/PropertyEditor/Private/UserInterface/PropertyEditor/SPropertyEditorAsset.cpp
@@ -25,6 +25,10 @@
 #include "ObjectPropertyNode.h"
 #include "PropertyHandleImpl.h"
 #include "HAL/PlatformApplicationMisc.h"
+#include "Editor/UnrealEdEngine.h"
+#include "UnrealEdGlobals.h"
+#include "Framework/Notifications/NotificationManager.h"
+#include "Widgets/Notifications/SNotificationList.h"
 
 #define LOCTEXT_NAMESPACE "PropertyEditor"
 
@@ -666,7 +670,8 @@ TSharedRef<SWidget> SPropertyEditorAsset::OnGetMenuContent()
 																	 NewAssetFactories,
 																	 OnShouldFilterAsset,
 																	 FOnAssetSelected::CreateSP(this, &SPropertyEditorAsset::OnAssetSelected),
-																	 FSimpleDelegate::CreateSP(this, &SPropertyEditorAsset::CloseComboButton));
+																	 FSimpleDelegate::CreateSP(this, &SPropertyEditorAsset::CloseComboButton),
+																	 GetMostSpecificPropertyHandle());
 	}
 }
 
@@ -832,16 +837,24 @@ void SPropertyEditorAsset::SetValue( const FAssetData& AssetData )
 {
 	AssetComboButton->SetIsOpen(false);
 
-	bool bAllowedToSetBasedOnFilter = CanSetBasedOnCustomClasses( AssetData );
-
-	if( bAllowedToSetBasedOnFilter )
+	if(CanSetBasedOnCustomClasses(AssetData))
 	{
-		if(PropertyEditor.IsValid())
+		FText AssetReferenceFilterFailureReason;
+		if (CanSetBasedOnAssetReferenceFilter(AssetData, &AssetReferenceFilterFailureReason))
 		{
-			PropertyEditor->GetPropertyHandle()->SetValue(AssetData);
-		}
+			if (PropertyEditor.IsValid())
+			{
+				PropertyEditor->GetPropertyHandle()->SetValue(AssetData);
+			}
 
-		OnSetObject.ExecuteIfBound(AssetData);
+			OnSetObject.ExecuteIfBound(AssetData);
+		}
+		else if (!AssetReferenceFilterFailureReason.IsEmpty())
+		{
+			FNotificationInfo Info(AssetReferenceFilterFailureReason);
+			Info.ExpireDuration = 4.f;
+			FSlateNotificationManager::Get().AddNotification(Info);
+		}
 	}
 }
 
@@ -1080,7 +1093,11 @@ FText SPropertyEditorAsset::GetOnBrowseToolTip() const
 void SPropertyEditorAsset::OnUse()
 {
 	// Use the property editor path if it is valid and there is no custom filtering required
-	if(PropertyEditor.IsValid() && !OnShouldFilterAsset.IsBound() && AllowedClassFilters.Num() == 0 && DisallowedClassFilters.Num() == 0)
+	if(PropertyEditor.IsValid()
+		&& !OnShouldFilterAsset.IsBound()
+		&& AllowedClassFilters.Num() == 0
+		&& DisallowedClassFilters.Num() == 0
+		&& (GUnrealEd ? !GUnrealEd->MakeAssetReferenceFilter(FAssetReferenceFilterContext()).IsValid() : true))
 	{
 		PropertyEditor->GetPropertyHandle()->SetObjectValueFromSelection();
 	}
@@ -1138,10 +1155,14 @@ bool SPropertyEditorAsset::OnAssetDraggedOver( const UObject* InObject ) const
 	if (CanEdit() && InObject != nullptr && InObject->IsA(ObjectClass))
 	{
 		// Check against custom asset filter
+		FAssetData AssetData(InObject);
 		if (!OnShouldFilterAsset.IsBound()
-			|| !OnShouldFilterAsset.Execute(FAssetData(InObject)))
+			|| !OnShouldFilterAsset.Execute(AssetData))
 		{
-			return CanSetBasedOnCustomClasses( FAssetData(InObject) );
+			if (CanSetBasedOnCustomClasses(AssetData))
+			{
+				return CanSetBasedOnAssetReferenceFilter(AssetData);
+			}
 		}
 	}
 
@@ -1275,6 +1296,45 @@ bool SPropertyEditorAsset::CanSetBasedOnCustomClasses( const FAssetData& InAsset
 	return bAllowedToSetBasedOnFilter;
 }
 
+bool SPropertyEditorAsset::CanSetBasedOnAssetReferenceFilter( const FAssetData& InAssetData, FText* OutOptionalFailureReason) const
+{
+	if (GUnrealEd && InAssetData.IsValid())
+	{
+		TSharedPtr<IPropertyHandle> PropertyHandleToUse = GetMostSpecificPropertyHandle();
+		FAssetReferenceFilterContext AssetReferenceFilterContext;
+		if (PropertyHandleToUse.IsValid())
+		{
+			TArray<UObject*> ReferencingObjects;
+			PropertyHandleToUse->GetOuterObjects(ReferencingObjects);
+			for (UObject* ReferencingObject : ReferencingObjects)
+			{
+				AssetReferenceFilterContext.ReferencingAssets.Add(FAssetData(ReferencingObject));
+			}
+		}
+		TSharedPtr<IAssetReferenceFilter> AssetReferenceFilter = GUnrealEd->MakeAssetReferenceFilter(AssetReferenceFilterContext);
+		if (AssetReferenceFilter.IsValid() && !AssetReferenceFilter->PassesFilter(InAssetData, OutOptionalFailureReason))
+		{
+			return false;
+		}
+	}
+
+	return true;
+}
+
+TSharedPtr<IPropertyHandle> SPropertyEditorAsset::GetMostSpecificPropertyHandle() const
+{
+	if (PropertyHandle.IsValid())
+	{
+		return PropertyHandle;
+	}
+	else if (PropertyEditor.IsValid())
+	{
+		return PropertyEditor->GetPropertyHandle();
+	}
+	
+	return TSharedPtr<IPropertyHandle>();
+}
+
 UClass* SPropertyEditorAsset::GetObjectPropertyClass(const UProperty* Property)
 {
 	UClass* Class = nullptr;
diff --git a/Engine/Source/Editor/PropertyEditor/Private/UserInterface/PropertyEditor/SPropertyEditorAsset.h b/Engine/Source/Editor/PropertyEditor/Private/UserInterface/PropertyEditor/SPropertyEditorAsset.h
--- a/Engine/Source/Editor/PropertyEditor/Private/UserInterface/PropertyEditor/SPropertyEditorAsset.h
+++ b/Engine/Source/Editor/PropertyEditor/Private/UserInterface/PropertyEditor/SPropertyEditorAsset.h
@@ -304,6 +304,12 @@ private:
 	/** @return true if the passed in AssetData can be used to set the property based on the list of custom classes */
 	bool CanSetBasedOnCustomClasses( const FAssetData& InAssetData ) const;
 
+	/** @return true if the passed in AssetData can be used to set the property based on the game specific asset reference filter */
+	bool CanSetBasedOnAssetReferenceFilter( const FAssetData& InAssetData, FText* OutOptionalFailureReason = nullptr ) const;
+
+	/** @return Returns the property handle most relevant */
+	TSharedPtr<class IPropertyHandle> GetMostSpecificPropertyHandle() const;
+
 	/**
 	 * Gets the class of the supplied property for use within the PropertyEditorAsset widget. Asserts if the property
 	 * is not supported by the widget.
diff --git a/Engine/Source/Editor/PropertyEditor/Private/UserInterface/PropertyEditor/SPropertyMenuAssetPicker.cpp b/Engine/Source/Editor/PropertyEditor/Private/UserInterface/PropertyEditor/SPropertyMenuAssetPicker.cpp
--- a/Engine/Source/Editor/PropertyEditor/Private/UserInterface/PropertyEditor/SPropertyMenuAssetPicker.cpp
+++ b/Engine/Source/Editor/PropertyEditor/Private/UserInterface/PropertyEditor/SPropertyMenuAssetPicker.cpp
@@ -20,6 +20,7 @@
 void SPropertyMenuAssetPicker::Construct( const FArguments& InArgs )
 {
 	CurrentObject = InArgs._InitialObject;
+	PropertyHandle = InArgs._PropertyHandle;
 	bAllowClear = InArgs._AllowClear;
 	bAllowCopyPaste = InArgs._AllowCopyPaste;
 	AllowedClasses = InArgs._AllowedClasses;
@@ -143,6 +144,8 @@ void SPropertyMenuAssetPicker::Construct( const FArguments& InArgs )
 		AssetPickerConfig.bAllowDragging = false;
 		// Save the settings into a special section for asset pickers for properties
 		AssetPickerConfig.SaveSettingsName = TEXT("AssetPropertyPicker");
+		// Populate the referencing assets via property handle
+		AssetPickerConfig.PropertyHandle = PropertyHandle;
 
 		MenuContent =
 			SNew(SBox)
diff --git a/Engine/Source/Editor/PropertyEditor/Private/UserInterface/PropertyEditor/SPropertyMenuAssetPicker.h b/Engine/Source/Editor/PropertyEditor/Private/UserInterface/PropertyEditor/SPropertyMenuAssetPicker.h
--- a/Engine/Source/Editor/PropertyEditor/Private/UserInterface/PropertyEditor/SPropertyMenuAssetPicker.h
+++ b/Engine/Source/Editor/PropertyEditor/Private/UserInterface/PropertyEditor/SPropertyMenuAssetPicker.h
@@ -18,6 +18,7 @@ public:
 		, _AllowClear(true)
 	{}
 		SLATE_ARGUMENT( FAssetData, InitialObject )
+		SLATE_ARGUMENT( TSharedPtr<IPropertyHandle>, PropertyHandle )
 		SLATE_ARGUMENT( bool, AllowClear )
 		SLATE_ARGUMENT( bool, AllowCopyPaste )
 		SLATE_ARGUMENT( TArray<const UClass*>, AllowedClasses )
@@ -88,6 +89,7 @@ private:
 
 private:
 	FAssetData CurrentObject;
+	TSharedPtr<IPropertyHandle> PropertyHandle;
 
 	/** Whether the asset can be 'None' in this case */
 	bool bAllowClear;
diff --git a/Engine/Source/Editor/PropertyEditor/Public/PropertyCustomizationHelpers.h b/Engine/Source/Editor/PropertyEditor/Public/PropertyCustomizationHelpers.h
--- a/Engine/Source/Editor/PropertyEditor/Public/PropertyCustomizationHelpers.h
+++ b/Engine/Source/Editor/PropertyEditor/Public/PropertyCustomizationHelpers.h
@@ -32,6 +32,7 @@ class SPropertyEditorStruct;
 class UActorComponent;
 class UFactory;
 class SToolTip;
+class IPropertyHandle;
 class IDetailGroup;
 class IDetailCategoryBuilder;
 
@@ -64,11 +65,11 @@ namespace PropertyCustomizationHelpers
 	PROPERTYEDITOR_API TSharedRef<SWidget> MakeNewBlueprintButton( FSimpleDelegate OnFindClicked, TAttribute<FText> OptionalToolTipText = FText(), TAttribute<bool> IsEnabled = true );
 	PROPERTYEDITOR_API TSharedRef<SWidget> MakeUseSelectedButton( FSimpleDelegate OnUseSelectedClicked, TAttribute<FText> OptionalToolTipText = FText(), TAttribute<bool> IsEnabled = true );
 	PROPERTYEDITOR_API TSharedRef<SWidget> MakeBrowseButton( FSimpleDelegate OnClearClicked, TAttribute<FText> OptionalToolTipText = FText(), TAttribute<bool> IsEnabled = true );
-	PROPERTYEDITOR_API TSharedRef<SWidget> MakeAssetPickerAnchorButton( FOnGetAllowedClasses OnGetAllowedClasses, FOnAssetSelected OnAssetSelectedFromPicker );
-	PROPERTYEDITOR_API TSharedRef<SWidget> MakeAssetPickerWithMenu( const FAssetData& InitialObject, const bool AllowClear, const TArray<const UClass*>& AllowedClasses, const TArray<UFactory*>& NewAssetFactories, FOnShouldFilterAsset OnShouldFilterAsset, FOnAssetSelected OnSet, FSimpleDelegate OnClose );
-	PROPERTYEDITOR_API TSharedRef<SWidget> MakeAssetPickerWithMenu( const FAssetData& InitialObject, const bool AllowClear, const TArray<const UClass*>& AllowedClasses, const TArray<const UClass*>& DisallowedClasses, const TArray<UFactory*>& NewAssetFactories, FOnShouldFilterAsset OnShouldFilterAsset, FOnAssetSelected OnSet, FSimpleDelegate OnClose );
-	PROPERTYEDITOR_API TSharedRef<SWidget> MakeAssetPickerWithMenu( const FAssetData& InitialObject, const bool AllowClear, const bool AllowCopyPaste, const TArray<const UClass*>& AllowedClasses, const TArray<UFactory*>& NewAssetFactories, FOnShouldFilterAsset OnShouldFilterAsset, FOnAssetSelected OnSet, FSimpleDelegate OnClose );
-	PROPERTYEDITOR_API TSharedRef<SWidget> MakeAssetPickerWithMenu( const FAssetData& InitialObject, const bool AllowClear, const bool AllowCopyPaste, const TArray<const UClass*>& AllowedClasses, const TArray<const UClass*>& DisallowedClasses, const TArray<UFactory*>& NewAssetFactories, FOnShouldFilterAsset OnShouldFilterAsset, FOnAssetSelected OnSet, FSimpleDelegate OnClose );
+	PROPERTYEDITOR_API TSharedRef<SWidget> MakeAssetPickerAnchorButton( FOnGetAllowedClasses OnGetAllowedClasses, FOnAssetSelected OnAssetSelectedFromPicker, const TSharedPtr<IPropertyHandle>& PropertyHandle = TSharedPtr<IPropertyHandle>());
+	PROPERTYEDITOR_API TSharedRef<SWidget> MakeAssetPickerWithMenu( const FAssetData& InitialObject, const bool AllowClear, const TArray<const UClass*>& AllowedClasses, const TArray<UFactory*>& NewAssetFactories, FOnShouldFilterAsset OnShouldFilterAsset, FOnAssetSelected OnSet, FSimpleDelegate OnClose, const TSharedPtr<IPropertyHandle>& PropertyHandle = TSharedPtr<IPropertyHandle>());
+	PROPERTYEDITOR_API TSharedRef<SWidget> MakeAssetPickerWithMenu( const FAssetData& InitialObject, const bool AllowClear, const TArray<const UClass*>& AllowedClasses, const TArray<const UClass*>& DisallowedClasses, const TArray<UFactory*>& NewAssetFactories, FOnShouldFilterAsset OnShouldFilterAsset, FOnAssetSelected OnSet, FSimpleDelegate OnClose, const TSharedPtr<IPropertyHandle>& PropertyHandle = TSharedPtr<IPropertyHandle>());
+	PROPERTYEDITOR_API TSharedRef<SWidget> MakeAssetPickerWithMenu( const FAssetData& InitialObject, const bool AllowClear, const bool AllowCopyPaste, const TArray<const UClass*>& AllowedClasses, const TArray<UFactory*>& NewAssetFactories, FOnShouldFilterAsset OnShouldFilterAsset, FOnAssetSelected OnSet, FSimpleDelegate OnClose, const TSharedPtr<IPropertyHandle>& PropertyHandle = TSharedPtr<IPropertyHandle>());
+	PROPERTYEDITOR_API TSharedRef<SWidget> MakeAssetPickerWithMenu( const FAssetData& InitialObject, const bool AllowClear, const bool AllowCopyPaste, const TArray<const UClass*>& AllowedClasses, const TArray<const UClass*>& DisallowedClasses, const TArray<UFactory*>& NewAssetFactories, FOnShouldFilterAsset OnShouldFilterAsset, FOnAssetSelected OnSet, FSimpleDelegate OnClose, const TSharedPtr<IPropertyHandle>& PropertyHandle = TSharedPtr<IPropertyHandle>());
 	PROPERTYEDITOR_API TSharedRef<SWidget> MakeActorPickerAnchorButton( FOnGetActorFilters OnGetActorFilters, FOnActorSelected OnActorSelectedFromPicker );
 	PROPERTYEDITOR_API TSharedRef<SWidget> MakeActorPickerWithMenu( AActor* const InitialActor, const bool AllowClear, FOnShouldFilterActor ActorFilter, FOnActorSelected OnSet, FSimpleDelegate OnClose, FSimpleDelegate OnUseSelected );
 	PROPERTYEDITOR_API TSharedRef<SWidget> MakeComponentPickerWithMenu( UActorComponent* const InitialComponent, const bool AllowClear, FOnShouldFilterActor ActorFilter, FOnShouldFilterComponent ComponentFilter, FOnComponentSelected OnSet, FSimpleDelegate OnClose );
diff --git a/Engine/Source/Editor/UnrealEd/Classes/Editor/UnrealEdEngine.h b/Engine/Source/Editor/UnrealEd/Classes/Editor/UnrealEdEngine.h
--- a/Engine/Source/Editor/UnrealEd/Classes/Editor/UnrealEdEngine.h
+++ b/Engine/Source/Editor/UnrealEd/Classes/Editor/UnrealEdEngine.h
@@ -108,6 +108,22 @@ struct FTemplateMapInfo
 
 class FPerformanceMonitor;
 
+
+/** Struct used in filtering allowed references between assets. Passes context about the referencers to game-level filters */
+struct FAssetReferenceFilterContext
+{
+	TArray<FAssetData> ReferencingAssets;
+};
+
+/** Used in filtering allowed references between assets. Implement a subclass of this and return it in OnMakeAssetReferenceFilter */
+class IAssetReferenceFilter
+{
+public:
+	virtual ~IAssetReferenceFilter() { }
+	/** Filter function to pass/fail an asset. Called in some situations that are performance-sensitive so is expected to be fast. */
+	virtual bool PassesFilter(const FAssetData& AssetData, FText* OutOptionalFailureReason = nullptr) const = 0;
+};
+
 UCLASS(config=Engine, transient)
 class UNREALED_API UUnrealEdEngine : public UEditorEngine, public FNotifyHook
 {
@@ -815,8 +831,17 @@ public:
 	void OnMatineeEditorClosed( class FEdMode* Mode, bool IsEntering );
 
 	bool IsComponentSelected(const UPrimitiveComponent* PrimComponent);
+
+	/** Returns a filter to restruct what assets show up in asset pickers based on what the selection is used for (i.e. what will reference the assets) */
+	DECLARE_DELEGATE_RetVal_OneParam(TSharedPtr<IAssetReferenceFilter>, FOnMakeAssetReferenceFilter, const FAssetReferenceFilterContext& /*Context*/);
+	FOnMakeAssetReferenceFilter& OnMakeAssetReferenceFilter() { return OnMakeAssetReferenceFilterDelegate; }
+	TSharedPtr<IAssetReferenceFilter> MakeAssetReferenceFilter(const FAssetReferenceFilterContext& Context) { return OnMakeAssetReferenceFilterDelegate.IsBound() ? OnMakeAssetReferenceFilterDelegate.Execute(Context) : nullptr; }
 	
 protected:
+
+	/** Returns a filter to restruct what assets show up in asset pickers based on what the selection is used for (i.e. what will reference the assets) */
+	FOnMakeAssetReferenceFilter OnMakeAssetReferenceFilterDelegate;
+
 	/** Called when global editor selection changes */
 	void OnEditorSelectionChanged(UObject* SelectionThatChanged);
 
