Add proper Reset to default transactions for custom reset to default

diff --git a/Engine/Source/Editor/PropertyEditor/Private/PropertyHandleImpl.cpp b/Engine/Source/Editor/PropertyEditor/Private/PropertyHandleImpl.cpp
--- a/Engine/Source/Editor/PropertyEditor/Private/PropertyHandleImpl.cpp
+++ b/Engine/Source/Editor/PropertyEditor/Private/PropertyHandleImpl.cpp
@@ -2489,6 +2489,7 @@ void FPropertyHandleBase::OnCustomResetToDefault(const FResetToDefaultOverride&
 {
 	if (OnCustomResetToDefault.OnResetToDefaultClicked().IsBound())
 	{
+		FScopedTransaction Transaction(LOCTEXT("PropertyCustomResetToDefault", "Custom Reset to Default"));
 		if (Implementation->GetPropertyUtilities().IsValid() && Implementation->GetPropertyUtilities()->GetNotifyHook() != nullptr)
 		{
 			Implementation->GetPropertyNode()->NotifyPreChange(Implementation->GetPropertyNode()->GetProperty(), Implementation->GetPropertyUtilities()->GetNotifyHook());
