Added Error log message when trying to use an InlineEditConditionToggle with a property that is not marked as being editable in editor. However, allowing them for now since this was the behaviour prior to 4.23.

diff --git a/Engine/Source/Editor/PropertyEditor/Private/Presentation/PropertyEditor/PropertyEditor.cpp b/Engine/Source/Editor/PropertyEditor/Private/Presentation/PropertyEditor/PropertyEditor.cpp
--- a/Engine/Source/Editor/PropertyEditor/Private/Presentation/PropertyEditor/PropertyEditor.cpp
+++ b/Engine/Source/Editor/PropertyEditor/Private/Presentation/PropertyEditor/PropertyEditor.cpp
@@ -22,6 +22,8 @@
 
 #define LOCTEXT_NAMESPACE "PropertyEditor"
 
+DEFINE_LOG_CATEGORY_STATIC(LogPropertyEditor, Log, All);
+
 const FString FPropertyEditor::MultipleValuesDisplayName = NSLOCTEXT("PropertyEditor", "MultipleValues", "Multiple Values").ToString();
 
 TSharedRef< FPropertyEditor > FPropertyEditor::Create( const TSharedRef< class FPropertyNode >& InPropertyNode, const TSharedRef<class IPropertyUtilities >& InPropertyUtilities )
@@ -599,11 +601,16 @@ bool FPropertyEditor::SupportsEditConditionToggle() const
 		const UBoolProperty* ConditionalProperty = EditConditionContext->GetSingleBoolProperty(EditConditionExpression);
 		if (ConditionalProperty != nullptr)
 		{
-			// If it's editable, then only put an inline toggle box if the metadata specifies it
 			static const FName Name_InlineEditConditionToggle("InlineEditConditionToggle");
-			if (ConditionalProperty->HasAllPropertyFlags(CPF_Edit) && 
-				ConditionalProperty->HasMetaData(Name_InlineEditConditionToggle))
+			if (ConditionalProperty->HasMetaData(Name_InlineEditConditionToggle))
 			{
+				// If the edit condition property is not marked as editable, it's technically a bug.
+				// However, this was the behaviour prior to 4.23, so just warn and allow it for now.
+				if (!ConditionalProperty->HasAllPropertyFlags(CPF_Edit))
+				{
+					UE_LOG(LogPropertyEditor, Error, TEXT("Property being used as InlineEditConditionToggle is not marked as editable: Field \"%s\" in class \"%s\"."), *ConditionalProperty->GetNameCPP(), *Property->GetOwnerStruct()->GetName());
+				}
+
 				return true;
 			}
 		}
