PropertyNode no longer stores the internal index for sets and maps, instead storing the logical index.

This was causing a crash when deleting multiple entries from a map, because the physical index overruns the size of a copy of the map (eg. in an object created from a CDO), because that does not copy the gaps in the map.

Source: 6968404 (Dev-Editor)

diff --git a/Engine/Source/Editor/PropertyEditor/Private/ItemPropertyNode.cpp b/Engine/Source/Editor/PropertyEditor/Private/ItemPropertyNode.cpp
--- a/Engine/Source/Editor/PropertyEditor/Private/ItemPropertyNode.cpp
+++ b/Engine/Source/Editor/PropertyEditor/Private/ItemPropertyNode.cpp
@@ -51,9 +51,10 @@ uint8* FItemPropertyNode::GetValueBaseAddress( uint8* StartAddress )
 		else if (OuterSetProp != NULL)
 		{
 			FScriptSetHelper SetHelper(OuterSetProp, ValueBaseAddress);
-			if (ValueBaseAddress != NULL && SetHelper.IsValidIndex(ArrayIndex))
+			if (ValueBaseAddress != NULL && SetHelper.Num())
 			{
-				return SetHelper.GetElementPtr(ArrayIndex);
+				int32 ActualIndex = SetHelper.FindInternalIndex(ArrayIndex);
+				return SetHelper.GetElementPtr(ActualIndex);
 			}
 
 			return NULL;
@@ -61,10 +62,10 @@ uint8* FItemPropertyNode::GetValueBaseAddress( uint8* StartAddress )
 		else if (OuterMapProp != NULL)
 		{
 			FScriptMapHelper MapHelper(OuterMapProp, ValueBaseAddress);
-			if (ValueBaseAddress != NULL && MapHelper.IsValidIndex(ArrayIndex))
+			if (ValueBaseAddress != NULL && MapHelper.Num() > 0)
 			{
-				uint8* PairPtr = MapHelper.GetPairPtr(ArrayIndex);
-
+				int32 ActualIndex = MapHelper.FindInternalIndex(ArrayIndex);
+				uint8* PairPtr = MapHelper.GetPairPtr(ActualIndex);
 				return MyProperty->ContainerPtrToValuePtr<uint8>(PairPtr);
 			}
 
@@ -209,7 +210,7 @@ void FItemPropertyNode::InitChildNodes()
 				FPropertyNodeInitParams InitParams;
 				InitParams.ParentNode = SharedThis(this);
 				InitParams.Property = ArrayProperty->Inner;
-				InitParams.ArrayOffset = Index*ArrayProperty->Inner->ElementSize;
+				InitParams.ArrayOffset = Index * ArrayProperty->Inner->ElementSize;
 				InitParams.ArrayIndex = Index;
 				InitParams.bAllowChildren = true;
 				InitParams.bForceHiddenPropertyVisibility = bShouldShowHiddenProperties;
@@ -233,25 +234,20 @@ void FItemPropertyNode::InitChildNodes()
 		{
 			FScriptSetHelper SetHelper(SetProperty, Set);
 
-			for (int32 TrueIndex = 0, ItemsLeft = SetHelper.Num(); ItemsLeft > 0; ++TrueIndex)
+			for (int32 Index = 0; Index < SetHelper.Num(); ++Index)
 			{
-				if (SetHelper.IsValidIndex(TrueIndex))
-				{
-					--ItemsLeft;
-
-					TSharedPtr<FItemPropertyNode> NewItemNode(new FItemPropertyNode);
+				TSharedPtr<FItemPropertyNode> NewItemNode(new FItemPropertyNode);
 
-					FPropertyNodeInitParams InitParams;
-					InitParams.ParentNode = SharedThis(this);
-					InitParams.Property = SetProperty->ElementProp;
-					InitParams.ArrayIndex = TrueIndex;
-					InitParams.bAllowChildren = true;
-					InitParams.bForceHiddenPropertyVisibility = bShouldShowHiddenProperties;
-					InitParams.bCreateDisableEditOnInstanceNodes = bShouldShowDisableEditOnInstance;
+				FPropertyNodeInitParams InitParams;
+				InitParams.ParentNode = SharedThis(this);
+				InitParams.Property = SetProperty->ElementProp;
+				InitParams.ArrayIndex = Index;
+				InitParams.bAllowChildren = true;
+				InitParams.bForceHiddenPropertyVisibility = bShouldShowHiddenProperties;
+				InitParams.bCreateDisableEditOnInstanceNodes = bShouldShowDisableEditOnInstance;
 
-					NewItemNode->InitNode(InitParams);
-					AddChildNode(NewItemNode);
-				}
+				NewItemNode->InitNode(InitParams);
+				AddChildNode(NewItemNode);
 			}
 		}
 	}
@@ -268,39 +264,34 @@ void FItemPropertyNode::InitChildNodes()
 		{
 			FScriptMapHelper MapHelper(MapProperty, Map);
 
-			for (int32 TrueIndex = 0, ItemsLeft = MapHelper.Num(); ItemsLeft > 0; ++TrueIndex)
+			for (int32 Index = 0; Index < MapHelper.Num(); ++Index)
 			{
-				if ( MapHelper.IsValidIndex(TrueIndex) )
-				{
-					--ItemsLeft;
-
-					// Construct the key node first
-					TSharedPtr<FPropertyNode> KeyNode(new FItemPropertyNode());
+				// Construct the key node first
+				TSharedPtr<FPropertyNode> KeyNode(new FItemPropertyNode());
 					
-					FPropertyNodeInitParams InitParams;
-					InitParams.ParentNode = SharedThis(this);	// Key Node needs to point to this node to ensure its data is set correctly
-					InitParams.Property = MapHelper.KeyProp;
-					InitParams.ArrayIndex = TrueIndex;
-					InitParams.ArrayOffset = 0;
-					InitParams.bAllowChildren = true;
-					InitParams.bForceHiddenPropertyVisibility = bShouldShowHiddenProperties;
-					InitParams.bCreateDisableEditOnInstanceNodes = bShouldShowDisableEditOnInstance;
-
-					KeyNode->InitNode(InitParams);
-					// Not adding the key node as a child, otherwise it'll show up in the wrong spot.
-
-					// Then the value node
-					TSharedPtr<FPropertyNode> ValueNode(new FItemPropertyNode());
+				FPropertyNodeInitParams InitParams;
+				InitParams.ParentNode = SharedThis(this);	// Key Node needs to point to this node to ensure its data is set correctly
+				InitParams.Property = MapHelper.KeyProp;
+				InitParams.ArrayIndex = Index;
+				InitParams.ArrayOffset = 0;
+				InitParams.bAllowChildren = true;
+				InitParams.bForceHiddenPropertyVisibility = bShouldShowHiddenProperties;
+				InitParams.bCreateDisableEditOnInstanceNodes = bShouldShowDisableEditOnInstance;
+
+				KeyNode->InitNode(InitParams);
+				// Not adding the key node as a child, otherwise it'll show up in the wrong spot.
+
+				// Then the value node
+				TSharedPtr<FPropertyNode> ValueNode(new FItemPropertyNode());
 					
-					// Reuse the init params
-					InitParams.ParentNode = SharedThis(this);
-					InitParams.Property = MapHelper.ValueProp;
+				// Reuse the init params
+				InitParams.ParentNode = SharedThis(this);
+				InitParams.Property = MapHelper.ValueProp;
 
-					ValueNode->InitNode(InitParams);
-					AddChildNode(ValueNode);
+				ValueNode->InitNode(InitParams);
+				AddChildNode(ValueNode);
 
-					FPropertyNode::SetupKeyValueNodePair(KeyNode, ValueNode);
-				}
+				FPropertyNode::SetupKeyValueNodePair(KeyNode, ValueNode);
 			}
 		}
 	}
diff --git a/Engine/Source/Editor/PropertyEditor/Private/PropertyChangeListener.cpp b/Engine/Source/Editor/PropertyEditor/Private/PropertyChangeListener.cpp
--- a/Engine/Source/Editor/PropertyEditor/Private/PropertyChangeListener.cpp
+++ b/Engine/Source/Editor/PropertyEditor/Private/PropertyChangeListener.cpp
@@ -131,7 +131,8 @@ public:
 			{
 				FScriptSetHelper SetHelper(OuterSetProperty, PropertyValueAddresses.BaseAddress);
 				
-				if ( !SetHelper.IsValidIndex(PropertyNodeRef.GetArrayIndex()) )
+				bool bIsValidIndex = PropertyNodeRef.GetArrayIndex() >= 0 && PropertyNodeRef.GetArrayIndex() < SetHelper.Num();
+				if (!bIsValidIndex)
 				{
 					bPropertyValid = false;
 				}
@@ -140,7 +141,8 @@ public:
 			{
 				FScriptMapHelper MapHelper(OuterMapProperty, PropertyValueAddresses.BaseAddress);
 
-				if ( !MapHelper.IsValidIndex(PropertyNodeRef.GetArrayIndex()) )
+				bool bIsValidIndex = PropertyNodeRef.GetArrayIndex() >= 0 && PropertyNodeRef.GetArrayIndex() < MapHelper.Num();
+				if (!bIsValidIndex)
 				{
 					bPropertyValid = false;
 				}
diff --git a/Engine/Source/Editor/PropertyEditor/Private/PropertyHandleImpl.cpp b/Engine/Source/Editor/PropertyEditor/Private/PropertyHandleImpl.cpp
--- a/Engine/Source/Editor/PropertyEditor/Private/PropertyHandleImpl.cpp
+++ b/Engine/Source/Editor/PropertyEditor/Private/PropertyHandleImpl.cpp
@@ -1517,7 +1517,8 @@ void FPropertyValueImpl::DeleteChild( TSharedPtr<FPropertyNode> ChildNodeToDelet
 						}
 					}
 
-					SetHelper.RemoveAt(ChildNodePtr->GetArrayIndex());
+					int32 InternalIndex = SetHelper.FindInternalIndex(ChildNodePtr->GetArrayIndex());
+					SetHelper.RemoveAt(InternalIndex);
 					SetHelper.Rehash();
 				}
 				else if (MapProperty)
@@ -1536,7 +1537,8 @@ void FPropertyValueImpl::DeleteChild( TSharedPtr<FPropertyNode> ChildNodeToDelet
 						}
 					}
 
-					MapHelper.RemoveAt(ChildNodePtr->GetArrayIndex());
+					int32 InternalIndex = MapHelper.FindInternalIndex(ChildNodePtr->GetArrayIndex());
+					MapHelper.RemoveAt(InternalIndex);
 					MapHelper.Rehash();
 				}
 
diff --git a/Engine/Source/Editor/PropertyEditor/Private/PropertyNode.cpp b/Engine/Source/Editor/PropertyEditor/Private/PropertyNode.cpp
--- a/Engine/Source/Editor/PropertyEditor/Private/PropertyNode.cpp
+++ b/Engine/Source/Editor/PropertyEditor/Private/PropertyNode.cpp
@@ -1576,7 +1576,8 @@ bool FPropertyNode::GetDiffersFromDefaultForObject( FPropertyItemValueDataTracke
 		{
 			FScriptSetHelper SetHelper(OuterSetProperty, ValueTracker.GetPropertyDefaultBaseAddress());
 
-			if ( ValueTracker.GetPropertyDefaultBaseAddress() != NULL && !SetHelper.IsValidIndex(GetArrayIndex()) )
+			bool bIsValidIndex = ArrayIndex >= 0 && ArrayIndex < SetHelper.Num();
+			if ( ValueTracker.GetPropertyDefaultBaseAddress() != NULL && !bIsValidIndex)
 			{
 				bDiffersFromDefaultForObject = true;
 			}
@@ -1585,7 +1586,8 @@ bool FPropertyNode::GetDiffersFromDefaultForObject( FPropertyItemValueDataTracke
 		{
 			FScriptMapHelper MapHelper(OuterMapProperty, ValueTracker.GetPropertyDefaultBaseAddress());
 
-			if ( ValueTracker.GetPropertyDefaultBaseAddress() != NULL && !MapHelper.IsValidIndex(GetArrayIndex()) )
+			bool bIsValidIndex = ArrayIndex >= 0 && ArrayIndex < MapHelper.Num();
+			if ( ValueTracker.GetPropertyDefaultBaseAddress() != NULL && !bIsValidIndex )
 			{
 				bDiffersFromDefaultForObject = true;
 			}
@@ -2572,7 +2574,7 @@ void FPropertyNode::PropagateContainerPropertyChange( UObject* ModifiedObject, c
 							check(false);	// Insert is not supported for sets
 							break;
 						case EPropertyArrayChangeType::Delete:
-							SetHelper.RemoveAt(ArrayIndex);
+							SetHelper.RemoveAt(SetHelper.FindInternalIndex(ArrayIndex));
 							SetHelper.Rehash();
 							break;
 						case EPropertyArrayChangeType::Duplicate:
@@ -2602,7 +2604,7 @@ void FPropertyNode::PropagateContainerPropertyChange( UObject* ModifiedObject, c
 							check(false);	// Insert is not supported for maps
 							break;
 						case EPropertyArrayChangeType::Delete:
-							MapHelper.RemoveAt(ArrayIndex);
+							MapHelper.RemoveAt(MapHelper.FindInternalIndex(ArrayIndex));
 							MapHelper.Rehash();
 							break;
 						case EPropertyArrayChangeType::Duplicate:
