Fix detail layout customization visibility for the case of a custom builder.

diff --git a/Engine/Source/Editor/PropertyEditor/Private/DetailCategoryBuilderImpl.cpp b/Engine/Source/Editor/PropertyEditor/Private/DetailCategoryBuilderImpl.cpp
--- a/Engine/Source/Editor/PropertyEditor/Private/DetailCategoryBuilderImpl.cpp
+++ b/Engine/Source/Editor/PropertyEditor/Private/DetailCategoryBuilderImpl.cpp
@@ -68,6 +68,14 @@ bool FDetailLayoutCustomization::HasExternalPropertyRow() const
 	return HasPropertyNode() && PropertyRow->HasExternalProperty();
 }
 
+bool FDetailLayoutCustomization::IsHidden() const
+{
+	return !IsValidCustomization()
+		|| (HasCustomWidget() && WidgetDecl->VisibilityAttr.Get() != EVisibility::Visible)
+		|| (HasPropertyNode() && PropertyRow->GetPropertyVisibility() != EVisibility::Visible)
+		|| (HasCustomBuilder() && CustomBuilderRow->AreChildCustomizationsHidden());
+}
+
 TSharedPtr<FPropertyNode> FDetailLayoutCustomization::GetPropertyNode() const
 {
 	return PropertyRow.IsValid() ? PropertyRow->GetPropertyNode() : nullptr;
diff --git a/Engine/Source/Editor/PropertyEditor/Private/DetailCategoryBuilderImpl.h b/Engine/Source/Editor/PropertyEditor/Private/DetailCategoryBuilderImpl.h
--- a/Engine/Source/Editor/PropertyEditor/Private/DetailCategoryBuilderImpl.h
+++ b/Engine/Source/Editor/PropertyEditor/Private/DetailCategoryBuilderImpl.h
@@ -41,15 +41,16 @@ struct FDetailLayoutCustomization
 	bool HasPropertyNode() const { return GetPropertyNode().IsValid(); }
 	/** @return true if this customization has a custom widget */
 	bool HasCustomWidget() const { return WidgetDecl.IsValid(); }
-	/** @return true if this customization has a custom builder (custom builders will set the custom widget */
+	/** @return true if this customization has a custom builder (custom builders will set the custom widget) */
 	bool HasCustomBuilder() const { return CustomBuilderRow.IsValid(); }
 	/** @return true if this customization has a group */
 	bool HasGroup() const { return DetailGroup.IsValid(); }
 	/** @return true if this has a customization for an external property row */
 	bool HasExternalPropertyRow() const;
-
 	/** @return true if this customization is valid */
 	bool IsValidCustomization() const { return HasPropertyNode() || HasCustomWidget() || HasCustomBuilder() || HasGroup(); }
+	/** @return true if the customized item is hidden */
+	bool IsHidden() const;
 	/** @return the property node for this customization (if any ) */
 	TSharedPtr<FPropertyNode> GetPropertyNode() const;
 	/** @return The row to display from this customization */
diff --git a/Engine/Source/Editor/PropertyEditor/Private/DetailCustomBuilderRow.cpp b/Engine/Source/Editor/PropertyEditor/Private/DetailCustomBuilderRow.cpp
--- a/Engine/Source/Editor/PropertyEditor/Private/DetailCustomBuilderRow.cpp
+++ b/Engine/Source/Editor/PropertyEditor/Private/DetailCustomBuilderRow.cpp
@@ -82,3 +82,20 @@ FDetailWidgetRow FDetailCustomBuilderRow::GetWidgetRow()
 {
 	return *HeaderRow;
 }
+
+bool FDetailCustomBuilderRow::AreChildCustomizationsHidden() const
+{
+	bool bChildCustomizationsHidden = true;
+	if (ChildrenBuilder)
+	{
+		for (const FDetailLayoutCustomization& ChildCustomizations : ChildrenBuilder->GetChildCustomizations())
+		{
+			if (!ChildCustomizations.IsHidden())
+			{
+				bChildCustomizationsHidden = false;
+				break;
+			}
+		}
+	}
+	return bChildCustomizationsHidden;
+}
\ No newline at end of file
diff --git a/Engine/Source/Editor/PropertyEditor/Private/DetailCustomBuilderRow.h b/Engine/Source/Editor/PropertyEditor/Private/DetailCustomBuilderRow.h
--- a/Engine/Source/Editor/PropertyEditor/Private/DetailCustomBuilderRow.h
+++ b/Engine/Source/Editor/PropertyEditor/Private/DetailCustomBuilderRow.h
@@ -28,7 +28,9 @@ public:
 	void OnGenerateChildren( FDetailNodeList& OutChildren );
 	bool IsInitiallyCollapsed() const;
 	FDetailWidgetRow GetWidgetRow();
-private:	
+	bool AreChildCustomizationsHidden() const;
+
+private:
 	/** Whether or not our parent is enabled */
 	TAttribute<bool> IsParentEnabled;
 	TSharedPtr<FDetailWidgetRow> HeaderRow;
diff --git a/Engine/Source/Editor/PropertyEditor/Private/DetailItemNode.cpp b/Engine/Source/Editor/PropertyEditor/Private/DetailItemNode.cpp
--- a/Engine/Source/Editor/PropertyEditor/Private/DetailItemNode.cpp
+++ b/Engine/Source/Editor/PropertyEditor/Private/DetailItemNode.cpp
@@ -406,15 +406,8 @@ bool FDetailItemNode::ShouldBeExpanded() const
 
 ENodeVisibility FDetailItemNode::GetVisibility() const
 {
-	const bool bHasAnythingToShow = Customization.IsValidCustomization();
-
-	const bool bIsForcedHidden = 
-		!bHasAnythingToShow 
-		|| (Customization.HasCustomWidget() && Customization.WidgetDecl->VisibilityAttr.Get() != EVisibility::Visible )
-		|| (Customization.HasPropertyNode() && Customization.PropertyRow->GetPropertyVisibility() != EVisibility::Visible );
-
 	ENodeVisibility Visibility;
-	if( bIsForcedHidden )
+	if(Customization.IsHidden())
 	{
 		Visibility = ENodeVisibility::ForcedHidden;
 	}
@@ -422,7 +415,6 @@ ENodeVisibility FDetailItemNode::GetVisibility() const
 	{
 		Visibility = (bShouldBeVisibleDueToFiltering || bShouldBeVisibleDueToChildFiltering) ? ENodeVisibility::Visible : ENodeVisibility::HiddenDueToFiltering;
 	}
-
 	return Visibility;
 }
 
