CL 4356067 copied to Release-4.21

diff --git a/Engine/Source/Editor/DetailCustomizations/Private/SoftObjectPathCustomization.cpp b/Engine/Source/Editor/DetailCustomizations/Private/SoftObjectPathCustomization.cpp
--- a/Engine/Source/Editor/DetailCustomizations/Private/SoftObjectPathCustomization.cpp
+++ b/Engine/Source/Editor/DetailCustomizations/Private/SoftObjectPathCustomization.cpp
@@ -11,24 +11,46 @@ void FSoftObjectPathCustomization::CustomizeHeader( TSharedRef<IPropertyHandle>
 {
 	StructPropertyHandle = InStructPropertyHandle;
 
-	const FString& ClassFilterString = StructPropertyHandle->GetMetaData("AllowedClasses");
-	if( !ClassFilterString.IsEmpty() )
+	const FString& AllowedClassFilterString = StructPropertyHandle->GetMetaData("AllowedClasses");
+	if (!AllowedClassFilterString.IsEmpty())
 	{
-		TArray<FString> CustomClassFilterNames;
-		ClassFilterString.ParseIntoArray(CustomClassFilterNames, TEXT(","), true);
+		TArray<FString> AllowedClassFilterNames;
+		AllowedClassFilterString.ParseIntoArray(AllowedClassFilterNames, TEXT(","), true);
 
-		for(auto It = CustomClassFilterNames.CreateConstIterator(); It; ++It)
+		for(auto It = AllowedClassFilterNames.CreateConstIterator(); It; ++It)
 		{
 			const FString& ClassName = *It;
 
 			UClass* Class = FindObject<UClass>(ANY_PACKAGE, *ClassName);
-			if(!Class)
+			if (!Class)
 			{
 				Class = LoadObject<UClass>(nullptr, *ClassName);
 			}
-			if(Class)
+			if (Class)
 			{
-				CustomClassFilters.Add(Class);
+				AllowedClassFilters.Add(Class);
+			}
+		}
+	}
+
+	const FString& DisallowedClassFilterString = StructPropertyHandle->GetMetaData("DisallowedClasses");
+	if (!DisallowedClassFilterString.IsEmpty())
+	{
+		TArray<FString> DisallowedClassFilterNames;
+		DisallowedClassFilterString.ParseIntoArray(DisallowedClassFilterNames, TEXT(","), true);
+
+		for (auto It = DisallowedClassFilterNames.CreateConstIterator(); It; ++It)
+		{
+			const FString& ClassName = *It;
+
+			UClass* Class = FindObject<UClass>(ANY_PACKAGE, *ClassName);
+			if (!Class)
+			{
+				Class = LoadObject<UClass>(nullptr, *ClassName);
+			}
+			if (Class)
+			{
+				DisallowedClassFilters.Add(Class);
 			}
 		}
 	}
@@ -37,15 +59,15 @@ void FSoftObjectPathCustomization::CustomizeHeader( TSharedRef<IPropertyHandle>
 
 	FOnShouldFilterAsset AssetFilter;
 	UClass* ClassFilter = UObject::StaticClass();
-	if(CustomClassFilters.Num() == 1 && !bExactClass)
+	if (AllowedClassFilters.Num() == 1 && DisallowedClassFilters.Num() <= 0 && !bExactClass)
 	{
 		// If we only have one class to filter on, set it as the class type filter rather than use a filter callback
 		// We can only do this if we don't need an exact match, as the class filter also allows derived types
 		// The class filter is much faster than the callback as we're not performing two different sets of type tests
 		// (one against UObject, one against the actual type)
-		ClassFilter = CustomClassFilters[0];
+		ClassFilter = AllowedClassFilters[0];
 	}
-	else if(CustomClassFilters.Num() > 0)
+	else if (AllowedClassFilters.Num() > 0 || DisallowedClassFilters.Num() > 0)
 	{
 		// Only bind the filter if we have classes that need filtering
 		AssetFilter.BindSP(this, &FSoftObjectPathCustomization::OnShouldFilterAsset);
@@ -81,11 +103,21 @@ bool FSoftObjectPathCustomization::OnShouldFilterAsset( const FAssetData& InAsse
 {
 	// Only bound if we have classes to filter on, so we don't need to test for an empty array here
 	UClass* const AssetClass = InAssetData.GetClass();
-	for(auto It = CustomClassFilters.CreateConstIterator(); It; ++It)
+
+	for (auto It = DisallowedClassFilters.CreateConstIterator(); It; ++It)
+	{
+		UClass* const DisallowClass = *It;
+		if (AssetClass->IsChildOf(DisallowClass))
+		{
+			return true;
+		}
+	}
+
+	for(auto It = AllowedClassFilters.CreateConstIterator(); It; ++It)
 	{
 		UClass* const FilterClass = *It;
 		const bool bMatchesFilter = (bExactClass) ? AssetClass == FilterClass : AssetClass->IsChildOf(FilterClass);
-		if(bMatchesFilter)
+		if (bMatchesFilter)
 		{
 			return false;
 		}
diff --git a/Engine/Source/Editor/DetailCustomizations/Private/SoftObjectPathCustomization.h b/Engine/Source/Editor/DetailCustomizations/Private/SoftObjectPathCustomization.h
--- a/Engine/Source/Editor/DetailCustomizations/Private/SoftObjectPathCustomization.h
+++ b/Engine/Source/Editor/DetailCustomizations/Private/SoftObjectPathCustomization.h
@@ -29,7 +29,9 @@ private:
 	/** Handle to the struct property being customized */
 	TSharedPtr<IPropertyHandle> StructPropertyHandle;
 	/** Classes that can be used with this property */
-	TArray<UClass*> CustomClassFilters;
+	TArray<UClass*> AllowedClassFilters;
+	/** Classes that can NOT be used with this property */
+	TArray<UClass*> DisallowedClassFilters;
 	/** Whether the classes in the above array must be an exact match, or whether they can also be a derived type; default is false */
 	bool bExactClass;
 };
diff --git a/Engine/Source/Editor/PropertyEditor/Private/PropertyCustomizationHelpers.cpp b/Engine/Source/Editor/PropertyEditor/Private/PropertyCustomizationHelpers.cpp
--- a/Engine/Source/Editor/PropertyEditor/Private/PropertyCustomizationHelpers.cpp
+++ b/Engine/Source/Editor/PropertyEditor/Private/PropertyCustomizationHelpers.cpp
@@ -263,12 +263,24 @@ namespace PropertyCustomizationHelpers
 			.OnAssetSelected( OnAssetSelectedFromPicker );
 	}
 
+	TArray<const UClass*> EmptyClassArray;
+
 	TSharedRef<SWidget> MakeAssetPickerWithMenu(const FAssetData& InitialObject, const bool AllowClear, const TArray<const UClass*>& AllowedClasses, const TArray<UFactory*>& NewAssetFactories, FOnShouldFilterAsset OnShouldFilterAsset, FOnAssetSelected OnSet, FSimpleDelegate OnClose)
 	{
-		return MakeAssetPickerWithMenu(InitialObject, AllowClear, false, AllowedClasses, NewAssetFactories, OnShouldFilterAsset, OnSet, OnClose);
+		return MakeAssetPickerWithMenu(InitialObject, AllowClear, false, AllowedClasses, EmptyClassArray, NewAssetFactories, OnShouldFilterAsset, OnSet, OnClose);
+	}
+
+	TSharedRef<SWidget> MakeAssetPickerWithMenu(const FAssetData& InitialObject, const bool AllowClear, const TArray<const UClass*>& AllowedClasses, const TArray<const UClass*>& DisallowedClasses, const TArray<UFactory*>& NewAssetFactories, FOnShouldFilterAsset OnShouldFilterAsset, FOnAssetSelected OnSet, FSimpleDelegate OnClose)
+	{
+		return MakeAssetPickerWithMenu(InitialObject, AllowClear, false, AllowedClasses, DisallowedClasses, NewAssetFactories, OnShouldFilterAsset, OnSet, OnClose);
+	}
+
+	TSharedRef<SWidget> MakeAssetPickerWithMenu(const FAssetData& InitialObject, const bool AllowClear, const bool AllowCopyPaste, const TArray<const UClass*>& AllowedClasses, const TArray<UFactory*>& NewAssetFactories, FOnShouldFilterAsset OnShouldFilterAsset, FOnAssetSelected OnSet, FSimpleDelegate OnClose)
+	{
+		return MakeAssetPickerWithMenu(InitialObject, AllowClear, AllowCopyPaste, AllowedClasses, EmptyClassArray, NewAssetFactories, OnShouldFilterAsset, OnSet, OnClose);
 	}
 
-	TSharedRef<SWidget> MakeAssetPickerWithMenu( const FAssetData& InitialObject, const bool AllowClear, const bool AllowCopyPaste, const TArray<const UClass*>& AllowedClasses, const TArray<UFactory*>& NewAssetFactories, FOnShouldFilterAsset OnShouldFilterAsset, FOnAssetSelected OnSet, FSimpleDelegate OnClose)
+	TSharedRef<SWidget> MakeAssetPickerWithMenu( const FAssetData& InitialObject, const bool AllowClear, const bool AllowCopyPaste, const TArray<const UClass*>& AllowedClasses, const TArray<const UClass*>& DisallowedClasses, const TArray<UFactory*>& NewAssetFactories, FOnShouldFilterAsset OnShouldFilterAsset, FOnAssetSelected OnSet, FSimpleDelegate OnClose)
 	{
 		return
 			SNew(SPropertyMenuAssetPicker)
@@ -276,6 +288,7 @@ namespace PropertyCustomizationHelpers
 			.AllowClear(AllowClear)
 			.AllowCopyPaste(AllowCopyPaste)
 			.AllowedClasses(AllowedClasses)
+			.DisallowedClasses(DisallowedClasses)
 			.NewAssetFactories(NewAssetFactories)
 			.OnShouldFilterAsset(OnShouldFilterAsset)
 			.OnSet(OnSet)
@@ -383,6 +396,11 @@ namespace PropertyCustomizationHelpers
 	}
 
 	TArray<UFactory*> GetNewAssetFactoriesForClasses(const TArray<const UClass*>& Classes)
+	{
+		return GetNewAssetFactoriesForClasses(Classes, EmptyClassArray);
+	}
+
+	TArray<UFactory*> GetNewAssetFactoriesForClasses(const TArray<const UClass*>& Classes, const TArray<const UClass*>& DisallowedClasses)
 	{
 		TArray<UFactory*> Factories;
 		for (TObjectIterator<UClass> It; It; ++It)
@@ -394,7 +412,9 @@ namespace PropertyCustomizationHelpers
 				if (Factory->ShouldShowInNewMenu() && ensure(!Factory->GetDisplayName().IsEmpty()))
 				{
 					UClass* SupportedClass = Factory->GetSupportedClass();
-					if (SupportedClass != nullptr && Classes.ContainsByPredicate([=](const UClass* InClass) { return SupportedClass->IsChildOf(InClass); }))
+					if (SupportedClass != nullptr 
+						&& Classes.ContainsByPredicate([=](const UClass* InClass) { return SupportedClass->IsChildOf(InClass); })
+						&& !DisallowedClasses.ContainsByPredicate([=](const UClass* InClass) { return SupportedClass->IsChildOf(InClass); }))
 					{
 						Factories.Add(Factory);
 					}
diff --git a/Engine/Source/Editor/PropertyEditor/Private/PropertyHandleImpl.cpp b/Engine/Source/Editor/PropertyEditor/Private/PropertyHandleImpl.cpp
--- a/Engine/Source/Editor/PropertyEditor/Private/PropertyHandleImpl.cpp
+++ b/Engine/Source/Editor/PropertyEditor/Private/PropertyHandleImpl.cpp
@@ -3718,6 +3718,25 @@ FPropertyAccess::Result FPropertyHandleObject::SetValueFromFormattedString(const
 				bSupportedObject = true;
 			}
 
+			if (bSupportedObject)
+			{
+				const FString& DisallowedClassesString = NodeProperty->GetMetaData("DisallowedClasses");
+				TArray<FString> DisallowedClassNames;
+				DisallowedClassesString.ParseIntoArray(DisallowedClassNames, TEXT(","), true);
+
+				for (const FString& DisallowedClassName : DisallowedClassNames)
+				{
+					UClass* DisallowedClass = FindObject<UClass>(ANY_PACKAGE, *DisallowedClassName);
+					const bool bIsInterface = DisallowedClass && DisallowedClass->HasAnyClassFlags(CLASS_Interface);
+
+					if ((DisallowedClass && QualifiedObject->IsA(DisallowedClass)) || (bIsInterface && QualifiedObject->GetClass()->ImplementsInterface(DisallowedClass)))
+					{
+						bSupportedObject = false;
+						break;
+					}
+				}
+			}
+
 			// Check required class
 			if (RequiredClass && !QualifiedClass->IsChildOf(RequiredClass))
 			{
diff --git a/Engine/Source/Editor/PropertyEditor/Private/UserInterface/PropertyEditor/SPropertyEditorAsset.cpp b/Engine/Source/Editor/PropertyEditor/Private/UserInterface/PropertyEditor/SPropertyEditorAsset.cpp
--- a/Engine/Source/Editor/PropertyEditor/Private/UserInterface/PropertyEditor/SPropertyEditorAsset.cpp
+++ b/Engine/Source/Editor/PropertyEditor/Private/UserInterface/PropertyEditor/SPropertyEditorAsset.cpp
@@ -99,33 +99,33 @@ void SPropertyEditorAsset::Construct( const FArguments& InArgs, const TSharedPtr
 		}
 		else
 		{
-			CustomClassFilters.Add(ObjectClass);
+			AllowedClassFilters.Add(ObjectClass);
 		}
 	}
 
 	// Account for the allowed classes specified in the property metadata
 	if (Property)
 	{
-		const FString* ClassFilterString;
+		const FString* AllowedClassesFilterString;
 		if (UArrayProperty* ArrayParent = Cast<UArrayProperty>(Property->GetOuter()))
 		{
-			ClassFilterString = &ArrayParent->GetMetaData("AllowedClasses");
+			AllowedClassesFilterString = &ArrayParent->GetMetaData("AllowedClasses");
 		}
 		else
 		{
-			ClassFilterString = &Property->GetMetaData("AllowedClasses");
+			AllowedClassesFilterString = &Property->GetMetaData("AllowedClasses");
 		}
 
-		if (ClassFilterString->IsEmpty())
+		if (AllowedClassesFilterString->IsEmpty())
 		{
-			CustomClassFilters.Add(ObjectClass);
+			AllowedClassFilters.Add(ObjectClass);
 		}
 		else
 		{
-			TArray<FString> CustomClassFilterNames;
-			ClassFilterString->ParseIntoArray(CustomClassFilterNames, TEXT(","), true);
+			TArray<FString> AllowedClassFilterNames;
+			AllowedClassesFilterString->ParseIntoArray(AllowedClassFilterNames, TEXT(","), true);
 
-			for (auto It = CustomClassFilterNames.CreateIterator(); It; ++It)
+			for (auto It = AllowedClassFilterNames.CreateIterator(); It; ++It)
 			{
 				FString& ClassName = *It;
 				// User can potentially list class names with leading or trailing whitespace
@@ -148,13 +148,63 @@ void SPropertyEditorAsset::Construct( const FArguments& InArgs, const TSharedPtr
 							UClass* const ClassWithInterface = (*ClassIt);
 							if (ClassWithInterface->ImplementsInterface(Class))
 							{
-								CustomClassFilters.Add(ClassWithInterface);
+								AllowedClassFilters.Add(ClassWithInterface);
 							}
 						}
 					}
 					else
 					{
-						CustomClassFilters.Add(Class);
+						AllowedClassFilters.Add(Class);
+					}
+				}
+			}
+		}
+
+		const FString* DisallowedClassesFilterString;
+		if (UArrayProperty* ArrayParent = Cast<UArrayProperty>(Property->GetOuter()))
+		{
+			DisallowedClassesFilterString = &ArrayParent->GetMetaData("DisallowedClasses");
+		}
+		else
+		{
+			DisallowedClassesFilterString = &Property->GetMetaData("DisallowedClasses");
+		}
+
+		if (!DisallowedClassesFilterString->IsEmpty())
+		{
+			TArray<FString> DisallowedClassFilterNames;
+			DisallowedClassesFilterString->ParseIntoArray(DisallowedClassFilterNames, TEXT(","), true);
+
+			for (auto It = DisallowedClassFilterNames.CreateIterator(); It; ++It)
+			{
+				FString& ClassName = *It;
+				// User can potentially list class names with leading or trailing whitespace
+				ClassName.TrimStartAndEndInline();
+
+				UClass* Class = FindObject<UClass>(ANY_PACKAGE, *ClassName);
+
+				if (!Class)
+				{
+					Class = LoadObject<UClass>(nullptr, *ClassName);
+				}
+
+				if (Class)
+				{
+					// If the class is an interface, expand it to be all classes in memory that implement the class.
+					if (Class->HasAnyClassFlags(CLASS_Interface))
+					{
+						for (TObjectIterator<UClass> ClassIt; ClassIt; ++ClassIt)
+						{
+							UClass* const ClassWithInterface = (*ClassIt);
+							if (ClassWithInterface->ImplementsInterface(Class))
+							{
+								DisallowedClassFilters.Add(ClassWithInterface);
+							}
+						}
+					}
+					else
+					{
+						DisallowedClassFilters.Add(Class);
 					}
 				}
 			}
@@ -165,9 +215,9 @@ void SPropertyEditorAsset::Construct( const FArguments& InArgs, const TSharedPtr
 	{
 		NewAssetFactories = InArgs._NewAssetFactories.GetValue();
 	}
-	else if (CustomClassFilters.Num() > 1 || !CustomClassFilters.Contains(UObject::StaticClass()))
+	else if (AllowedClassFilters.Num() > 1 && !AllowedClassFilters.Contains(UObject::StaticClass()))
 	{
-		NewAssetFactories = PropertyCustomizationHelpers::GetNewAssetFactoriesForClasses(CustomClassFilters);
+		NewAssetFactories = PropertyCustomizationHelpers::GetNewAssetFactoriesForClasses(AllowedClassFilters, DisallowedClassFilters);
 	}
 	
 	TSharedPtr<SHorizontalBox> ValueContentBox = nullptr;
@@ -609,7 +659,8 @@ TSharedRef<SWidget> SPropertyEditorAsset::OnGetMenuContent()
 	{
 		return PropertyCustomizationHelpers::MakeAssetPickerWithMenu(Value.AssetData,
 																	 bAllowClear,
-																	 CustomClassFilters,
+																	 AllowedClassFilters,
+																	 DisallowedClassFilters,
 																	 NewAssetFactories,
 																	 OnShouldFilterAsset,
 																	 FOnAssetSelected::CreateSP(this, &SPropertyEditorAsset::OnAssetSelected),
@@ -943,7 +994,7 @@ void SPropertyEditorAsset::OnActorSelected( AActor* InActor )
 
 void SPropertyEditorAsset::OnGetAllowedClasses(TArray<const UClass*>& AllowedClasses)
 {
-	AllowedClasses.Append(CustomClassFilters);
+	AllowedClasses.Append(AllowedClassFilters);
 }
 
 void SPropertyEditorAsset::OnOpenAssetEditor()
@@ -1008,7 +1059,7 @@ FText SPropertyEditorAsset::GetOnBrowseToolTip() const
 void SPropertyEditorAsset::OnUse()
 {
 	// Use the property editor path if it is valid and there is no custom filtering required
-	if(PropertyEditor.IsValid() && !OnShouldFilterAsset.IsBound() && CustomClassFilters.Num() == 0)
+	if(PropertyEditor.IsValid() && !OnShouldFilterAsset.IsBound() && AllowedClassFilters.Num() == 0 && DisallowedClassFilters.Num() == 0)
 	{
 		PropertyEditor->GetPropertyHandle()->SetObjectValueFromSelection();
 	}
@@ -1163,17 +1214,39 @@ bool SPropertyEditorAsset::CanEdit() const
 bool SPropertyEditorAsset::CanSetBasedOnCustomClasses( const FAssetData& InAssetData ) const
 {
 	bool bAllowedToSetBasedOnFilter = true;
-	if( InAssetData.IsValid() && CustomClassFilters.Num() > 0 )
+	if( InAssetData.IsValid() && AllowedClassFilters.Num() > 0 )
 	{
-		bAllowedToSetBasedOnFilter = false;
-		UClass* AssetClass = InAssetData.GetClass();
-		for( const UClass* AllowedClass : CustomClassFilters )
+		const int32 AllowedClassCount = AllowedClassFilters.Num();
+		const int32 DisallowedClassCount = DisallowedClassFilters.Num();
+
+		if (AllowedClassCount > 0 || DisallowedClassCount > 0)
 		{
-			const bool bAllowedClassIsInterface = AllowedClass->HasAnyClassFlags(CLASS_Interface);
-			if( AssetClass->IsChildOf( AllowedClass ) || (bAllowedClassIsInterface && AssetClass->ImplementsInterface(AllowedClass)) )
+			UClass* AssetClass = InAssetData.GetClass();
+			if (AllowedClassCount > 0)
 			{
-				bAllowedToSetBasedOnFilter = true;
-				break;
+				bAllowedToSetBasedOnFilter = false; 
+				for (const UClass* AllowedClass : AllowedClassFilters)
+				{
+					const bool bAllowedClassIsInterface = AllowedClass->HasAnyClassFlags(CLASS_Interface);
+					if (AssetClass->IsChildOf(AllowedClass) || (bAllowedClassIsInterface && AssetClass->ImplementsInterface(AllowedClass)))
+					{
+						bAllowedToSetBasedOnFilter = true;
+						break;
+					}
+				}
+			}
+
+			if (DisallowedClassCount > 0 && bAllowedToSetBasedOnFilter)
+			{
+				for (const UClass* DisallowedClass : DisallowedClassFilters)
+				{
+					const bool bDisallowedClassIsInterface = DisallowedClass->HasAnyClassFlags(CLASS_Interface);
+					if (AssetClass->IsChildOf(DisallowedClass) || (bDisallowedClassIsInterface && AssetClass->ImplementsInterface(DisallowedClass)))
+					{
+						bAllowedToSetBasedOnFilter = false;
+						break;
+					}
+				}
 			}
 		}
 	}
diff --git a/Engine/Source/Editor/PropertyEditor/Private/UserInterface/PropertyEditor/SPropertyEditorAsset.h b/Engine/Source/Editor/PropertyEditor/Private/UserInterface/PropertyEditor/SPropertyEditorAsset.h
--- a/Engine/Source/Editor/PropertyEditor/Private/UserInterface/PropertyEditor/SPropertyEditorAsset.h
+++ b/Engine/Source/Editor/PropertyEditor/Private/UserInterface/PropertyEditor/SPropertyEditorAsset.h
@@ -333,7 +333,10 @@ private:
 	UClass* ObjectClass;
 
 	/** Classes that can be used with this property */
-	TArray<const UClass*> CustomClassFilters;
+	TArray<const UClass*> AllowedClassFilters;
+
+	/** Classes that can NOT be used with this property */
+	TArray<const UClass*> DisallowedClassFilters;
 
 	/** A list of the factories we can use to create new assets */
 	TArray<UFactory*> NewAssetFactories;
diff --git a/Engine/Source/Editor/PropertyEditor/Private/UserInterface/PropertyEditor/SPropertyMenuAssetPicker.cpp b/Engine/Source/Editor/PropertyEditor/Private/UserInterface/PropertyEditor/SPropertyMenuAssetPicker.cpp
--- a/Engine/Source/Editor/PropertyEditor/Private/UserInterface/PropertyEditor/SPropertyMenuAssetPicker.cpp
+++ b/Engine/Source/Editor/PropertyEditor/Private/UserInterface/PropertyEditor/SPropertyMenuAssetPicker.cpp
@@ -23,6 +23,7 @@ void SPropertyMenuAssetPicker::Construct( const FArguments& InArgs )
 	bAllowClear = InArgs._AllowClear;
 	bAllowCopyPaste = InArgs._AllowCopyPaste;
 	AllowedClasses = InArgs._AllowedClasses;
+	DisallowedClasses = InArgs._DisallowedClasses;
 	NewAssetFactories = InArgs._NewAssetFactories;
 	OnShouldFilterAsset = InArgs._OnShouldFilterAsset;
 	OnSet = InArgs._OnSet;
@@ -116,6 +117,12 @@ void SPropertyMenuAssetPicker::Construct( const FArguments& InArgs )
 				AssetPickerConfig.Filter.ClassNames.Add( AllowedClasses[i]->GetFName() );
 			}
 		}
+
+		for (int32 i = 0; i < DisallowedClasses.Num(); ++i)
+		{
+			AssetPickerConfig.Filter.RecursiveClassesExclusionSet.Add(DisallowedClasses[i]->GetFName());
+		}
+
 		// Allow child classes
 		AssetPickerConfig.Filter.bRecursiveClasses = true;
 		// Set a delegate for setting the asset from the picker
diff --git a/Engine/Source/Editor/PropertyEditor/Private/UserInterface/PropertyEditor/SPropertyMenuAssetPicker.h b/Engine/Source/Editor/PropertyEditor/Private/UserInterface/PropertyEditor/SPropertyMenuAssetPicker.h
--- a/Engine/Source/Editor/PropertyEditor/Private/UserInterface/PropertyEditor/SPropertyMenuAssetPicker.h
+++ b/Engine/Source/Editor/PropertyEditor/Private/UserInterface/PropertyEditor/SPropertyMenuAssetPicker.h
@@ -21,6 +21,7 @@ public:
 		SLATE_ARGUMENT( bool, AllowClear )
 		SLATE_ARGUMENT( bool, AllowCopyPaste )
 		SLATE_ARGUMENT( TArray<const UClass*>, AllowedClasses )
+		SLATE_ARGUMENT( TArray<const UClass*>, DisallowedClasses )
 		SLATE_ARGUMENT( TArray<UFactory*>, NewAssetFactories )
 		SLATE_EVENT( FOnShouldFilterAsset, OnShouldFilterAsset )
 		SLATE_EVENT( FOnAssetSelected, OnSet )
@@ -94,9 +95,12 @@ private:
 	/** Whether the asset can be copied or pasted */
 	bool bAllowCopyPaste;
 
-	/** Array of classes to filter by */
+	/** Array of classes to Allow in filter */
 	TArray<const UClass*> AllowedClasses;
 
+	/** Array of classes to Disallow in filter */
+	TArray<const UClass*> DisallowedClasses;
+
 	/** Array of factories which can create a new asset of suitable type */
 	TArray<UFactory*> NewAssetFactories;
 
diff --git a/Engine/Source/Editor/PropertyEditor/Public/PropertyCustomizationHelpers.h b/Engine/Source/Editor/PropertyEditor/Public/PropertyCustomizationHelpers.h
--- a/Engine/Source/Editor/PropertyEditor/Public/PropertyCustomizationHelpers.h
+++ b/Engine/Source/Editor/PropertyEditor/Public/PropertyCustomizationHelpers.h
@@ -64,7 +64,9 @@ namespace PropertyCustomizationHelpers
 	PROPERTYEDITOR_API TSharedRef<SWidget> MakeBrowseButton( FSimpleDelegate OnClearClicked, TAttribute<FText> OptionalToolTipText = FText(), TAttribute<bool> IsEnabled = true );
 	PROPERTYEDITOR_API TSharedRef<SWidget> MakeAssetPickerAnchorButton( FOnGetAllowedClasses OnGetAllowedClasses, FOnAssetSelected OnAssetSelectedFromPicker );
 	PROPERTYEDITOR_API TSharedRef<SWidget> MakeAssetPickerWithMenu( const FAssetData& InitialObject, const bool AllowClear, const TArray<const UClass*>& AllowedClasses, const TArray<UFactory*>& NewAssetFactories, FOnShouldFilterAsset OnShouldFilterAsset, FOnAssetSelected OnSet, FSimpleDelegate OnClose );
+	PROPERTYEDITOR_API TSharedRef<SWidget> MakeAssetPickerWithMenu( const FAssetData& InitialObject, const bool AllowClear, const TArray<const UClass*>& AllowedClasses, const TArray<const UClass*>& DisallowedClasses, const TArray<UFactory*>& NewAssetFactories, FOnShouldFilterAsset OnShouldFilterAsset, FOnAssetSelected OnSet, FSimpleDelegate OnClose );
 	PROPERTYEDITOR_API TSharedRef<SWidget> MakeAssetPickerWithMenu( const FAssetData& InitialObject, const bool AllowClear, const bool AllowCopyPaste, const TArray<const UClass*>& AllowedClasses, const TArray<UFactory*>& NewAssetFactories, FOnShouldFilterAsset OnShouldFilterAsset, FOnAssetSelected OnSet, FSimpleDelegate OnClose );
+	PROPERTYEDITOR_API TSharedRef<SWidget> MakeAssetPickerWithMenu( const FAssetData& InitialObject, const bool AllowClear, const bool AllowCopyPaste, const TArray<const UClass*>& AllowedClasses, const TArray<const UClass*>& DisallowedClasses, const TArray<UFactory*>& NewAssetFactories, FOnShouldFilterAsset OnShouldFilterAsset, FOnAssetSelected OnSet, FSimpleDelegate OnClose );
 	PROPERTYEDITOR_API TSharedRef<SWidget> MakeActorPickerAnchorButton( FOnGetActorFilters OnGetActorFilters, FOnActorSelected OnActorSelectedFromPicker );
 	PROPERTYEDITOR_API TSharedRef<SWidget> MakeActorPickerWithMenu( AActor* const InitialActor, const bool AllowClear, FOnShouldFilterActor ActorFilter, FOnActorSelected OnSet, FSimpleDelegate OnClose, FSimpleDelegate OnUseSelected );
 	PROPERTYEDITOR_API TSharedRef<SWidget> MakeInteractiveActorPicker(FOnGetAllowedClasses OnGetAllowedClasses, FOnShouldFilterActor OnShouldFilterActor, FOnActorSelected OnActorSelectedFromPicker);
@@ -78,6 +80,9 @@ namespace PropertyCustomizationHelpers
 	/** Returns a list of factories which can be used to create new assets, based on the supplied class */
 	PROPERTYEDITOR_API TArray<UFactory*> GetNewAssetFactoriesForClasses(const TArray<const UClass*>& Classes);
 
+	/** Returns a list of factories which can be used to create new assets, based on the supplied classes and respecting the disallowed set */
+	PROPERTYEDITOR_API TArray<UFactory*> GetNewAssetFactoriesForClasses(const TArray<const UClass*>& Classes, const TArray<const UClass*>& DisallowedClasses);
+
 	/** 
 	 * Build a combo button that you bind to a Name or String property or use general delegates
 	 * 
diff --git a/Engine/Source/Runtime/SlateCore/Public/Styling/SlateBrush.h b/Engine/Source/Runtime/SlateCore/Public/Styling/SlateBrush.h
--- a/Engine/Source/Runtime/SlateCore/Public/Styling/SlateBrush.h
+++ b/Engine/Source/Runtime/SlateCore/Public/Styling/SlateBrush.h
@@ -337,7 +337,7 @@ private:
 	 * The image to render for this brush, can be a UTexture or UMaterialInterface or an object implementing 
 	 * the AtlasedTextureInterface. 
 	 */
-	UPROPERTY(EditAnywhere, BlueprintReadWrite, Category=Brush, meta=( AllowPrivateAccess="true", DisplayThumbnail="true", DisplayName="Image", AllowedClasses="Texture,MaterialInterface,SlateTextureAtlasInterface" ))
+	UPROPERTY(EditAnywhere, BlueprintReadWrite, Category=Brush, meta=( AllowPrivateAccess="true", DisplayThumbnail="true", DisplayName="Image", AllowedClasses="Texture,MaterialInterface,SlateTextureAtlasInterface", DisallowedClasses = "MediaTexture"))
 	UObject* ResourceObject;
 
 protected:
