Fix crash resetting cascade distributions to default

diff --git a/Engine/Source/Editor/PropertyEditor/Private/PropertyNode.cpp b/Engine/Source/Editor/PropertyEditor/Private/PropertyNode.cpp
--- a/Engine/Source/Editor/PropertyEditor/Private/PropertyNode.cpp
+++ b/Engine/Source/Editor/PropertyEditor/Private/PropertyNode.cpp
@@ -544,7 +544,7 @@ EPropertyDataValidationResult FPropertyNode::EnsureDataIsValid()
 				{
 					UObject* Obj = ObjectProperty->GetObjectPropertyValue(Addr);
 
-					if (!bShowInnerObjectPropertiesObjectChanged && HasNodeFlags(EPropertyNodeFlags::ShowInnerObjectProperties) && ChildNodes.Num() == 1)
+					if (!bShowInnerObjectPropertiesObjectChanged && HasNodeFlags(EPropertyNodeFlags::ShowInnerObjectProperties|EPropertyNodeFlags::EditInlineNew) && ChildNodes.Num() == 1)
 					{
 						bool bChildObjectFound = false;
 						// should never have more than one node (0 is ok if the object property is null)
