Fix inline edit conditions not supporting undo

PR #4854: UE-61432: Add FScopedTransaction to SetEditConditionState (Contributed by [REDACTED_PERSON])

diff --git a/Engine/Source/Editor/PropertyEditor/Private/Presentation/PropertyEditor/PropertyEditor.cpp b/Engine/Source/Editor/PropertyEditor/Private/Presentation/PropertyEditor/PropertyEditor.cpp
--- a/Engine/Source/Editor/PropertyEditor/Private/Presentation/PropertyEditor/PropertyEditor.cpp
+++ b/Engine/Source/Editor/PropertyEditor/Private/Presentation/PropertyEditor/PropertyEditor.cpp
@@ -404,6 +404,8 @@ bool FPropertyEditor::IsEditConst() const
 
 void FPropertyEditor::SetEditConditionState( bool bShouldEnable )
 {
+	const FScopedTransaction Transaction(FText::Format(LOCTEXT("SetEditConditionState", "Set {0} edit condition state "), PropertyNode->GetDisplayName()));
+
 	// Propagate the value change to any instances if we're editing a template object.
 	FObjectPropertyNode* ObjectNode = PropertyNode->FindObjectItemParent();
 
