fix spelling error
#rnx

diff --git a/Engine/Source/Editor/PropertyEditor/Private/PropertyHandleImpl.cpp b/Engine/Source/Editor/PropertyEditor/Private/PropertyHandleImpl.cpp
--- a/Engine/Source/Editor/PropertyEditor/Private/PropertyHandleImpl.cpp
+++ b/Engine/Source/Editor/PropertyEditor/Private/PropertyHandleImpl.cpp
@@ -1042,7 +1042,7 @@ void FPropertyValueImpl::AddChild()
 		PropertyNodePin->GetReadAddress( !!PropertyNodePin->HasNodeFlags(EPropertyNodeFlags::SingleSelectOnly), ReadAddresses, true, false, true );
 		if ( ReadAddresses.Num() )
 		{
-			// determines whether we actually changed any values (if the user clicks the "emtpy" button when the array is already empty,
+			// determines whether we actually changed any values (if the user clicks the "empty" button when the array is already empty,
 			// we don't want the objects to be marked dirty)
 			bool bNotifiedPreChange = false;
 
