Fixing missing asset factories in asset picker window

diff --git a/Engine/Source/Editor/PropertyEditor/Private/UserInterface/PropertyEditor/SPropertyEditorAsset.cpp b/Engine/Source/Editor/PropertyEditor/Private/UserInterface/PropertyEditor/SPropertyEditorAsset.cpp
--- a/Engine/Source/Editor/PropertyEditor/Private/UserInterface/PropertyEditor/SPropertyEditorAsset.cpp
+++ b/Engine/Source/Editor/PropertyEditor/Private/UserInterface/PropertyEditor/SPropertyEditorAsset.cpp
@@ -215,7 +215,8 @@ void SPropertyEditorAsset::Construct( const FArguments& InArgs, const TSharedPtr
 	{
 		NewAssetFactories = InArgs._NewAssetFactories.GetValue();
 	}
-	else if (AllowedClassFilters.Num() > 1 && !AllowedClassFilters.Contains(UObject::StaticClass()))
+	// If there are more allowed classes than just UObject 
+	else if (AllowedClassFilters.Num() > 1 || !AllowedClassFilters.Contains(UObject::StaticClass()))
 	{
 		NewAssetFactories = PropertyCustomizationHelpers::GetNewAssetFactoriesForClasses(AllowedClassFilters, DisallowedClassFilters);
 	}
