Fix for missing copy/paste entries on default asset picker menus

diff --git a/Engine/Source/Editor/PropertyEditor/Private/PropertyCustomizationHelpers.cpp b/Engine/Source/Editor/PropertyEditor/Private/PropertyCustomizationHelpers.cpp
--- a/Engine/Source/Editor/PropertyEditor/Private/PropertyCustomizationHelpers.cpp
+++ b/Engine/Source/Editor/PropertyEditor/Private/PropertyCustomizationHelpers.cpp
@@ -267,12 +267,12 @@ namespace PropertyCustomizationHelpers
 
 	TSharedRef<SWidget> MakeAssetPickerWithMenu(const FAssetData& InitialObject, const bool AllowClear, const TArray<const UClass*>& AllowedClasses, const TArray<UFactory*>& NewAssetFactories, FOnShouldFilterAsset OnShouldFilterAsset, FOnAssetSelected OnSet, FSimpleDelegate OnClose)
 	{
-		return MakeAssetPickerWithMenu(InitialObject, AllowClear, false, AllowedClasses, EmptyClassArray, NewAssetFactories, OnShouldFilterAsset, OnSet, OnClose);
+		return MakeAssetPickerWithMenu(InitialObject, AllowClear, true, AllowedClasses, EmptyClassArray, NewAssetFactories, OnShouldFilterAsset, OnSet, OnClose);
 	}
 
 	TSharedRef<SWidget> MakeAssetPickerWithMenu(const FAssetData& InitialObject, const bool AllowClear, const TArray<const UClass*>& AllowedClasses, const TArray<const UClass*>& DisallowedClasses, const TArray<UFactory*>& NewAssetFactories, FOnShouldFilterAsset OnShouldFilterAsset, FOnAssetSelected OnSet, FSimpleDelegate OnClose)
 	{
-		return MakeAssetPickerWithMenu(InitialObject, AllowClear, false, AllowedClasses, DisallowedClasses, NewAssetFactories, OnShouldFilterAsset, OnSet, OnClose);
+		return MakeAssetPickerWithMenu(InitialObject, AllowClear, true, AllowedClasses, DisallowedClasses, NewAssetFactories, OnShouldFilterAsset, OnSet, OnClose);
 	}
 
 	TSharedRef<SWidget> MakeAssetPickerWithMenu(const FAssetData& InitialObject, const bool AllowClear, const bool AllowCopyPaste, const TArray<const UClass*>& AllowedClasses, const TArray<UFactory*>& NewAssetFactories, FOnShouldFilterAsset OnShouldFilterAsset, FOnAssetSelected OnSet, FSimpleDelegate OnClose)
