Fix for details panel in UMG editor changing when dragging widgets

diff --git a/Engine/Source/Editor/PropertyEditor/Private/UserInterface/PropertyEditor/SResetToDefaultPropertyEditor.cpp b/Engine/Source/Editor/PropertyEditor/Private/UserInterface/PropertyEditor/SResetToDefaultPropertyEditor.cpp
--- a/Engine/Source/Editor/PropertyEditor/Private/UserInterface/PropertyEditor/SResetToDefaultPropertyEditor.cpp
+++ b/Engine/Source/Editor/PropertyEditor/Private/UserInterface/PropertyEditor/SResetToDefaultPropertyEditor.cpp
@@ -20,7 +20,7 @@ void SResetToDefaultPropertyEditor::Construct(const FArguments& InArgs, const TS
 	bValueDiffersFromDefault = false;
 	OptionalCustomResetToDefault = InArgs._CustomResetToDefault;
 
-	if (InPropertyHandle.IsValid())
+	if (InPropertyHandle.IsValid() && OptionalCustomResetToDefault.IsSet())
 	{
 		InPropertyHandle->MarkResetToDefaultCustomized();
 	}
