Merging change 4343338 from 4.20 which was never integrated back to main.

PropertyEditor - Clear expanded node state on pre-set and restore expanded node state as they're added or initialized, instead of while filtering since that's too late in some cases.  Fixes expansion issues with some customizations, notably in cascade modules.

diff --git a/Engine/Source/Editor/PropertyEditor/Private/DetailLayoutBuilderImpl.cpp b/Engine/Source/Editor/PropertyEditor/Private/DetailLayoutBuilderImpl.cpp
--- a/Engine/Source/Editor/PropertyEditor/Private/DetailLayoutBuilderImpl.cpp
+++ b/Engine/Source/Editor/PropertyEditor/Private/DetailLayoutBuilderImpl.cpp
@@ -484,6 +484,10 @@ TSharedRef<IPropertyHandle> FDetailLayoutBuilderImpl::GetPropertyHandle( TShared
 void FDetailLayoutBuilderImpl::AddExternalRootPropertyNode(TSharedRef<FComplexPropertyNode> InExternalRootNode)
 {
 	ExternalRootPropertyNodes.Add(InExternalRootNode);
+	if (GetDetailsView())
+	{
+		GetDetailsView()->RestoreExpandedItems(InExternalRootNode);
+	}
 }
 
 void FDetailLayoutBuilderImpl::RemoveExternalRootPropertyNode(TSharedRef<FComplexPropertyNode> InExternalRootNode)
diff --git a/Engine/Source/Editor/PropertyEditor/Private/IDetailsViewPrivate.h b/Engine/Source/Editor/PropertyEditor/Private/IDetailsViewPrivate.h
--- a/Engine/Source/Editor/PropertyEditor/Private/IDetailsViewPrivate.h
+++ b/Engine/Source/Editor/PropertyEditor/Private/IDetailsViewPrivate.h
@@ -125,4 +125,9 @@ public:
 
 	/** Saves the expansion state of property nodes for the selected object set */
 	virtual void SaveExpandedItems(TSharedRef<FPropertyNode> StartNode) = 0;
+
+	/**
+	* Restores the expansion state of property nodes for the selected object set
+	*/
+	virtual void RestoreExpandedItems(TSharedRef<FPropertyNode> StartNode) = 0;
 };
diff --git a/Engine/Source/Editor/PropertyEditor/Private/SDetailsView.cpp b/Engine/Source/Editor/PropertyEditor/Private/SDetailsView.cpp
--- a/Engine/Source/Editor/PropertyEditor/Private/SDetailsView.cpp
+++ b/Engine/Source/Editor/PropertyEditor/Private/SDetailsView.cpp
@@ -794,6 +794,7 @@ void SDetailsView::PreSetObject(int32 InNewNumObjects)
 	}
 
 	RootPropertyNodes.Empty(InNewNumObjects);
+	ExpandedDetailNodes.Empty();
 
 	if(DetailsViewArgs.bAllowMultipleTopLevelObjects)
 	{
@@ -881,6 +882,8 @@ void SDetailsView::PostSetObject()
 		FObjectPropertyNode* RootPropertyNode = ComplexRootNode->AsObjectNode();
 
 		RootPropertyNode->InitNode( InitParams );
+
+		RestoreExpandedItems(ComplexRootNode.ToSharedRef());
 	}
 
 	UpdatePropertyMaps();
diff --git a/Engine/Source/Editor/PropertyEditor/Private/SDetailsViewBase.cpp b/Engine/Source/Editor/PropertyEditor/Private/SDetailsViewBase.cpp
--- a/Engine/Source/Editor/PropertyEditor/Private/SDetailsViewBase.cpp
+++ b/Engine/Source/Editor/PropertyEditor/Private/SDetailsViewBase.cpp
@@ -963,7 +963,6 @@ void SDetailsViewBase::RestoreExpandedItems(TSharedRef<FPropertyNode> InitialSta
 void SDetailsViewBase::UpdateFilteredDetails()
 {
 	RootTreeNodes.Reset();
-	ExpandedDetailNodes.Empty();
 
 	FDetailNodeList InitialRootNodeList;
 	
@@ -983,8 +982,6 @@ void SDetailsViewBase::UpdateFilteredDetails()
 			RootPropertyNode->FilterNodes(CurrentFilter.FilterStrings);
 			RootPropertyNode->ProcessSeenFlags(true);
 
-			RestoreExpandedItems(RootPropertyNode.ToSharedRef());
-
 			TSharedPtr<FDetailLayoutBuilderImpl>& DetailLayout = DetailLayouts[RootNodeIndex].DetailLayout;
 			if(DetailLayout.IsValid())
 			{
@@ -995,8 +992,6 @@ void SDetailsViewBase::UpdateFilteredDetails()
 					{
 						ExternalRootNode->FilterNodes(CurrentFilter.FilterStrings);
 						ExternalRootNode->ProcessSeenFlags(true);
-					
-						RestoreExpandedItems(ExternalRootNode.ToSharedRef());
 					}
 				}
 
diff --git a/Engine/Source/Editor/PropertyEditor/Private/SDetailsViewBase.h b/Engine/Source/Editor/PropertyEditor/Private/SDetailsViewBase.h
--- a/Engine/Source/Editor/PropertyEditor/Private/SDetailsViewBase.h
+++ b/Engine/Source/Editor/PropertyEditor/Private/SDetailsViewBase.h
@@ -167,7 +167,8 @@ public:
 	virtual void UpdateSinglePropertyMap(TSharedPtr<FComplexPropertyNode> InRootPropertyNode, FDetailLayoutData& LayoutData, bool bIsExternal) override;
 	virtual FNotifyHook* GetNotifyHook() const override { return DetailsViewArgs.NotifyHook; }
 	virtual const FCustomPropertyTypeLayoutMap& GetCustomPropertyTypeLayoutMap() const { return InstancedTypeToLayoutMap; }
-	void SaveExpandedItems( TSharedRef<FPropertyNode> StartNode ) override;
+	virtual void SaveExpandedItems( TSharedRef<FPropertyNode> StartNode ) override;
+	virtual void RestoreExpandedItems(TSharedRef<FPropertyNode> StartNode) override;
 
 	virtual bool IsConnected() const = 0;
 	virtual FRootPropertyNodeList& GetRootNodes() = 0;
@@ -210,13 +211,6 @@ public:
 	virtual void Tick( const FGeometry& AllottedGeometry, const double InCurrentTime, const float InDeltaTime ) override;
 	// End of SWidget interface
 
-
-	/**
-	* Restores the expansion state of property nodes for the selected object set
-	*/
-	void RestoreExpandedItems(TSharedRef<FPropertyNode> StartNode);
-
-
 protected:
 	/**
 	 * Called when a color property is changed from a color picker
diff --git a/Engine/Source/Editor/PropertyEditor/Private/SStructureDetailsView.cpp b/Engine/Source/Editor/PropertyEditor/Private/SStructureDetailsView.cpp
--- a/Engine/Source/Editor/PropertyEditor/Private/SStructureDetailsView.cpp
+++ b/Engine/Source/Editor/PropertyEditor/Private/SStructureDetailsView.cpp
@@ -217,6 +217,7 @@ void SStructureDetailsView::SetStructureData(TSharedPtr<FStructOnScope> InStruct
 	RootNodesPendingKill.Add(RootNode);
 
 	RootNodes.Empty(1);
+	ExpandedDetailNodes.Empty();
 
 	RootNode = MakeShareable(new FStructurePropertyNode);
 	RootNodes.Add(RootNode);
