Fix regression that prevents actors from being set as the default for a variable in the level script blueprint.

diff --git a/Engine/Source/Editor/PropertyEditor/Private/UserInterface/PropertyEditor/SPropertyEditorAsset.cpp b/Engine/Source/Editor/PropertyEditor/Private/UserInterface/PropertyEditor/SPropertyEditorAsset.cpp
--- a/Engine/Source/Editor/PropertyEditor/Private/UserInterface/PropertyEditor/SPropertyEditorAsset.cpp
+++ b/Engine/Source/Editor/PropertyEditor/Private/UserInterface/PropertyEditor/SPropertyEditorAsset.cpp
@@ -5,6 +5,7 @@
 #include "Engine/SkeletalMesh.h"
 #include "Components/StaticMeshComponent.h"
 #include "Engine/StaticMesh.h"
+#include "Engine/LevelScriptActor.h"
 #include "Editor.h"
 #include "Modules/ModuleManager.h"
 #include "UObject/UObjectHash.h"
@@ -268,7 +269,7 @@ void SPropertyEditorAsset::Construct( const FArguments& InArgs, const TSharedPtr
 				// Go through all the found objects and see if any are a CDO, we can't set an actor in a CDO default.
 				for (UObject* Obj : ObjectList)
 				{
-					if (Obj->IsTemplate())
+					if (Obj->IsTemplate() && !Obj->IsA<ALevelScriptActor>())
 					{
 						IsEnabledAttribute.Set(false);
 						TooltipAttribute.Set(LOCTEXT("VariableHasDisableEditOnTemplateTooltip", "Editing this value in a Class Default Object is not allowed"));
