PR #5034: Clear search filters on Escape button press in the settings windows (Contributed by [REDACTED_PERSON])

diff --git a/Engine/Source/Editor/PropertyEditor/Private/SDetailsView.cpp b/Engine/Source/Editor/PropertyEditor/Private/SDetailsView.cpp
--- a/Engine/Source/Editor/PropertyEditor/Private/SDetailsView.cpp
+++ b/Engine/Source/Editor/PropertyEditor/Private/SDetailsView.cpp
@@ -192,6 +192,7 @@ void SDetailsView::Construct(const FArguments& InArgs)
 				SAssignNew(SearchBox, SSearchBox)
 				.HintText(LOCTEXT("SearchDetailsHint", "Search Details"))
 				.OnTextChanged(this, &SDetailsView::OnFilterTextChanged)
+				.OnTextCommitted(this, &SDetailsView::OnFilterTextCommitted)
 				.AddMetaData<FTagMetaData>(TEXT("Details.Search"))
 			]
 		];
diff --git a/Engine/Source/Editor/PropertyEditor/Private/SDetailsView.h b/Engine/Source/Editor/PropertyEditor/Private/SDetailsView.h
--- a/Engine/Source/Editor/PropertyEditor/Private/SDetailsView.h
+++ b/Engine/Source/Editor/PropertyEditor/Private/SDetailsView.h
@@ -136,9 +136,6 @@ private:
 	/** Called at the end of SetObjectArray after we change the objects being observed */
 	void PostSetObject();
 	
-	/** Called when the filter button is clicked */
-	void OnFilterButtonClicked();
-
 	/** Called to get the visibility of the actor name area */
 	EVisibility GetActorNameAreaVisibility() const;
 
diff --git a/Engine/Source/Editor/PropertyEditor/Private/SDetailsViewBase.cpp b/Engine/Source/Editor/PropertyEditor/Private/SDetailsViewBase.cpp
--- a/Engine/Source/Editor/PropertyEditor/Private/SDetailsViewBase.cpp
+++ b/Engine/Source/Editor/PropertyEditor/Private/SDetailsViewBase.cpp
@@ -574,6 +574,16 @@ void SDetailsViewBase::OnFilterTextChanged(const FText& InFilterText)
 
 }
 
+void SDetailsViewBase::OnFilterTextCommitted(const FText& InSearchText, ETextCommit::Type InCommitType)
+{
+	if (InCommitType == ETextCommit::OnCleared)
+	{
+		SearchBox->SetText(FText::GetEmpty());
+		OnFilterTextChanged(FText::GetEmpty());
+		FSlateApplication::Get().ClearKeyboardFocus(EFocusCause::Cleared);
+	}
+}
+
 TSharedPtr<SWidget> SDetailsViewBase::GetNameAreaWidget()
 {
 	return DetailsViewArgs.bCustomNameAreaLocation ? NameArea : nullptr;
diff --git a/Engine/Source/Editor/PropertyEditor/Private/SDetailsViewBase.h b/Engine/Source/Editor/PropertyEditor/Private/SDetailsViewBase.h
--- a/Engine/Source/Editor/PropertyEditor/Private/SDetailsViewBase.h
+++ b/Engine/Source/Editor/PropertyEditor/Private/SDetailsViewBase.h
@@ -309,6 +309,9 @@ protected:
 	/** Called when the filter text changes.  This filters specific property nodes out of view */
 	void OnFilterTextChanged(const FText& InFilterText);
 
+	/** Called when the filter text is cleared. Resets the filters */
+	void OnFilterTextCommitted(const FText& InSearchText, ETextCommit::Type InCommitType);
+
 	/** Called when the list of currently differing properties changes */
 	virtual void UpdatePropertiesWhitelist(const TSet<FPropertyPath> InWhitelistedProperties) override { CurrentFilter.WhitelistedProperties = InWhitelistedProperties; }
 
