Added the ability to customize and existing property on a class in a details panel without moving it from its location on the details panel.

See IDetailLayoutBuilder::EditDefaultProperty

Also fixed not being able to customize properties in sub-categories

diff --git a/Engine/Source/Editor/PropertyEditor/Private/DetailCategoryBuilderImpl.cpp b/Engine/Source/Editor/PropertyEditor/Private/DetailCategoryBuilderImpl.cpp
--- a/Engine/Source/Editor/PropertyEditor/Private/DetailCategoryBuilderImpl.cpp
+++ b/Engine/Source/Editor/PropertyEditor/Private/DetailCategoryBuilderImpl.cpp
@@ -29,6 +29,27 @@ void FDetailLayout::AddDefaultLayout(const FDetailLayoutCustomization& Layout, b
 	AddLayoutInternal(Layout, bAdvanced ? DefaultAdvancedLayouts : DefaultSimpleLayouts);
 }
 
+FDetailLayoutCustomization* FDetailLayout::GetDefaultLayout(const TSharedRef<FPropertyNode>& PropertyNode)
+{
+	FDetailLayoutCustomization* Customization = 
+		DefaultSimpleLayouts.FindByPredicate([&PropertyNode](const FDetailLayoutCustomization& TestCustomization)
+		{
+			return TestCustomization.GetPropertyNode() == PropertyNode;
+		});
+
+	// Didn't find it in the simple layouts, look in advanced layouts
+	if (!Customization)
+	{
+		Customization = 
+			DefaultAdvancedLayouts.FindByPredicate([&PropertyNode](const FDetailLayoutCustomization& TestCustomization)
+			{
+				return TestCustomization.GetPropertyNode() == PropertyNode;
+			});
+	}
+
+	return Customization;
+}
+
 void FDetailLayout::AddLayoutInternal(const FDetailLayoutCustomization& Layout, FCustomizationList& ListToUse)
 {
 	ListToUse.Add(Layout);
@@ -435,6 +456,18 @@ void FDetailCategoryImpl::OnAdvancedDropdownClicked()
 	RefreshTree(bRefilterCategory);
 }
 
+FDetailLayoutCustomization* FDetailCategoryImpl::GetDefaultCustomization(TSharedRef<FPropertyNode> PropertyNode)
+{
+	FDetailLayout& Layout = GetLayoutForInstance(GetParentLayoutImpl().GetCurrentCustomizationVariableName());
+	
+	FDetailLayoutCustomization* Customization = Layout.GetDefaultLayout(PropertyNode);
+	if (Customization)
+	{
+		return Customization;
+	}
+	return nullptr;
+}
+
 bool FDetailCategoryImpl::ShouldShowAdvanced() const
 {
 	return bUserShowAdvanced || bForceAdvanced;
@@ -792,11 +825,16 @@ void FDetailCategoryImpl::GenerateChildrenForLayouts()
 
 
 void FDetailCategoryImpl::GetChildren(FDetailNodeList& OutChildren)
+{
+	GetGeneratedChildren(OutChildren, false, false);
+}
+
+void FDetailCategoryImpl::GetGeneratedChildren(FDetailNodeList& OutChildren, bool bIgnoreVisibility, bool bIgnoreAdvancedDropdown)
 {
 	for (int32 ChildIndex = 0; ChildIndex < SimpleChildNodes.Num(); ++ChildIndex)
 	{
 		TSharedRef<FDetailTreeNode>& Child = SimpleChildNodes[ChildIndex];
-		if (Child->GetVisibility() == ENodeVisibility::Visible)
+		if (bIgnoreVisibility || Child->GetVisibility() == ENodeVisibility::Visible)
 		{
 			if (Child->ShouldShowOnlyChildren())
 			{
@@ -809,7 +847,7 @@ void FDetailCategoryImpl::GetChildren(FDetailNodeList& OutChildren)
 		}
 	}
 
-	if (ShouldShowAdvanced())
+	if (!bIgnoreAdvancedDropdown && ShouldShowAdvanced())
 	{
 		if (AdvancedDropdownNodeTop.IsValid())
 		{
@@ -820,7 +858,7 @@ void FDetailCategoryImpl::GetChildren(FDetailNodeList& OutChildren)
 		{
 			TSharedRef<FDetailTreeNode>& Child = AdvancedChildNodes[ChildIndex];
 
-			if (Child->GetVisibility() == ENodeVisibility::Visible)
+			if (bIgnoreVisibility || Child->GetVisibility() == ENodeVisibility::Visible)
 			{
 				if (Child->ShouldShowOnlyChildren())
 				{
@@ -834,7 +872,7 @@ void FDetailCategoryImpl::GetChildren(FDetailNodeList& OutChildren)
 		}
 	}
 
-	if (AdvancedDropdownNodeBottom.IsValid())
+	if (!bIgnoreAdvancedDropdown && AdvancedDropdownNodeBottom.IsValid())
 	{
 		OutChildren.Add(AdvancedDropdownNodeBottom.ToSharedRef());
 	}
diff --git a/Engine/Source/Editor/PropertyEditor/Private/DetailCategoryBuilderImpl.h b/Engine/Source/Editor/PropertyEditor/Private/DetailCategoryBuilderImpl.h
--- a/Engine/Source/Editor/PropertyEditor/Private/DetailCategoryBuilderImpl.h
+++ b/Engine/Source/Editor/PropertyEditor/Private/DetailCategoryBuilderImpl.h
@@ -73,6 +73,8 @@ public:
 	const FCustomizationList& GetDefaultSimpleLayouts() const { return DefaultSimpleLayouts; }
 	const FCustomizationList& GetDefaultAdvancedLayouts() const { return DefaultAdvancedLayouts; }
 
+	FDetailLayoutCustomization* GetDefaultLayout(const TSharedRef<FPropertyNode>& PropertyNode);
+
 	bool HasAdvancedLayouts() const { return CustomAdvancedLayouts.Num() > 0 || DefaultAdvancedLayouts.Num() > 0; }
 
 	FName GetInstanceName() const { return InstanceName; }
@@ -127,6 +129,7 @@ public:
 	 * @return Gets a layout at a specific instance
 	 */
 	const FDetailLayout& operator[](int32 Index) const { return Layouts[Index]; }
+	FDetailLayout& operator[](int32 Index) { return Layouts[Index]; }
 
 	/**
 	 * @return Whether or not we need to display a group border around a list of details.
@@ -190,6 +193,11 @@ public:
 	virtual bool ShouldShowOnlyChildren() const override { return bShowOnlyChildren; }
 	virtual FName GetNodeName() const override { return GetCategoryName(); }
 
+	/**
+	 * Gets all generated children with options for ignoring current child visibility or advanced dropdowns
+	 */
+	void GetGeneratedChildren(FDetailNodeList& OutChildren, bool bIgnoreVisibility, bool bIgnoreAdvancedDropdown);
+
 	FCustomPropertyTypeLayoutMap GetCustomPropertyTypeLayoutMap() const;
 
 	/**
@@ -301,6 +309,7 @@ public:
 	 */
 	void SetCategoryAsSpecialFavorite() { bFavoriteCategory = true; bForceAdvanced = true; }
 
+	FDetailLayoutCustomization* GetDefaultCustomization(TSharedRef<FPropertyNode> PropertyNode);
 private:
 	virtual void OnItemExpansionChanged(bool bIsExpanded, bool bShouldSaveState) override;
 
diff --git a/Engine/Source/Editor/PropertyEditor/Private/DetailLayoutBuilderImpl.cpp b/Engine/Source/Editor/PropertyEditor/Private/DetailLayoutBuilderImpl.cpp
--- a/Engine/Source/Editor/PropertyEditor/Private/DetailLayoutBuilderImpl.cpp
+++ b/Engine/Source/Editor/PropertyEditor/Private/DetailLayoutBuilderImpl.cpp
@@ -8,6 +8,7 @@
 #include "StructurePropertyNode.h"
 #include "DetailMultiTopLevelObjectRootNode.h"
 #include "ObjectEditorUtils.h"
+#include "DetailPropertyRow.h"
 
 FDetailLayoutBuilderImpl::FDetailLayoutBuilderImpl(TSharedPtr<FComplexPropertyNode>& InRootNode, FClassToPropertyMap& InPropertyMap, const TSharedRef<IPropertyUtilities>& InPropertyUtilities, const TSharedRef<IPropertyGenerationUtilities>& InPropertyGenerationUtilities, const TSharedPtr< IDetailsViewPrivate >& InDetailsView, bool bIsExternal)
 	: RootNode( InRootNode )
@@ -97,6 +98,36 @@ FDetailWidgetRow& FDetailLayoutBuilderImpl::AddCustomRowToCategory(TSharedPtr<IP
 	return MyCategory.AddCustomRow(CustomSearchString, bForAdvanced);
 }
 
+IDetailPropertyRow* FDetailLayoutBuilderImpl::EditDefaultProperty(TSharedPtr<IPropertyHandle> InPropertyHandle)
+{
+	if (InPropertyHandle.IsValid() && InPropertyHandle->IsValidHandle())
+	{
+		TSharedPtr<FPropertyNode> PropertyNode = GetPropertyNode(InPropertyHandle);
+		if (PropertyNode.IsValid())
+		{
+			UProperty* Property = InPropertyHandle->GetProperty();
+
+			// Get the property's category name
+			FName CategoryFName = FObjectEditorUtils::GetCategoryFName(Property);
+
+			// Get the layout builder's category builder
+			TSharedPtr<FDetailCategoryImpl> DefaultCategory = DefaultCategoryMap.FindRef(CategoryFName);
+
+			if(DefaultCategory.IsValid())
+			{
+				FDetailLayoutCustomization* Customization = DefaultCategory->GetDefaultCustomization(PropertyNode.ToSharedRef());
+				if (Customization)
+				{
+					return Customization->PropertyRow.Get();
+				}
+			}
+		}
+
+	}
+
+	return nullptr;
+}
+
 TSharedRef<IPropertyHandle> FDetailLayoutBuilderImpl::GetProperty( const FName PropertyPath, const UClass* ClassOutermost, FName InInstanceName )
 {	
 	TSharedPtr<FPropertyHandleBase> PropertyHandle; 
@@ -161,22 +192,25 @@ FDetailCategoryImpl& FDetailLayoutBuilderImpl::DefaultCategory( FName CategoryNa
 	return *CategoryImpl;
 }
 
+TSharedPtr<FDetailCategoryImpl> FDetailLayoutBuilderImpl::GetSubCategoryImpl(FName CategoryName)
+{
+	return SubCategoryMap.FindRef(CategoryName);
+}
+
 bool FDetailLayoutBuilderImpl::HasCategory(FName CategoryName)
 {
 	return DefaultCategoryMap.Contains(CategoryName);
 }
 
 void FDetailLayoutBuilderImpl::BuildCategories( const FCategoryMap& CategoryMap, TArray< TSharedRef<FDetailCategoryImpl> >& OutSimpleCategories, TArray< TSharedRef<FDetailCategoryImpl> >& OutAdvancedCategories )
-{
+{		
 	for( FCategoryMap::TConstIterator It(CategoryMap); It; ++It )
 	{
 		TSharedRef<FDetailCategoryImpl> DetailCategory = It.Value().ToSharedRef();
-		//If there is a delimiter in the name it mean its a sub category, we dont show sub category at the root level
-		FString CategoryDelimiterString;
-		CategoryDelimiterString.AppendChar(FPropertyNodeConstants::CategoryDelimiterChar);
+
 		TSharedPtr<FComplexPropertyNode> RootPropertyNode = GetRootNode();
-		const bool bCategoryHidden = PropertyEditorHelpers::IsCategoryHiddenByClass(RootPropertyNode, DetailCategory->GetCategoryName()) 
-			|| ForceHiddenCategories.Contains(DetailCategory->GetCategoryName()) || DetailCategory->GetCategoryName().ToString().Contains(CategoryDelimiterString);
+		const bool bCategoryHidden = PropertyEditorHelpers::IsCategoryHiddenByClass(RootPropertyNode, DetailCategory->GetCategoryName())
+			|| ForceHiddenCategories.Contains(DetailCategory->GetCategoryName());
 
 		if( !bCategoryHidden )
 		{
@@ -208,22 +242,47 @@ void FDetailLayoutBuilderImpl::GenerateDetailLayout()
 	};
 
 	// Merge the two category lists and sort them based on priority
-	FCategoryMap AllCategories = CustomCategoryMap;
-	AllCategories.Append( DefaultCategoryMap );
+	//FCategoryMap AllCategories = CustomCategoryMap;
+	//AllCategories.Append( DefaultCategoryMap );
 
 	TArray< TSharedRef<FDetailCategoryImpl> > SimpleCategories;
 	TArray< TSharedRef<FDetailCategoryImpl> > AdvancedOnlyCategories;
 
-	// Customizations can add more categories while customizing so just keep doing this until the maps are empty
-	while(DefaultCategoryMap.Num() > 0 || CustomCategoryMap.Num() > 0)
+	//If there is a delimiter in the name it mean its a sub category, we dont show sub category at the root level
+	FString CategoryDelimiterString;
+	CategoryDelimiterString.AppendChar(FPropertyNodeConstants::CategoryDelimiterChar);
+
+	SubCategoryMap.Empty();
+	for (FCategoryMap::TIterator It(DefaultCategoryMap); It; ++It)
+	{
+		// Remove all subcategories
+		TSharedPtr<FDetailCategoryImpl> DetailCategory = It.Value();
+		// Note: Sub-categories are added later
+		int32 Index = INDEX_NONE;
+		if (DetailCategory->GetCategoryName().ToString().FindChar(FPropertyNodeConstants::CategoryDelimiterChar, Index))
+		{
+			SubCategoryMap.Add(It.Key(), DetailCategory);
+			It.RemoveCurrent();
+		}
+	}
+
+	// Build default categories
+	while (DefaultCategoryMap.Num() > 0)
 	{
 		FCategoryMap DefaultCategoryMapCopy = DefaultCategoryMap;
-		FCategoryMap CustomCategoryMapCopy = CustomCategoryMap;
 
 		DefaultCategoryMap.Empty();
-		CustomCategoryMap.Empty();
 
 		BuildCategories(DefaultCategoryMapCopy, SimpleCategories, AdvancedOnlyCategories);
+	}
+
+	// Customizations can add more categories while customizing so just keep doing this until the maps are empty
+	while(CustomCategoryMap.Num() > 0)
+	{
+		FCategoryMap CustomCategoryMapCopy = CustomCategoryMap;
+
+		CustomCategoryMap.Empty();
+
 		BuildCategories(CustomCategoryMapCopy, SimpleCategories, AdvancedOnlyCategories);
 	}
 
diff --git a/Engine/Source/Editor/PropertyEditor/Private/DetailLayoutBuilderImpl.h b/Engine/Source/Editor/PropertyEditor/Private/DetailLayoutBuilderImpl.h
--- a/Engine/Source/Editor/PropertyEditor/Private/DetailLayoutBuilderImpl.h
+++ b/Engine/Source/Editor/PropertyEditor/Private/DetailLayoutBuilderImpl.h
@@ -32,6 +32,7 @@ public:
 	virtual IDetailCategoryBuilder& EditCategory(FName CategoryName, const FText& NewLocalizedDisplayName = FText::GetEmpty(), ECategoryPriority::Type CategoryType = ECategoryPriority::Default) override;
 	virtual IDetailPropertyRow& AddPropertyToCategory(TSharedPtr<IPropertyHandle> InPropertyHandle) override;
 	virtual FDetailWidgetRow& AddCustomRowToCategory(TSharedPtr<IPropertyHandle> InPropertyHandle, const FText& InCustomSearchString, bool bForAdvanced = false) override;
+	virtual IDetailPropertyRow* EditDefaultProperty(TSharedPtr<IPropertyHandle> InPropertyHandle) override;
 	virtual TSharedRef<IPropertyHandle> GetProperty(const FName PropertyPath, const UClass* ClassOutermost, FName InInstanceName) override;
 	virtual FName GetTopLevelProperty() override;
 	virtual void HideProperty(const TSharedPtr<IPropertyHandle> Property) override;
@@ -45,6 +46,7 @@ public:
 	virtual UClass* GetBaseClass() const override;
 	virtual const TArray<TWeakObjectPtr<UObject>>& GetSelectedObjects() const override;
 	virtual bool HasClassDefaultObject() const override;
+
 	/**
 	 * Creates a default category. The SDetails view will generate widgets in default categories
 	 *
@@ -52,6 +54,8 @@ public:
 	 */
 	FDetailCategoryImpl& DefaultCategory(FName CategoryName);
 
+	TSharedPtr<FDetailCategoryImpl> GetSubCategoryImpl(FName CategoryName);
+
 	/**
 	* Return true if the category exist
 	*/
@@ -224,6 +228,8 @@ private:
 	FCategoryMap CustomCategoryMap;
 	/** A mapping of category names to categories which have been not customized */
 	FCategoryMap DefaultCategoryMap;
+	/** Mapping of subcategory names to their subcategory builder */
+	FCategoryMap SubCategoryMap;
 	/** A mapping of classes to top level properties in that class */
 	FClassToPropertyMap& PropertyMap;
 	/** Force hidden categories set by the user */
diff --git a/Engine/Source/Editor/PropertyEditor/Private/DetailPropertyRow.cpp b/Engine/Source/Editor/PropertyEditor/Private/DetailPropertyRow.cpp
--- a/Engine/Source/Editor/PropertyEditor/Private/DetailPropertyRow.cpp
+++ b/Engine/Source/Editor/PropertyEditor/Private/DetailPropertyRow.cpp
@@ -12,6 +12,7 @@
 #include "ItemPropertyNode.h"
 #include "ObjectPropertyNode.h"
 #include "DetailWidgetRow.h"
+#include "CategoryPropertyNode.h"
 
 const float FDetailWidgetRow::DefaultValueMinWidth = 125.0f;
 const float FDetailWidgetRow::DefaultValueMaxWidth = 125.0f;
@@ -235,7 +236,23 @@ void FDetailPropertyRow::OnItemNodeInitialized( TSharedRef<FDetailCategoryImpl>
 
 void FDetailPropertyRow::OnGenerateChildren( FDetailNodeList& OutChildren )
 {
-	if( PropertyNode->AsCategoryNode() || PropertyNode->GetProperty() || ExternalObjectLayout.IsValid())
+	if (PropertyNode->AsCategoryNode())
+	{
+		// This is a sub-category.  Populate from SubCategory builder
+		TSharedRef<FDetailCategoryImpl> ParentCategoryRef = ParentCategory.Pin().ToSharedRef();
+		FDetailLayoutBuilderImpl& LayoutBuilder = ParentCategoryRef->GetParentLayoutImpl();
+		TSharedPtr<FDetailCategoryImpl> MyCategory = LayoutBuilder.GetSubCategoryImpl(PropertyNode->AsCategoryNode()->GetCategoryName());
+		check(MyCategory.IsValid());
+
+		MyCategory->GenerateLayout();
+
+		// Ignore the header of the category by just getting the categories children directly. We are the header in this case.
+		// Also ignore visibility here as we dont have a filter yet and the children will be filtered later anyway
+		const bool bIgnoreVisibility = true;
+		const bool bIgnoreAdvancedDropdown = true;
+		MyCategory->GetGeneratedChildren(OutChildren, bIgnoreVisibility, bIgnoreAdvancedDropdown);
+	}
+	else if (PropertyNode->GetProperty() || ExternalObjectLayout.IsValid())
 	{
 		GenerateChildrenForPropertyNode( PropertyNode, OutChildren );
 	}
diff --git a/Engine/Source/Editor/PropertyEditor/Public/DetailLayoutBuilder.h b/Engine/Source/Editor/PropertyEditor/Public/DetailLayoutBuilder.h
--- a/Engine/Source/Editor/PropertyEditor/Public/DetailLayoutBuilder.h
+++ b/Engine/Source/Editor/PropertyEditor/Public/DetailLayoutBuilder.h
@@ -94,7 +94,7 @@ public:
 	 * @param NewLocalizedDisplayName	The new display name of the category (optional)
 	 * @param CategoryType				Category type to define sort order.  Category display order is sorted by this type (optional)
 	 */
-	virtual IDetailCategoryBuilder& EditCategory( FName CategoryName, const FText& NewLocalizedDisplayName = FText::GetEmpty(), ECategoryPriority::Type CategoryType = ECategoryPriority::Default ) = 0;
+	virtual IDetailCategoryBuilder& EditCategory(FName CategoryName, const FText& NewLocalizedDisplayName = FText::GetEmpty(), ECategoryPriority::Type CategoryType = ECategoryPriority::Default) = 0;
 
 	/** 
 	 * Adds the property to its given category automatically. Useful in detail customizations which want to preserve categories.
@@ -104,13 +104,23 @@ public:
 	virtual IDetailPropertyRow& AddPropertyToCategory(TSharedPtr<IPropertyHandle> InPropertyHandle) = 0;
 
 	/**
-	* Adds a custom row to the property's category automatically. Useful in detail customizations which want to preserve categories.
-	* @param InPropertyHandle			The handle to the property that you want to add to its own category.
-	* @param InCustomSearchString		A string which is used to filter this custom row when a user types into the details panel search box.
-	* @return the detail widget that can be further customized. 
-	*/
+	 * Adds a custom row to the property's category automatically. Useful in detail customizations which want to preserve categories.
+	 * @param InPropertyHandle			The handle to the property that you want to add to its own category.
+	 * @param InCustomSearchString		A string which is used to filter this custom row when a user types into the details panel search box.
+	 * @return the detail widget that can be further customized.
+	 */
 	virtual FDetailWidgetRow& AddCustomRowToCategory(TSharedPtr<IPropertyHandle> InPropertyHandle, const FText& InCustomSearchString, bool bForAdvanced = false) = 0;
 
+	/**
+	 * Allows for the customization of a property row for a property that already exists on a class being edited in the details panel
+	 * The property will remain in the default location but the widget or other attributes for the property can be changed 
+	 * Note This cannot be used to customize other customizations
+
+	 * @param InPropertyHandle	The handle to the property that you want to add to its own category.
+	 * @return					The property row to edit or nullptr if the property row does not exist
+	 */
+	virtual IDetailPropertyRow* EditDefaultProperty(TSharedPtr<IPropertyHandle> InPropertyHandle) = 0;
+
 	/**
 	 * Hides an entire category
 	 *
