Fix crash and missing properties for categories that arent nested objects (like categories of sub-objects)

diff --git a/Engine/Source/Editor/PropertyEditor/Private/DetailPropertyRow.cpp b/Engine/Source/Editor/PropertyEditor/Private/DetailPropertyRow.cpp
--- a/Engine/Source/Editor/PropertyEditor/Private/DetailPropertyRow.cpp
+++ b/Engine/Source/Editor/PropertyEditor/Private/DetailPropertyRow.cpp
@@ -236,23 +236,24 @@ void FDetailPropertyRow::OnItemNodeInitialized( TSharedRef<FDetailCategoryImpl>
 
 void FDetailPropertyRow::OnGenerateChildren( FDetailNodeList& OutChildren )
 {
-	if (PropertyNode->AsCategoryNode())
+	if (PropertyNode->AsCategoryNode() && PropertyNode->GetParentNode() && !PropertyNode->GetParentNode()->AsObjectNode())
 	{
 		// This is a sub-category.  Populate from SubCategory builder
 		TSharedRef<FDetailCategoryImpl> ParentCategoryRef = ParentCategory.Pin().ToSharedRef();
 		FDetailLayoutBuilderImpl& LayoutBuilder = ParentCategoryRef->GetParentLayoutImpl();
 		TSharedPtr<FDetailCategoryImpl> MyCategory = LayoutBuilder.GetSubCategoryImpl(PropertyNode->AsCategoryNode()->GetCategoryName());
-		check(MyCategory.IsValid());
-
-		MyCategory->GenerateLayout();
+		if(MyCategory.IsValid())
+		{
+			MyCategory->GenerateLayout();
 
-		// Ignore the header of the category by just getting the categories children directly. We are the header in this case.
-		// Also ignore visibility here as we dont have a filter yet and the children will be filtered later anyway
-		const bool bIgnoreVisibility = true;
-		const bool bIgnoreAdvancedDropdown = true;
-		MyCategory->GetGeneratedChildren(OutChildren, bIgnoreVisibility, bIgnoreAdvancedDropdown);
+			// Ignore the header of the category by just getting the categories children directly. We are the header in this case.
+			// Also ignore visibility here as we dont have a filter yet and the children will be filtered later anyway
+			const bool bIgnoreVisibility = true;
+			const bool bIgnoreAdvancedDropdown = true;
+			MyCategory->GetGeneratedChildren(OutChildren, bIgnoreVisibility, bIgnoreAdvancedDropdown);
+		}
 	}
-	else if (PropertyNode->GetProperty() || ExternalObjectLayout.IsValid())
+	else if (PropertyNode->AsCategoryNode() || PropertyNode->GetProperty() || ExternalObjectLayout.IsValid())
 	{
 		GenerateChildrenForPropertyNode( PropertyNode, OutChildren );
 	}
