Exposed a way to get the category name and display name of a property during details customization

diff --git a/Engine/Source/Editor/PropertyEditor/Private/PropertyHandleImpl.cpp b/Engine/Source/Editor/PropertyEditor/Private/PropertyHandleImpl.cpp
--- a/Engine/Source/Editor/PropertyEditor/Private/PropertyHandleImpl.cpp
+++ b/Engine/Source/Editor/PropertyEditor/Private/PropertyHandleImpl.cpp
@@ -23,6 +23,7 @@
 #include "Widgets/Notifications/SNotificationList.h"
 #include "UObject/EnumProperty.h"
 #include "IDetailPropertyRow.h"
+#include "ObjectEditorUtils.h"
 
 #define LOCTEXT_NAMESPACE "PropertyHandleImplementation"
 
@@ -3088,6 +3089,30 @@ void FPropertyHandleBase::ExecuteCustomResetToDefault(const FResetToDefaultOverr
 	Implementation->GetPropertyUtilities()->EnqueueDeferredAction(FSimpleDelegate::CreateLambda([this, InOnCustomResetToDefault]() { OnCustomResetToDefault(InOnCustomResetToDefault); }));
 }
 
+FName FPropertyHandleBase::GetDefaultCategoryName() const
+{
+	UProperty* Property = GetProperty();
+
+	if (Property)
+	{
+		return FObjectEditorUtils::GetCategoryFName(Property);
+	}
+
+	return NAME_None;
+}
+
+FText FPropertyHandleBase::GetDefaultCategoryText() const
+{
+	UProperty* Property = GetProperty();
+
+	if (Property)
+	{
+		return FObjectEditorUtils::GetCategoryText(Property);
+	}
+
+	return FText::GetEmpty();
+}
+
 /** Implements common property value functions */
 #define IMPLEMENT_PROPERTY_VALUE( ClassName ) \
 	ClassName::ClassName( TSharedRef<FPropertyNode> PropertyNode, FNotifyHook* NotifyHook, TSharedPtr<IPropertyUtilities> PropertyUtilities ) \
diff --git a/Engine/Source/Editor/PropertyEditor/Private/PropertyHandleImpl.h b/Engine/Source/Editor/PropertyEditor/Private/PropertyHandleImpl.h
--- a/Engine/Source/Editor/PropertyEditor/Private/PropertyHandleImpl.h
+++ b/Engine/Source/Editor/PropertyEditor/Private/PropertyHandleImpl.h
@@ -498,6 +498,8 @@ public:
 	virtual TArray<TSharedPtr<IPropertyHandle>> AddChildStructure( TSharedRef<FStructOnScope> ChildStructure ) override;
 	virtual bool CanResetToDefault() const override;
 	virtual void ExecuteCustomResetToDefault(const FResetToDefaultOverride& InOnCustomResetToDefault) override;
+	virtual FName GetDefaultCategoryName() const override;
+	virtual FText GetDefaultCategoryText() const override;
 
 	TSharedPtr<FPropertyNode> GetPropertyNode() const;
 	void OnCustomResetToDefault(const FResetToDefaultOverride& OnCustomResetToDefault);
diff --git a/Engine/Source/Editor/PropertyEditor/Public/PropertyHandle.h b/Engine/Source/Editor/PropertyEditor/Public/PropertyHandle.h
--- a/Engine/Source/Editor/PropertyEditor/Public/PropertyHandle.h
+++ b/Engine/Source/Editor/PropertyEditor/Public/PropertyHandle.h
@@ -605,6 +605,18 @@ public:
 	 * Sets an override for this property's reset to default behavior
 	 */
 	virtual void ExecuteCustomResetToDefault(const class FResetToDefaultOverride& OnCustomResetToDefault) = 0;
+
+	/**
+	 * Gets the category FName that a property is in at the default location defined by the class the property is in
+	 * It does not handle the property being moved to another category during customization
+	 */
+	virtual FName GetDefaultCategoryName() const = 0;
+
+	/**
+	* Gets the category display name that a property is in at the default location defined by the class the property is in
+	* It does not handle the property being moved to another category during customization
+	*/
+	virtual FText GetDefaultCategoryText() const = 0;
 };
 
 /**
