Fix for instanced EditInlineNew properties being improperly linked after copying

diff --git a/Engine/Source/Editor/PropertyEditor/Private/SDetailSingleItemRow.cpp b/Engine/Source/Editor/PropertyEditor/Private/SDetailSingleItemRow.cpp
--- a/Engine/Source/Editor/PropertyEditor/Private/SDetailSingleItemRow.cpp
+++ b/Engine/Source/Editor/PropertyEditor/Private/SDetailSingleItemRow.cpp
@@ -573,6 +573,56 @@ void SDetailSingleItemRow::OnPasteProperty()
 			TSharedPtr<IPropertyHandle> Handle = PropertyEditorHelpers::GetPropertyHandle(PropertyNode.ToSharedRef(), OwnerTreeNode.Pin()->GetDetailsView()->GetNotifyHook(), OwnerTreeNode.Pin()->GetDetailsView()->GetPropertyUtilities());
 
 			Handle->SetValueFromFormattedString(ClipboardContent);
+			TArray<TSharedPtr<IPropertyHandle>> CopiedHandles;
+
+			CopiedHandles.Add(Handle);
+
+			while (CopiedHandles.Num() > 0)
+			{
+
+				Handle = CopiedHandles.Pop();
+
+				// Add all child properties to the list so we can check them next
+				uint32 NumChildren;
+				Handle->GetNumChildren(NumChildren);
+				for (uint32 ChildIndex = 0; ChildIndex < NumChildren; ChildIndex++)
+				{
+					CopiedHandles.Add(Handle->GetChildHandle(ChildIndex));
+				}
+
+				UObject* NewValueAsObject = nullptr;
+				if (FPropertyAccess::Success == Handle->GetValue(NewValueAsObject))
+				{
+				
+					// if the object is instanced, then we need to do a deep copy.
+					if (Handle->GetProperty() != nullptr
+						&& (Handle->GetProperty()->PropertyFlags & (CPF_InstancedReference | CPF_ContainsInstancedReference)) != 0)
+					{
+						UObject* DuplicateOuter = nullptr;
+
+						TArray<UObject*> Outers;
+						Handle->GetOuterObjects(Outers);
+
+						// Update the duplicate's outer to point to this outer. The source's outer may be some other object/asset
+						// but we want this to own the duplicate.
+						if (Outers.Num() > 0)
+						{
+							DuplicateOuter = Outers[0];
+						}
+
+						// This does a deep copy of NewValueAsObject. It's subobjects and property data will be 
+						// copied.
+						UObject* DuplicateOfNewValue = DuplicateObject<UObject>(NewValueAsObject, DuplicateOuter);
+						TArray<FString> DuplicateValueAsString;
+						DuplicateValueAsString.Add(DuplicateOfNewValue->GetPathName());
+						Handle->SetPerObjectValues(DuplicateValueAsString);
+
+					}
+				}
+			}
+
+			// Need to refresh the details panel in case a property was pasted over another.
+			OwnerTreeNode.Pin()->GetDetailsView()->ForceRefresh();
 		}
 	}
 }
