Adjusted padding slightly to ensure highliglting looks good.

diff --git a/Engine/Source/Editor/PropertyEditor/Private/SDetailsView.cpp b/Engine/Source/Editor/PropertyEditor/Private/SDetailsView.cpp
--- a/Engine/Source/Editor/PropertyEditor/Private/SDetailsView.cpp
+++ b/Engine/Source/Editor/PropertyEditor/Private/SDetailsView.cpp
@@ -197,14 +197,14 @@ void SDetailsView::Construct(const FArguments& InArgs)
 		[
 			SNew(SOverlay)
 			+SOverlay::Slot()
-			.Padding(2.0f, 0.0f, 0.0f, 0.0f)
+			.Padding(0.0f, 0.0f, 0.0f, 0.0f)
 			[
 				SNew(SImage)
 				.Image(FEditorStyle::GetBrush("Searching.SearchActiveTab"))
 				.Visibility_Lambda([&](){ return this->bHasActiveFilter ? EVisibility::Visible : EVisibility::Collapsed; })
 			]
 			+SOverlay::Slot()
-			.Padding(0.f, 2.0f, 4.f, 2.f)
+			.Padding(2.f, 2.0f, 4.f, 2.f)
 			.VAlign( VAlign_Center )
 			[
 				// Create the search box
