Comment fix

diff --git a/Engine/Source/Editor/PropertyEditor/Private/PropertyHandleImpl.h b/Engine/Source/Editor/PropertyEditor/Private/PropertyHandleImpl.h
--- a/Engine/Source/Editor/PropertyEditor/Private/PropertyHandleImpl.h
+++ b/Engine/Source/Editor/PropertyEditor/Private/PropertyHandleImpl.h
@@ -180,7 +180,7 @@ public:
 	 * Get the value of a property as a formatted string, possibly using an alternate form more suitable for display in the UI
 	 *
 	 * @param OutValue	The formatted string value to set
-	 * @param PortFlags	Determines how the property's value is accessed. Defaults to PPF_PropertyWindow
+	 * @param PortFlags	Determines how the property's value is accessed. 
 	 * @return The result of attempting to get the value
 	 */
 	FPropertyAccess::Result GetValueAsDisplayString( FString& OutString, EPropertyPortFlags PortFlags ) const;
