Don't show categories with no visible properties

diff --git a/Engine/Source/Editor/PropertyEditor/Private/ObjectPropertyNode.cpp b/Engine/Source/Editor/PropertyEditor/Private/ObjectPropertyNode.cpp
--- a/Engine/Source/Editor/PropertyEditor/Private/ObjectPropertyNode.cpp
+++ b/Engine/Source/Editor/PropertyEditor/Private/ObjectPropertyNode.cpp
@@ -413,7 +413,7 @@ void FObjectPropertyNode::InternalInitChildNodes( FName SinglePropertyName )
 	}
 
 	// Create a merged list of user-enforced sorted info, hidden category info, etc...
-	TSet<FName> CategoriesToShow;
+	TSet<FName> CategoriesFromBlueprints, CategoriesFromProperties;
 	TArray<FName> SortedCategories;
 
 #if WITH_EDITORONLY_DATA
@@ -423,9 +423,9 @@ void FObjectPropertyNode::InternalInitChildNodes( FName SinglePropertyName )
 		{
 			for (const FName& TestCategory : BP->CategorySorting)
 			{
-				if (!CategoriesToShow.Contains(TestCategory))
+				if (!CategoriesFromBlueprints.Contains(TestCategory))
 				{
-					CategoriesToShow.Add(TestCategory);
+					CategoriesFromBlueprints.Add(TestCategory);
 					SortedCategories.Add(TestCategory);
 				}
 			}
@@ -475,15 +475,24 @@ void FObjectPropertyNode::InternalInitChildNodes( FName SinglePropertyName )
 			const bool bShowIfDisableEditOnInstance = !(*It)->HasAnyPropertyFlags(CPF_DisableEditOnInstance) || bShouldShowDisableEditOnInstance;
 			if( bShouldShowHiddenProperties || (bShowIfNonHiddenEditableProperty && bShowIfDisableEditOnInstance) )
 			{
-				if (!CategoriesToShow.Contains(CategoryName))
+				if (!CategoriesFromBlueprints.Contains(CategoryName) && !CategoriesFromProperties.Contains(CategoryName) )
 				{
-					CategoriesToShow.Add(CategoryName);
 					SortedCategories.Add(CategoryName);
 				}
+				CategoriesFromProperties.Add(CategoryName);
 			}
 		}
 	}
 
+	// Categories from the Blueprint class may include categories with no associated properties, so remove those ones
+	for ( const FName& CategoryName : CategoriesFromBlueprints )
+	{
+		if ( !CategoriesFromProperties.Contains(CategoryName) )
+		{
+			SortedCategories.Remove(CategoryName);
+		}
+	}
+
 	//////////////////////////////////////////
 	// Add the category headers and the child items that belong inside of them.
 
