Added the ability to get the current set of categories from a DetailLayoutBuilder

This allows details customizations to then iterate over default properties in the order they are sorted into via metata. Useful e.g. when applying optional logic to all properties in a class/struct.

diff --git a/Engine/Source/Editor/PropertyEditor/Private/DetailLayoutBuilderImpl.cpp b/Engine/Source/Editor/PropertyEditor/Private/DetailLayoutBuilderImpl.cpp
--- a/Engine/Source/Editor/PropertyEditor/Private/DetailLayoutBuilderImpl.cpp
+++ b/Engine/Source/Editor/PropertyEditor/Private/DetailLayoutBuilderImpl.cpp
@@ -69,6 +69,17 @@ IDetailCategoryBuilder& FDetailLayoutBuilderImpl::EditCategory( FName CategoryNa
 	return *CategoryImpl;
 }
 
+void FDetailLayoutBuilderImpl::GetCategoryNames(TArray<FName>& OutCategoryNames) const
+{
+	OutCategoryNames.Reserve(DefaultCategoryMap.Num() + CustomCategoryMap.Num());
+
+	TArray<FName> TempCategoryNames;
+	DefaultCategoryMap.GenerateKeyArray(TempCategoryNames);
+	OutCategoryNames.Append(TempCategoryNames);
+	CustomCategoryMap.GenerateKeyArray(TempCategoryNames);
+	OutCategoryNames.Append(TempCategoryNames);
+}
+
 IDetailPropertyRow& FDetailLayoutBuilderImpl::AddPropertyToCategory(TSharedPtr<IPropertyHandle> InPropertyHandle)
 {
 	// Get the UProperty itself
diff --git a/Engine/Source/Editor/PropertyEditor/Private/DetailLayoutBuilderImpl.h b/Engine/Source/Editor/PropertyEditor/Private/DetailLayoutBuilderImpl.h
--- a/Engine/Source/Editor/PropertyEditor/Private/DetailLayoutBuilderImpl.h
+++ b/Engine/Source/Editor/PropertyEditor/Private/DetailLayoutBuilderImpl.h
@@ -30,6 +30,7 @@ public:
 	virtual void GetObjectsBeingCustomized(TArray< TWeakObjectPtr<UObject> >& OutObjects) const override;
 	virtual void GetStructsBeingCustomized(TArray< TSharedPtr<FStructOnScope> >& OutStructs) const override;
 	virtual IDetailCategoryBuilder& EditCategory(FName CategoryName, const FText& NewLocalizedDisplayName = FText::GetEmpty(), ECategoryPriority::Type CategoryType = ECategoryPriority::Default) override;
+	virtual void GetCategoryNames(TArray<FName>& OutCategoryNames) const override;
 	virtual IDetailPropertyRow& AddPropertyToCategory(TSharedPtr<IPropertyHandle> InPropertyHandle) override;
 	virtual FDetailWidgetRow& AddCustomRowToCategory(TSharedPtr<IPropertyHandle> InPropertyHandle, const FText& InCustomSearchString, bool bForAdvanced = false) override;
 	virtual IDetailPropertyRow* EditDefaultProperty(TSharedPtr<IPropertyHandle> InPropertyHandle) override;
diff --git a/Engine/Source/Editor/PropertyEditor/Public/DetailLayoutBuilder.h b/Engine/Source/Editor/PropertyEditor/Public/DetailLayoutBuilder.h
--- a/Engine/Source/Editor/PropertyEditor/Public/DetailLayoutBuilder.h
+++ b/Engine/Source/Editor/PropertyEditor/Public/DetailLayoutBuilder.h
@@ -96,6 +96,12 @@ public:
 	 */
 	virtual IDetailCategoryBuilder& EditCategory(FName CategoryName, const FText& NewLocalizedDisplayName = FText::GetEmpty(), ECategoryPriority::Type CategoryType = ECategoryPriority::Default) = 0;
 
+	/**
+	 * Gets the current set of existing category names. This includes both categories derived from properties and categories added via EditCategory.
+	 * @param	OutCategoryNames	 The array of category names
+	 */
+	virtual void GetCategoryNames(TArray<FName>& OutCategoryNames) const = 0;
+
 	/** 
 	 * Adds the property to its given category automatically. Useful in detail customizations which want to preserve categories.
 	 * @param InPropertyHandle			The handle to the property that you want to add to its own category.
