Fixes for char literals being handled as bool instead of FString in IPropertyHandle::SetValue() and FFCPXMLNode::SetContent().
Fix for bad use of TAttribute in FBodyInstanceCustomizationHelper::AddMassInKg().

Reference: https://github.com/EpicGames/UnrealEngine/pull/5345

diff --git a/Engine/Source/Editor/DetailCustomizations/Private/BodyInstanceCustomization.cpp b/Engine/Source/Editor/DetailCustomizations/Private/BodyInstanceCustomization.cpp
--- a/Engine/Source/Editor/DetailCustomizations/Private/BodyInstanceCustomization.cpp
+++ b/Engine/Source/Editor/DetailCustomizations/Private/BodyInstanceCustomization.cpp
@@ -1560,7 +1560,7 @@ void FBodyInstanceCustomizationHelper::AddMassInKg(IDetailCategoryBuilder& Physi
 			.Padding(0.f, 0.f, 10.f, 0.f)
 			[
 				SNew(SNumericEntryBox<float>)
-				.IsEnabled(&FBodyInstanceCustomizationHelper::IsBodyMassEnabled)
+				.IsEnabled(this, &FBodyInstanceCustomizationHelper::IsBodyMassEnabled)
 				.Font(IDetailLayoutBuilder::GetDetailFont())
 				.Value(this, &FBodyInstanceCustomizationHelper::OnGetBodyMass)
 				.OnValueCommitted(this, &FBodyInstanceCustomizationHelper::OnSetBodyMass)
diff --git a/Engine/Source/Editor/MovieSceneTools/Private/FCPXML/FCPXMLNode.h b/Engine/Source/Editor/MovieSceneTools/Private/FCPXML/FCPXMLNode.h
--- a/Engine/Source/Editor/MovieSceneTools/Private/FCPXML/FCPXMLNode.h
+++ b/Engine/Source/Editor/MovieSceneTools/Private/FCPXML/FCPXMLNode.h
@@ -143,6 +143,13 @@ public:
 	/** Sets the new value of the node */
 	void SetContent(bool InContent);
 
+	/** Sets the new value of the node */
+	FORCEINLINE void SetContent(const TCHAR* InContent)
+	{
+		// Force an FString conversion here, otherwise it calls the bool overload
+		SetContent(FString(InContent));
+	}
+
 	/** Gets a list of child nodes */
 	const TArray<TSharedRef<FFCPXMLNode>>& GetChildNodes() const;
 	/** Gets number of child nodes */
diff --git a/Engine/Source/Editor/PropertyEditor/Private/PropertyHandleImpl.cpp b/Engine/Source/Editor/PropertyEditor/Private/PropertyHandleImpl.cpp
--- a/Engine/Source/Editor/PropertyEditor/Private/PropertyHandleImpl.cpp
+++ b/Engine/Source/Editor/PropertyEditor/Private/PropertyHandleImpl.cpp
@@ -2116,6 +2116,11 @@ IMPLEMENT_PROPERTY_ACCESSOR( UObject* )
 IMPLEMENT_PROPERTY_ACCESSOR( const UObject* )
 IMPLEMENT_PROPERTY_ACCESSOR( FAssetData )
 
+FPropertyAccess::Result FPropertyHandleBase::SetValue( const TCHAR* InValue, EPropertyValueSetFlags::Type Flags )
+{
+	return FPropertyAccess::Fail;
+}
+
 FPropertyHandleBase::FPropertyHandleBase( TSharedPtr<FPropertyNode> PropertyNode, FNotifyHook* NotifyHook, TSharedPtr<IPropertyUtilities> PropertyUtilities )
 	: Implementation( MakeShareable( new FPropertyValueImpl( PropertyNode, NotifyHook, PropertyUtilities ) ) )
 {
@@ -3560,6 +3565,11 @@ FPropertyAccess::Result FPropertyHandleString::SetValue( const FString& NewValue
 	return Implementation->SetValueAsString( NewValue, Flags );
 }
 
+FPropertyAccess::Result FPropertyHandleString::SetValue( const TCHAR* NewValue, EPropertyValueSetFlags::Type Flags )
+{
+	return Implementation->SetValueAsString( NewValue, Flags );
+}
+
 FPropertyAccess::Result FPropertyHandleString::GetValue( FName& OutValue ) const
 {
 	void* PropValue = nullptr;
@@ -4415,6 +4425,11 @@ FPropertyAccess::Result FPropertyHandleText::SetValue(const FString& NewValue, E
 	return SetValue(FText::FromString(NewValue), Flags);
 }
 
+FPropertyAccess::Result FPropertyHandleText::SetValue(const TCHAR* NewValue, EPropertyValueSetFlags::Type Flags)
+{
+	return SetValue(FText::FromString(NewValue), Flags);
+}
+
 // Sets
 
 bool FPropertyHandleSet::Supports(TSharedRef<FPropertyNode> PropertyNode)
diff --git a/Engine/Source/Editor/PropertyEditor/Private/PropertyHandleImpl.h b/Engine/Source/Editor/PropertyEditor/Private/PropertyHandleImpl.h
--- a/Engine/Source/Editor/PropertyEditor/Private/PropertyHandleImpl.h
+++ b/Engine/Source/Editor/PropertyEditor/Private/PropertyHandleImpl.h
@@ -416,6 +416,8 @@ public:
 	DECLARE_PROPERTY_ACCESSOR( FAssetData )
 	virtual FPropertyAccess::Result GetValueData(void*& OutAddress) const override;
 
+	virtual FPropertyAccess::Result SetValue( const TCHAR* InValue, EPropertyValueSetFlags::Type Flags = EPropertyValueSetFlags::DefaultFlags ) override;
+
 	/** IPropertyHandle interface */
 	virtual bool IsValidHandle() const override;
 	virtual FText GetPropertyDisplayName() const override;
@@ -579,6 +581,7 @@ public:
 	static bool Supports( TSharedRef<FPropertyNode> PropertyNode );
 	virtual FPropertyAccess::Result GetValue( FString& OutValue ) const override;
 	virtual FPropertyAccess::Result SetValue( const FString& InValue, EPropertyValueSetFlags::Type Flags = EPropertyValueSetFlags::DefaultFlags ) override;
+	virtual FPropertyAccess::Result SetValue( const TCHAR* InValue, EPropertyValueSetFlags::Type Flags = EPropertyValueSetFlags::DefaultFlags ) override;
 
 	virtual FPropertyAccess::Result GetValue( FName& OutValue ) const override;
 	virtual FPropertyAccess::Result SetValue( const FName& InValue, EPropertyValueSetFlags::Type Flags = EPropertyValueSetFlags::DefaultFlags ) override;
@@ -671,6 +674,7 @@ public:
 	virtual FPropertyAccess::Result SetValue(const FText& InValue, EPropertyValueSetFlags::Type Flags = EPropertyValueSetFlags::DefaultFlags) override;
 
 	virtual FPropertyAccess::Result SetValue(const FString& InValue, EPropertyValueSetFlags::Type Flags = EPropertyValueSetFlags::DefaultFlags) override;
+	virtual FPropertyAccess::Result SetValue(const TCHAR* InValue, EPropertyValueSetFlags::Type Flags = EPropertyValueSetFlags::DefaultFlags) override;
 };
 
 class FPropertyHandleSet : public FPropertyHandleBase, public IPropertyHandleSet
diff --git a/Engine/Source/Editor/PropertyEditor/Public/PropertyHandle.h b/Engine/Source/Editor/PropertyEditor/Public/PropertyHandle.h
--- a/Engine/Source/Editor/PropertyEditor/Public/PropertyHandle.h
+++ b/Engine/Source/Editor/PropertyEditor/Public/PropertyHandle.h
@@ -296,7 +296,8 @@ public:
 	virtual FPropertyAccess::Result SetValue( UObject* const& InValue,  EPropertyValueSetFlags::Type Flags = EPropertyValueSetFlags::DefaultFlags ) = 0;
 	virtual FPropertyAccess::Result SetValue( const UObject* const& InValue,  EPropertyValueSetFlags::Type Flags = EPropertyValueSetFlags::DefaultFlags ) = 0;
 	virtual FPropertyAccess::Result SetValue( const FAssetData& InValue,  EPropertyValueSetFlags::Type Flags = EPropertyValueSetFlags::DefaultFlags ) = 0;
-	
+	virtual FPropertyAccess::Result SetValue( const TCHAR* InValue, EPropertyValueSetFlags::Type Flags = EPropertyValueSetFlags::DefaultFlags ) = 0;
+
 	/**
 	 * Called to manually notify root objects that this property is about to change
 	 * This does not need to be called when SetValue functions are used since it will be called automatically
