PR #5389: UE-67811: Hide NativeClass column in Bulk Edit table (Contributed by [REDACTED_PERSON])

diff --git a/Engine/Source/Editor/PropertyEditor/Private/Presentation/PropertyTable/PropertyTable.cpp b/Engine/Source/Editor/PropertyEditor/Private/Presentation/PropertyTable/PropertyTable.cpp
--- a/Engine/Source/Editor/PropertyEditor/Private/Presentation/PropertyTable/PropertyTable.cpp
+++ b/Engine/Source/Editor/PropertyEditor/Private/Presentation/PropertyTable/PropertyTable.cpp
@@ -1283,7 +1283,8 @@ void FPropertyTable::UpdateColumns()
 			{
 				TWeakObjectPtr< UProperty > Property = *PropertyIter;
 
-				if ( PropertyIter->HasAnyPropertyFlags(CPF_AssetRegistrySearchable) )
+				// Don't expose CPF_NativeAccessSpecifierProtected and CPF_NativeAccessSpecifierPrivate properties
+				if ( PropertyIter->HasAllPropertyFlags(CPF_NativeAccessSpecifierPublic | CPF_AssetRegistrySearchable) )
 				{
 					TArray< TSharedRef< IPropertyTableColumn > > MapResults;
 
