Added check/x to hover tooltip for PropertyEditorClass field to show if dragged class is valid

diff --git a/Engine/Source/Editor/PropertyEditor/Private/UserInterface/PropertyEditor/SPropertyEditorClass.cpp b/Engine/Source/Editor/PropertyEditor/Private/UserInterface/PropertyEditor/SPropertyEditorClass.cpp
--- a/Engine/Source/Editor/PropertyEditor/Private/UserInterface/PropertyEditor/SPropertyEditorClass.cpp
+++ b/Engine/Source/Editor/PropertyEditor/Private/UserInterface/PropertyEditor/SPropertyEditorClass.cpp
@@ -272,6 +272,67 @@ void SPropertyEditorClass::SendToObjects(const FString& NewValue)
 	}
 }
 
+void SPropertyEditorClass::OnDragEnter(const FGeometry& MyGeometry, const FDragDropEvent& DragDropEvent)
+{
+	TSharedPtr<FAssetDragDropOp> UnloadedClassOp = DragDropEvent.GetOperationAs<FAssetDragDropOp>();
+	if (UnloadedClassOp.IsValid())
+	{
+		bool bAllAssetWereLoaded = true;
+
+		FString AssetPath;
+		FString PathName;
+
+		// Find the class/blueprint path
+		if (UnloadedClassOp->HasAssets())
+		{
+			AssetPath = UnloadedClassOp->GetAssets()[0].ObjectPath.ToString();
+		}
+		else if (UnloadedClassOp->HasAssetPaths())
+		{
+			AssetPath = UnloadedClassOp->GetAssetPaths()[0];
+		}
+
+		// Check to see if the asset can be found, otherwise load it.
+		UObject* Object = FindObject<UObject>(nullptr, *AssetPath);
+		if (Object == nullptr)
+		{
+			// Load the package.
+			GWarn->BeginSlowTask(LOCTEXT("OnDrop_LoadPackage", "Fully Loading Package For Drop"), true, false);
+
+			Object = LoadObject<UObject>(nullptr, *AssetPath);
+
+			GWarn->EndSlowTask();
+		}
+
+		if (UClass* Class = Cast<UClass>(Object))
+		{
+			// This was pointing to a class directly
+			UnloadedClassOp->SetToolTip(FText::GetEmpty(), FEditorStyle::GetBrush(TEXT("Graph.ConnectorFeedback.OK")));
+		}
+		else if (UBlueprint* Blueprint = Cast<UBlueprint>(Object))
+		{
+			if (Blueprint->GeneratedClass)
+			{
+				// This was pointing to a blueprint, get generated class
+				UnloadedClassOp->SetToolTip(FText::GetEmpty(), FEditorStyle::GetBrush(TEXT("Graph.ConnectorFeedback.OK")));
+			}
+		}
+		else
+		{
+			UnloadedClassOp->SetToolTip(FText::GetEmpty(), FEditorStyle::GetBrush(TEXT("Graph.ConnectorFeedback.Error")));
+		}
+	}
+}
+
+void SPropertyEditorClass::OnDragLeave(const FDragDropEvent& DragDropEvent)
+{
+	TSharedPtr<FAssetDragDropOp> UnloadedClassOp = DragDropEvent.GetOperationAs<FAssetDragDropOp>();
+	if (UnloadedClassOp.IsValid())
+	{
+		UnloadedClassOp->ResetToDefaultToolTip();
+	}
+}
+
 FReply SPropertyEditorClass::OnDrop(const FGeometry& MyGeometry, const FDragDropEvent& DragDropEvent)
 {
 	TSharedPtr<FClassDragDropOp> ClassOperation = DragDropEvent.GetOperationAs<FClassDragDropOp>();
diff --git a/Engine/Source/Editor/PropertyEditor/Private/UserInterface/PropertyEditor/SPropertyEditorClass.h b/Engine/Source/Editor/PropertyEditor/Private/UserInterface/PropertyEditor/SPropertyEditorClass.h
--- a/Engine/Source/Editor/PropertyEditor/Private/UserInterface/PropertyEditor/SPropertyEditorClass.h
+++ b/Engine/Source/Editor/PropertyEditor/Private/UserInterface/PropertyEditor/SPropertyEditorClass.h
@@ -65,6 +65,8 @@ public:
 private:
 	void SendToObjects(const FString& NewValue);
 
+	virtual void OnDragEnter(const FGeometry& MyGeometry, const FDragDropEvent& DragDropEvent) override;
+	virtual void OnDragLeave(const FDragDropEvent& DragDropEvent) override;
 	virtual FReply OnDrop(const FGeometry& MyGeometry, const FDragDropEvent& DragDropEvent) override;
 
 	/** 
