Fixes for array handling in details panel when multiple actors are selected

diff --git a/Engine/Source/Editor/PropertyEditor/Private/PropertyEditorHelpers.cpp b/Engine/Source/Editor/PropertyEditor/Private/PropertyEditorHelpers.cpp
--- a/Engine/Source/Editor/PropertyEditor/Private/PropertyEditorHelpers.cpp
+++ b/Engine/Source/Editor/PropertyEditor/Private/PropertyEditorHelpers.cpp
@@ -941,10 +941,27 @@ namespace PropertyEditorHelpers
 		TAttribute<bool>::FGetter IsPropertyButtonEnabledDelegate = TAttribute<bool>::FGetter::CreateStatic(&IsPropertyButtonEnabled, WeakPropertyEditor);
 		TAttribute<bool> IsEnabledAttribute = TAttribute<bool>::Create( IsPropertyButtonEnabledDelegate );
 
+		TAttribute<bool> IsAddEnabledAttribute = TAttribute<bool>::Create(
+			TAttribute<bool>::FGetter::CreateLambda([WeakPropertyEditor, PropertyEditor]() {
+			if (WeakPropertyEditor.IsValid())
+			{
+				UProperty* Property = WeakPropertyEditor.Pin()->GetProperty();
+				// Check for multiple array selections with mismatched values
+				UArrayProperty* ArrayProperty = Cast<UArrayProperty>(Property);
+				if (ArrayProperty)
+				{
+					FString ArrayString;
+					FPropertyAccess::Result GetValResult = PropertyEditor->GetPropertyHandle()->GetValueAsDisplayString(ArrayString);
+					return IsPropertyButtonEnabled(WeakPropertyEditor) && GetValResult == FPropertyAccess::Success;
+				}
+			}
+			return IsPropertyButtonEnabled(WeakPropertyEditor);
+		}));
+
 		switch( ButtonType )
 		{
 		case EPropertyButton::Add:
-			NewButton = PropertyCustomizationHelpers::MakeAddButton( FSimpleDelegate::CreateSP( PropertyEditor, &FPropertyEditor::AddItem ), FText(), IsEnabledAttribute );
+			NewButton = PropertyCustomizationHelpers::MakeAddButton( FSimpleDelegate::CreateSP( PropertyEditor, &FPropertyEditor::AddItem ), FText(), IsAddEnabledAttribute);
 			break;
 
 		case EPropertyButton::Empty:
diff --git a/Engine/Source/Editor/PropertyEditor/Private/PropertyHandleImpl.cpp b/Engine/Source/Editor/PropertyEditor/Private/PropertyHandleImpl.cpp
--- a/Engine/Source/Editor/PropertyEditor/Private/PropertyHandleImpl.cpp
+++ b/Engine/Source/Editor/PropertyEditor/Private/PropertyHandleImpl.cpp
@@ -1690,51 +1690,50 @@ void FPropertyValueImpl::MoveElementTo(int32 OriginalIndex, int32 NewIndex)
 		FReadAddressList ReadAddresses;
 		void* Addr = nullptr;
 		ParentNode->GetReadAddress(!!ParentNode->HasNodeFlags(EPropertyNodeFlags::SingleSelectOnly), ReadAddresses);
-		if (ReadAddresses.Num())
+		for (int32 i = 0; i < ReadAddresses.Num(); ++i)
 		{
-			Addr = ReadAddresses.GetAddress(0);
-		}
-
-		if (Addr)
-		{
-			FScriptArrayHelper	ArrayHelper(ArrayProperty, Addr);
-			int32 Index = ChildNodePtr->GetArrayIndex();
+			void* Addr = ReadAddresses.GetAddress(i);
+			if (Addr)
+			{
+				FScriptArrayHelper	ArrayHelper(ArrayProperty, Addr);
+				int32 Index = ChildNodePtr->GetArrayIndex();
 
-			// List of top level objects sent to the PropertyChangedEvent
-			TArray<const UObject*> TopLevelObjects;
+				// List of top level objects sent to the PropertyChangedEvent
+				TArray<const UObject*> TopLevelObjects;
 
-			UObject* Obj = ObjectNode ? ObjectNode->GetUObject(0) : nullptr;
-			if (Obj)
-			{
-				if ((Obj->HasAnyFlags(RF_ClassDefaultObject | RF_ArchetypeObject) ||
-					(Obj->HasAnyFlags(RF_DefaultSubObject) && Obj->GetOuter()->HasAnyFlags(RF_ClassDefaultObject | RF_ArchetypeObject))))
+				UObject* Obj = ObjectNode ? ObjectNode->GetUObject(0) : nullptr;
+				if (Obj)
 				{
-					FString OrgArrayContent;
-					ArrayProperty->ExportText_Direct(OrgArrayContent, Addr, Addr, nullptr, 0);
+					if ((Obj->HasAnyFlags(RF_ClassDefaultObject | RF_ArchetypeObject) ||
+						(Obj->HasAnyFlags(RF_DefaultSubObject) && Obj->GetOuter()->HasAnyFlags(RF_ClassDefaultObject | RF_ArchetypeObject))))
+					{
+						FString OrgArrayContent;
+						ArrayProperty->ExportText_Direct(OrgArrayContent, Addr, Addr, nullptr, 0);
 
-					ChildNodePtr->PropagateContainerPropertyChange(Obj, OrgArrayContent, EPropertyArrayChangeType::Insert, Index);
-				}
+						ChildNodePtr->PropagateContainerPropertyChange(Obj, OrgArrayContent, EPropertyArrayChangeType::Insert, Index);
+					}
 
-				TopLevelObjects.Add(Obj);
-			}
+					TopLevelObjects.Add(Obj);
+				}
 
-			ArrayHelper.InsertValues(Index, 1);
+				ArrayHelper.InsertValues(Index, 1);
 
-			//set up indices for the coming events
-			TArray< TMap<FString, int32> > ArrayIndicesPerObject;
-			for (int32 ObjectIndex = 0; ObjectIndex < ReadAddresses.Num(); ++ObjectIndex)
-			{
-				//add on array index so we can tell which entry just changed
-				ArrayIndicesPerObject.Add(TMap<FString, int32>());
-				FPropertyValueImpl::GenerateArrayIndexMapToObjectNode(ArrayIndicesPerObject[ObjectIndex], ChildNodePtr);
-			}
+				//set up indices for the coming events
+				TArray< TMap<FString, int32> > ArrayIndicesPerObject;
+				for (int32 ObjectIndex = 0; ObjectIndex < ReadAddresses.Num(); ++ObjectIndex)
+				{
+					//add on array index so we can tell which entry just changed
+					ArrayIndicesPerObject.Add(TMap<FString, int32>());
+					FPropertyValueImpl::GenerateArrayIndexMapToObjectNode(ArrayIndicesPerObject[ObjectIndex], ChildNodePtr);
+				}
 
-			FPropertyChangedEvent ChangeEvent(ParentNode->GetProperty(), EPropertyChangeType::ArrayAdd, &TopLevelObjects);
-			ChangeEvent.SetArrayIndexPerObject(ArrayIndicesPerObject);
+				FPropertyChangedEvent ChangeEvent(ParentNode->GetProperty(), EPropertyChangeType::ArrayAdd, &TopLevelObjects);
+				ChangeEvent.SetArrayIndexPerObject(ArrayIndicesPerObject);
 
-			if (PropertyUtilities.IsValid())
-			{
-				ChildNodePtr->FixPropertiesInEvent(ChangeEvent);
+				if (PropertyUtilities.IsValid())
+				{
+					ChildNodePtr->FixPropertiesInEvent(ChangeEvent);
+				}
 			}
 		}
 	}
diff --git a/Engine/Source/Editor/PropertyEditor/Private/UserInterface/PropertyEditor/SPropertyEditorArray.h b/Engine/Source/Editor/PropertyEditor/Private/UserInterface/PropertyEditor/SPropertyEditorArray.h
--- a/Engine/Source/Editor/PropertyEditor/Private/UserInterface/PropertyEditor/SPropertyEditorArray.h
+++ b/Engine/Source/Editor/PropertyEditor/Private/UserInterface/PropertyEditor/SPropertyEditorArray.h
@@ -67,6 +67,13 @@ public:
 private:
 	FText GetArrayTextValue() const
 	{
+		FString ArrayString;
+		FPropertyAccess::Result GetValResult = PropertyEditor->GetPropertyHandle()->GetValueAsDisplayString(ArrayString);
+
+		if (GetValResult == FPropertyAccess::MultipleValues)
+		{
+			return LOCTEXT("MultipleValues", "Multiple Values");
+		}
 		return FText::Format( LOCTEXT("NumArrayItemsFmt", "{0} Array elements"), FText::AsNumber(PropertyEditor->GetPropertyNode()->GetNumChildNodes()) );
 	}
 
