Fix for shadowed variable

diff --git a/Engine/Source/Editor/PropertyEditor/Private/PropertyHandleImpl.cpp b/Engine/Source/Editor/PropertyEditor/Private/PropertyHandleImpl.cpp
--- a/Engine/Source/Editor/PropertyEditor/Private/PropertyHandleImpl.cpp
+++ b/Engine/Source/Editor/PropertyEditor/Private/PropertyHandleImpl.cpp
@@ -1688,7 +1688,6 @@ void FPropertyValueImpl::MoveElementTo(int32 OriginalIndex, int32 NewIndex)
 
 
 		FReadAddressList ReadAddresses;
-		void* Addr = nullptr;
 		ParentNode->GetReadAddress(!!ParentNode->HasNodeFlags(EPropertyNodeFlags::SingleSelectOnly), ReadAddresses);
 		for (int32 i = 0; i < ReadAddresses.Num(); ++i)
 		{
