Back out revision 77 : It was causing some CommitText issue, basically the OnValueChanged was being set as InteractiveChange (as it should) but was also calling a BeginTransaction, resulting in breaking all the Undo functionality.

diff --git a/Engine/Source/Editor/PropertyEditor/Private/PropertyHandleImpl.cpp b/Engine/Source/Editor/PropertyEditor/Private/PropertyHandleImpl.cpp
--- a/Engine/Source/Editor/PropertyEditor/Private/PropertyHandleImpl.cpp
+++ b/Engine/Source/Editor/PropertyEditor/Private/PropertyHandleImpl.cpp
@@ -3396,7 +3396,7 @@ FPropertyAccess::Result FPropertyHandleFloat::SetValue( const float& NewValue, E
 	// Clamp the value from any meta data ranges stored on the property value
 	float FinalValue = ClampValueFromMetaData<float>( NewValue, *Implementation->GetPropertyNode() );
 
-	const FString ValueStr = FString::Printf( TEXT("%.7f"), FinalValue );
+	const FString ValueStr = FString::Printf( TEXT("%f"), FinalValue );
 	Res = Implementation->ImportText( ValueStr, Flags );
 
 	return Res;
