Fixed container diff regression introduced in cl 4802047 - containers differences were being detected multiple times

diff --git a/Engine/Source/Editor/PropertyEditor/Private/SDetailsViewBase.cpp b/Engine/Source/Editor/PropertyEditor/Private/SDetailsViewBase.cpp
--- a/Engine/Source/Editor/PropertyEditor/Private/SDetailsViewBase.cpp
+++ b/Engine/Source/Editor/PropertyEditor/Private/SDetailsViewBase.cpp
@@ -111,14 +111,17 @@ void SDetailsViewBase::HideFilterArea(bool bHide)
 
 static void GetPropertiesInOrderDisplayedRecursive(const TArray< TSharedRef<FDetailTreeNode> >& TreeNodes, TArray< FPropertyPath > &OutLeaves)
 {
-	for (auto& TreeNode : TreeNodes)
+	for( const TSharedRef<FDetailTreeNode>& TreeNode : TreeNodes )
 	{
-		if (TreeNode->IsLeaf())
+		const bool bIsPropertyRoot = TreeNode->IsLeaf() ||
+			(	TreeNode->GetPropertyNode() && 
+				TreeNode->GetPropertyNode()->GetProperty() );
+		if( bIsPropertyRoot )
 		{
 			FPropertyPath Path = TreeNode->GetPropertyPath();
 			// Some leaf nodes are not associated with properties, specifically the collision presets.
 			// @todo doc: investigate what we can do about this, result is that for these fields
-			// we can't highlight hte property in the diff tool.
+			// we can't highlight the property in the diff tool.
 			if( Path.GetNumProperties() != 0 )
 			{
 				OutLeaves.Push(Path);
