Fixed differences not being detected in properties that are nested within more than one category
#rnx

diff --git a/Engine/Source/Editor/Kismet/Private/SBlueprintDiff.cpp b/Engine/Source/Editor/Kismet/Private/SBlueprintDiff.cpp
--- a/Engine/Source/Editor/Kismet/Private/SBlueprintDiff.cpp
+++ b/Engine/Source/Editor/Kismet/Private/SBlueprintDiff.cpp
@@ -192,12 +192,12 @@ FCDODiffControl::FCDODiffControl(
 				bDiffers = true;
 				// if there are any nested differences associated with PropertyIdentifier, add those
 				// as well:
-				OrderedProperties.Push(&Difference);
+				OrderedProperties.AddUnique(&Difference);
 			}
 			else if (Difference.Identifier.IsSubPropertyMatch(PropertyIdentifier))
 			{
 				bDiffers = true;
-				OrderedProperties.Push(&Difference);
+				OrderedProperties.AddUnique(&Difference);
 			}
 		}
 		return bDiffers;
diff --git a/Engine/Source/Editor/PropertyEditor/Private/DetailItemNode.h b/Engine/Source/Editor/PropertyEditor/Private/DetailItemNode.h
--- a/Engine/Source/Editor/PropertyEditor/Private/DetailItemNode.h
+++ b/Engine/Source/Editor/PropertyEditor/Private/DetailItemNode.h
@@ -68,7 +68,7 @@ public:
 	virtual FPropertyPath GetPropertyPath() const override;
 	virtual void SetIsHighlighted(bool bInIsHighlighted) override { bIsHighlighted = bInIsHighlighted; }
 	virtual bool IsHighlighted() const override { return bIsHighlighted; }
-	virtual bool IsLeaf() override { return true; }
+	virtual bool IsLeaf() override { return Children.Num() == 0;  }
 	virtual TAttribute<bool> IsPropertyEditingEnabled() const override { return IsParentEnabled; }
 	virtual TSharedPtr<FPropertyNode> GetPropertyNode() const override;
 	virtual TSharedPtr<FComplexPropertyNode> GetExternalRootPropertyNode() const override;
