Converted PropagateContainerPropertyChange to use UProperty::Identical when determining if container data needs to be migrated

This is consistent with other places in the editor, and avoids issues with certain properties that have ExportText forms that may appear different, even though the underlying data is considered identical.

diff --git a/Engine/Source/Editor/PropertyEditor/Private/PropertyHandleImpl.cpp b/Engine/Source/Editor/PropertyEditor/Private/PropertyHandleImpl.cpp
--- a/Engine/Source/Editor/PropertyEditor/Private/PropertyHandleImpl.cpp
+++ b/Engine/Source/Editor/PropertyEditor/Private/PropertyHandleImpl.cpp
@@ -1114,11 +1114,8 @@ void FPropertyValueImpl::AddChild()
 							(Obj->HasAnyFlags(RF_DefaultSubObject) && Obj->GetOuter()->HasAnyFlags(RF_ClassDefaultObject | RF_ArchetypeObject))) &&
 							!FApp::IsGame())
 						{
-							FString OrgContent;
-							NodeProperty->ExportText_Direct(OrgContent, Addr, Addr, nullptr, 0);
-
 							TMap<UObject*, bool> PropagationResult;
-							PropertyNodePin->PropagateContainerPropertyChange(Obj, OrgContent, EPropertyArrayChangeType::Add, -1, &PropagationResult);
+							PropertyNodePin->PropagateContainerPropertyChange(Obj, Addr, EPropertyArrayChangeType::Add, -1, &PropagationResult);
 							PropagationResultPerObject.Add(MoveTemp(PropagationResult));
 						}
 
@@ -1229,9 +1226,7 @@ void FPropertyValueImpl::ClearChildren()
 							(Obj->HasAnyFlags(RF_DefaultSubObject) && Obj->GetOuter()->HasAnyFlags(RF_ClassDefaultObject | RF_ArchetypeObject))) &&
 							!FApp::IsGame())
 						{
-							FString OrgContent;
-							NodeProperty->ExportText_Direct(OrgContent, Addr, Addr, nullptr, 0);
-							PropertyNodePin->PropagateContainerPropertyChange(Obj, OrgContent, EPropertyArrayChangeType::Clear, -1);
+							PropertyNodePin->PropagateContainerPropertyChange(Obj, Addr, EPropertyArrayChangeType::Clear, -1);
 						}
 
 						TopLevelObjects.Add(Obj);
@@ -1384,11 +1379,8 @@ void FPropertyValueImpl::InsertChild( TSharedPtr<FPropertyNode> ChildNodeToInser
 				(Obj->HasAnyFlags(RF_DefaultSubObject) && Obj->GetOuter()->HasAnyFlags(RF_ClassDefaultObject | RF_ArchetypeObject))) &&
 				!FApp::IsGame())
 			{
-				FString OrgArrayContent;
-				ArrayProperty->ExportText_Direct(OrgArrayContent, Addr, Addr, nullptr, 0);
-
 				TMap<UObject*, bool> PropagationResult;
-				ChildNodePtr->PropagateContainerPropertyChange(Obj, OrgArrayContent, EPropertyArrayChangeType::Insert, Index, &PropagationResult);
+				ChildNodePtr->PropagateContainerPropertyChange(Obj, Addr, EPropertyArrayChangeType::Insert, Index, &PropagationResult);
 				PropagationResultPerObject.Add(MoveTemp(PropagationResult));
 			}
 
@@ -1481,11 +1473,8 @@ void FPropertyValueImpl::DeleteChild( TSharedPtr<FPropertyNode> ChildNodeToDelet
 						(Obj->HasAnyFlags(RF_DefaultSubObject) && Obj->GetOuter()->HasAnyFlags(RF_ClassDefaultObject | RF_ArchetypeObject))) &&
 						!FApp::IsGame())
 					{
-						FString OrgContent;
-						Cast<UProperty>(NodeProperty->GetOuter())->ExportText_Direct(OrgContent, Address, Address, nullptr, 0);
-
 						TMap<UObject*, bool> PropagationResult;
-						ChildNodePtr->PropagateContainerPropertyChange(Obj, OrgContent, EPropertyArrayChangeType::Delete, Index, &PropagationResult);
+						ChildNodePtr->PropagateContainerPropertyChange(Obj, Address, EPropertyArrayChangeType::Delete, Index, &PropagationResult);
 
 						PropagationResultPerObject.Add(MoveTemp(PropagationResult));
 					}
@@ -1622,9 +1611,7 @@ void FPropertyValueImpl::SwapChildren( TSharedPtr<FPropertyNode> FirstChildNode,
 					if ((Obj->HasAnyFlags(RF_ClassDefaultObject | RF_ArchetypeObject) || (Obj->HasAnyFlags(RF_DefaultSubObject) && Obj->GetOuter()->HasAnyFlags(RF_ClassDefaultObject | RF_ArchetypeObject))) 
 						&& !FApp::IsGame())
 					{
-						FString OrgContent;
-						Cast<UProperty>(FirstNodeProperty->GetOuter())->ExportText_Direct(OrgContent, Address, Address, nullptr, 0);
-						FirstChildNodePtr->PropagateContainerPropertyChange(Obj, OrgContent, EPropertyArrayChangeType::Swap, FirstIndex, nullptr, SecondIndex);
+						FirstChildNodePtr->PropagateContainerPropertyChange(Obj, Address, EPropertyArrayChangeType::Swap, FirstIndex, nullptr, SecondIndex);
 					}
 
 					TopLevelObjects.Add(Obj);
@@ -1706,10 +1693,7 @@ void FPropertyValueImpl::MoveElementTo(int32 OriginalIndex, int32 NewIndex)
 					if ((Obj->HasAnyFlags(RF_ClassDefaultObject | RF_ArchetypeObject) ||
 						(Obj->HasAnyFlags(RF_DefaultSubObject) && Obj->GetOuter()->HasAnyFlags(RF_ClassDefaultObject | RF_ArchetypeObject))))
 					{
-						FString OrgArrayContent;
-						ArrayProperty->ExportText_Direct(OrgArrayContent, Addr, Addr, nullptr, 0);
-
-						ChildNodePtr->PropagateContainerPropertyChange(Obj, OrgArrayContent, EPropertyArrayChangeType::Insert, Index);
+						ChildNodePtr->PropagateContainerPropertyChange(Obj, Addr, EPropertyArrayChangeType::Insert, Index);
 					}
 
 					TopLevelObjects.Add(Obj);
@@ -1777,9 +1761,7 @@ void FPropertyValueImpl::MoveElementTo(int32 OriginalIndex, int32 NewIndex)
 							if ((Obj->HasAnyFlags(RF_ClassDefaultObject | RF_ArchetypeObject) ||
 								(Obj->HasAnyFlags(RF_DefaultSubObject) && Obj->GetOuter()->HasAnyFlags(RF_ClassDefaultObject | RF_ArchetypeObject))))
 							{
-								FString OrgContent;
-								NodeProperty->ExportText_Direct(OrgContent, Addr, Addr, nullptr, 0);
-								PropertyNodePin->PropagateContainerPropertyChange(Obj, OrgContent, EPropertyArrayChangeType::Add, -1);
+								PropertyNodePin->PropagateContainerPropertyChange(Obj, Addr, EPropertyArrayChangeType::Add, -1);
 							}
 
 							TopLevelObjects.Add(Obj);
@@ -1852,9 +1834,7 @@ void FPropertyValueImpl::MoveElementTo(int32 OriginalIndex, int32 NewIndex)
 						if ((Obj->HasAnyFlags(RF_ClassDefaultObject | RF_ArchetypeObject) ||
 							(Obj->HasAnyFlags(RF_DefaultSubObject) && Obj->GetOuter()->HasAnyFlags(RF_ClassDefaultObject | RF_ArchetypeObject))))
 						{
-							FString OrgContent;
-							Cast<UProperty>(FirstNodeProperty->GetOuter())->ExportText_Direct(OrgContent, Address, Address, nullptr, 0);
-							FirstChildNodePtr->PropagateContainerPropertyChange(Obj, OrgContent, EPropertyArrayChangeType::Swap, FirstIndex, nullptr, SecondIndex);
+							FirstChildNodePtr->PropagateContainerPropertyChange(Obj, Address, EPropertyArrayChangeType::Swap, FirstIndex, nullptr, SecondIndex);
 						}
 
 						TopLevelObjects.Add(Obj);
@@ -1931,9 +1911,7 @@ void FPropertyValueImpl::MoveElementTo(int32 OriginalIndex, int32 NewIndex)
 							(Obj->HasAnyFlags(RF_DefaultSubObject) && Obj->GetOuter()->HasAnyFlags(RF_ClassDefaultObject | RF_ArchetypeObject))) &&
 							!FApp::IsGame())
 						{
-							FString OrgContent;
-							Cast<UProperty>(NodeProperty->GetOuter())->ExportText_Direct(OrgContent, Address, Address, nullptr, 0);
-							ChildNodePtr->PropagateContainerPropertyChange(Obj, OrgContent, EPropertyArrayChangeType::Delete, Index);
+							ChildNodePtr->PropagateContainerPropertyChange(Obj, Address, EPropertyArrayChangeType::Delete, Index);
 						}
 
 						TopLevelObjects.Add(Obj);
@@ -2015,11 +1993,8 @@ void FPropertyValueImpl::DuplicateChild( TSharedPtr<FPropertyNode> ChildNodeToDu
 				(Obj->HasAnyFlags(RF_DefaultSubObject) && Obj->GetOuter()->HasAnyFlags(RF_ClassDefaultObject | RF_ArchetypeObject))) &&
 				!FApp::IsGame())
 			{
-				FString OrgContent;
-				Cast<UProperty>(NodeProperty->GetOuter())->ExportText_Direct(OrgContent, Addr, Addr, nullptr, 0);
-
 				TMap<UObject*, bool> PropagationResult;
-				ChildNodePtr->PropagateContainerPropertyChange(Obj, OrgContent, EPropertyArrayChangeType::Duplicate, Index, &PropagationResult);
+				ChildNodePtr->PropagateContainerPropertyChange(Obj, Addr, EPropertyArrayChangeType::Duplicate, Index, &PropagationResult);
 				PropagationResultPerObject.Add(MoveTemp(PropagationResult));
 			}
 
diff --git a/Engine/Source/Editor/PropertyEditor/Private/PropertyNode.cpp b/Engine/Source/Editor/PropertyEditor/Private/PropertyNode.cpp
--- a/Engine/Source/Editor/PropertyEditor/Private/PropertyNode.cpp
+++ b/Engine/Source/Editor/PropertyEditor/Private/PropertyNode.cpp
@@ -2395,8 +2395,10 @@ bool FPropertyNode::IsFilterAcceptable(const TArray<FString>& InAcceptableNames,
 	return bCompleteMatchFound;
 }
 
-void FPropertyNode::PropagateContainerPropertyChange( UObject* ModifiedObject, const FString& OriginalContainerContent, EPropertyArrayChangeType::Type ChangeType, int32 Index, TMap<UObject*, bool>* PropagationResult, int32 SwapIndex /*= INDEX_NONE*/)
+void FPropertyNode::PropagateContainerPropertyChange( UObject* ModifiedObject, const void* OriginalContainerAddr, EPropertyArrayChangeType::Type ChangeType, int32 Index, TMap<UObject*, bool>* PropagationResult, int32 SwapIndex /*= INDEX_NONE*/)
 {
+	check(OriginalContainerAddr);
+
 	UProperty* NodeProperty = GetProperty();
 	UArrayProperty* ArrayProperty = NULL;
 	USetProperty* SetProperty = NULL;
@@ -2479,10 +2481,8 @@ void FPropertyNode::PropagateContainerPropertyChange( UObject* ModifiedObject, c
 
 			if (Addr != nullptr)
 			{
-				FString OriginalContent;
-				ConvertedProperty->ExportText_Direct(OriginalContent, Addr, Addr, nullptr, PPF_None);
-
-				bool bIsDefaultContainerContent = OriginalContent == OriginalContainerContent;
+				checkf(OriginalContainerAddr != Addr, TEXT("PropagateContainerPropertyChange tried to propagate a change onto itself!"));
+				const bool bIsDefaultContainerContent = ConvertedProperty->Identical(OriginalContainerAddr, Addr);
 
 				// Return instance changes result to caller
 				if (PropagationResult != nullptr)
diff --git a/Engine/Source/Editor/PropertyEditor/Private/PropertyNode.h b/Engine/Source/Editor/PropertyEditor/Private/PropertyNode.h
--- a/Engine/Source/Editor/PropertyEditor/Private/PropertyNode.h
+++ b/Engine/Source/Editor/PropertyEditor/Private/PropertyNode.h
@@ -528,11 +528,11 @@ public:
 	 * Propagates the property change of a container property to all instances of an archetype
 	 *
 	 * @param	ModifiedObject				Object which property has been modified
-	 * @param	OriginalContainerContent	Original content of the container before the modification ( as returned by ExportText_Direct )
+	 * @param	OriginalContainerAddr		Original address holding the container value before the modification
 	 * @param	ChangeType					In which way was the container modified
 	 * @param	Index						Index of the modified item
 	 */
-	void PropagateContainerPropertyChange(UObject* ModifiedObject, const FString& OriginalContainerContent,
+	void PropagateContainerPropertyChange(UObject* ModifiedObject, const void* OriginalContainerAddr,
 		EPropertyArrayChangeType::Type ChangeType, int32 Index, TMap<UObject*, bool>* PropagationResult = nullptr, int32 SwapIndex = INDEX_NONE);
 
 	/** Broadcasts when a property value changes */
