FPropertyTable doesn't handle a partial refresh properly, so I wipe and rebuild the entire table. I also fixed a bug where PurgeInvalidObjectNodes had inverted logic, the code to rebuild the table needs some refactoring

diff --git a/Engine/Source/Editor/PropertyEditor/Private/Presentation/PropertyTable/PropertyTable.cpp b/Engine/Source/Editor/PropertyEditor/Private/Presentation/PropertyTable/PropertyTable.cpp
--- a/Engine/Source/Editor/PropertyEditor/Private/Presentation/PropertyTable/PropertyTable.cpp
+++ b/Engine/Source/Editor/PropertyEditor/Private/Presentation/PropertyTable/PropertyTable.cpp
@@ -35,6 +35,14 @@ FPropertyTable::FPropertyTable()
 {
 }
 
+FPropertyTable::~FPropertyTable()
+{
+	if (GEditor && ObjectsReplacedHandle.IsValid())
+	{
+		GEditor->OnObjectsReplaced().Remove(ObjectsReplacedHandle);
+	}
+}
+
 void FPropertyTable::Tick()
 {
 	// Execute any deferred actions
@@ -214,6 +222,22 @@ void FPropertyTable::RemoveRow( const TSharedRef< class IPropertyTableRow >& Row
 	}
 }
 
+void FPropertyTable::ResetTable()
+{
+	Rows.Reset();
+	Columns.Reset();
+	SourceObjectPropertyNodes.Reset();
+	SelectedColumns.Reset();
+	SelectedRows.Reset();
+	SelectedCells.Reset();
+	StartingCellSelectionRange = nullptr;
+	EndingCellSelectionRange = nullptr;
+	CurrentRow = nullptr;
+	CurrentCell = nullptr;
+	CurrentColumn = nullptr;
+	LastClickedCell = nullptr;
+}
+
 void FPropertyTable::PurgeInvalidObjectNodes()
 {
 	TArray< TSharedRef< FObjectPropertyNode > > ValidNodes;
@@ -221,7 +245,7 @@ void FPropertyTable::PurgeInvalidObjectNodes()
 	{
 		const TWeakObjectPtr< UObject > Object = NodeIter.Key();
 
-		if ( !Object.IsValid() )
+		if ( Object.IsValid() )
 		{
 			ValidNodes.Add( NodeIter.Value() );
 		}
@@ -264,7 +288,7 @@ TArray< FPropertyInfo > FPropertyTable::GetPossibleExtensionsForPath( const TSha
 		{
 			const FPropertyInfo& Info = *ExtensionIter;
 
-			if ( Info.ArrayIndex == INDEX_NONE &&
+			if ( Info.ArrayIndex == INDEX_NONE && Info.Property.IsValid() && 
 				 ( Info.Property->IsA( UStructProperty::StaticClass() ) ||
 				   Info.Property->IsA( UArrayProperty::StaticClass() ) ) )
 			{
@@ -993,19 +1017,56 @@ void FPropertyTable::SetObjects( const TArray< TWeakObjectPtr< UObject > >& Obje
 
 	UpdateColumns();
 	UpdateRows();
+
+	// Bind to object delegates, we can't do this at construction because the shared pointer isn't set up yet
+	if (GEditor && !ObjectsReplacedHandle.IsValid())
+	{
+		ObjectsReplacedHandle = GEditor->OnObjectsReplaced().AddSP(this, &FPropertyTable::OnObjectsReplaced);
+	}
 }
 
 void FPropertyTable::SetObjects( const TArray< UObject* >& Objects )
 {
-	SourceObjects.Empty();
+	TArray<TWeakObjectPtr<UObject>> WeakObjects;
 	
-	for( auto ObjectIter = Objects.CreateConstIterator(); ObjectIter; ++ObjectIter )
+	for(UObject* Object : Objects)
 	{
-		SourceObjects.Add( *ObjectIter );
+		WeakObjects.Add(Object);
 	}
 
-	UpdateColumns();
-	UpdateRows();
+	SetObjects(WeakObjects);
+}
+
+void FPropertyTable::OnObjectsReplaced(const TMap<UObject*, UObject*>& ReplacementMap)
+{
+	bool bChangedAny = false;
+
+	// Fix up any objects in our source list that were replaced, then refresh UI
+	for (int32 i = 0; i < SourceObjects.Num(); i++)
+	{
+		UObject* SourceObject = SourceObjects[i].Get(true);
+		UObject* ReplacedObject = ReplacementMap.FindRef(SourceObject);
+
+		if (ReplacedObject && ReplacedObject != SourceObject)
+		{
+			SourceObjectPropertyNodes.Remove(SourceObject);
+			SourceObjects[i] = ReplacedObject;
+
+			bChangedAny = true;
+		}
+	}
+
+	if (bChangedAny)
+	{
+		// The update functions do not correctly handle partial reconstruction, so delete everything and reconstruct
+		ResetTable();
+
+		UpdateRows();
+		UpdateColumns();
+
+		// This reset our selection so broadcast it
+		SelectionChanged.Broadcast();
+	}
 }
 
 void FPropertyTable::SetShowRowHeader( const bool InShowRowHeader ) 
diff --git a/Engine/Source/Editor/PropertyEditor/Private/Presentation/PropertyTable/PropertyTable.h b/Engine/Source/Editor/PropertyEditor/Private/Presentation/PropertyTable/PropertyTable.h
--- a/Engine/Source/Editor/PropertyEditor/Private/Presentation/PropertyTable/PropertyTable.h
+++ b/Engine/Source/Editor/PropertyEditor/Private/Presentation/PropertyTable/PropertyTable.h
@@ -15,6 +15,7 @@ class FPropertyTable : public TSharedFromThis< FPropertyTable >, public IPropert
 public: 
 
 	FPropertyTable();
+	virtual ~FPropertyTable();
 
 	virtual void Tick() override;
 
@@ -167,6 +168,9 @@ private:
 
 	void PurgeInvalidObjectNodes();
 
+	void ResetTable();
+
+	void OnObjectsReplaced(const TMap<UObject*, UObject*>& ReplacementMap);
 
 private:
 
@@ -218,5 +222,7 @@ private:
 
 	/** The Orientation of this table, I.e. do we swap columns and rows */
 	EPropertyTableOrientation::Type Orientation;
+
+	FDelegateHandle ObjectsReplacedHandle;
 };
 
diff --git a/Engine/Source/Editor/PropertyEditor/Private/PropertyEditorToolkit.cpp b/Engine/Source/Editor/PropertyEditor/Private/PropertyEditorToolkit.cpp
--- a/Engine/Source/Editor/PropertyEditor/Private/PropertyEditorToolkit.cpp
+++ b/Engine/Source/Editor/PropertyEditor/Private/PropertyEditorToolkit.cpp
@@ -461,6 +461,13 @@ void FPropertyEditorToolkit::GridSelectionChanged()
 {
 	TArray< TWeakObjectPtr< UObject > > SelectedObjects;
 	PropertyTable->GetSelectedTableObjects( SelectedObjects );
+
+	if (SelectedObjects.Num() == 0)
+	{
+		// If none are selected, show all of them to match the initial open behavior
+		SelectedObjects = PropertyTable->GetSelectedObjects();
+	}
+
 	PropertyTree->SetObjectArray( SelectedObjects );
 
 	const TSet< TSharedRef< IPropertyTableRow > > SelectedRows = PropertyTable->GetSelectedRows();
diff --git a/Engine/Source/Editor/UnrealEd/Private/Toolkits/SimpleAssetEditor.cpp b/Engine/Source/Editor/UnrealEd/Private/Toolkits/SimpleAssetEditor.cpp
--- a/Engine/Source/Editor/UnrealEd/Private/Toolkits/SimpleAssetEditor.cpp
+++ b/Engine/Source/Editor/UnrealEd/Private/Toolkits/SimpleAssetEditor.cpp
@@ -38,6 +38,7 @@ const FName FSimpleAssetEditor::SimpleEditorAppIdentifier( TEXT( "GenericEditorA
 FSimpleAssetEditor::~FSimpleAssetEditor()
 {
 	GEditor->GetEditorSubsystem<UImportSubsystem>()->OnAssetPostImport.RemoveAll(this);
+	GEditor->OnObjectsReplaced().RemoveAll(this);
 
 	DetailsView.Reset();
 	PropertiesTab.Reset();
@@ -51,7 +52,8 @@ void FSimpleAssetEditor::InitEditor( const EToolkitMode::Type Mode, const TShare
 	const bool bIsLockable = false;
 
 	EditingObjects = ObjectsToEdit;
-	GEditor->GetEditorSubsystem<UImportSubsystem>()->OnAssetPostImport.AddRaw(this, &FSimpleAssetEditor::HandleAssetPostImport);
+	GEditor->GetEditorSubsystem<UImportSubsystem>()->OnAssetPostImport.AddSP(this, &FSimpleAssetEditor::HandleAssetPostImport);
+	GEditor->OnObjectsReplaced().AddSP(this, &FSimpleAssetEditor::OnObjectsReplaced);
 
 	FPropertyEditorModule& PropertyEditorModule = FModuleManager::GetModuleChecked<FPropertyEditorModule>( "PropertyEditor" );
 	const FDetailsViewArgs DetailsViewArgs( bIsUpdatable, bIsLockable, true, FDetailsViewArgs::ObjectsUseNameArea, false );
@@ -251,6 +253,29 @@ void FSimpleAssetEditor::HandleAssetPostImport(UFactory* InFactory, UObject* InO
 	}
 }
 
+void FSimpleAssetEditor::OnObjectsReplaced(const TMap<UObject*, UObject*>& ReplacementMap)
+{
+	bool bChangedAny = false;
+
+	// Refresh our details view if one of the objects replaced was in the map. This gets called before the reinstance GC fixup, so we might as well fixup EditingObjects now too
+	for (int32 i = 0; i < EditingObjects.Num(); i++)
+	{
+		UObject* SourceObject = EditingObjects[i];
+		UObject* ReplacedObject = ReplacementMap.FindRef(SourceObject);
+
+		if (ReplacedObject && ReplacedObject != SourceObject)
+		{
+			EditingObjects[i] = ReplacedObject;
+			bChangedAny = true;
+		}
+	}
+
+	if (bChangedAny)
+	{
+		DetailsView->SetObjects(EditingObjects);
+	}
+}
+
 FString FSimpleAssetEditor::GetWorldCentricTabPrefix() const
 {
 	return LOCTEXT("WorldCentricTabPrefix", "Generic Asset ").ToString();
diff --git a/Engine/Source/Editor/UnrealEd/Public/Toolkits/SimpleAssetEditor.h b/Engine/Source/Editor/UnrealEd/Public/Toolkits/SimpleAssetEditor.h
--- a/Engine/Source/Editor/UnrealEd/Public/Toolkits/SimpleAssetEditor.h
+++ b/Engine/Source/Editor/UnrealEd/Public/Toolkits/SimpleAssetEditor.h
@@ -52,6 +52,9 @@ private:
 	/** Handles when an asset is imported */
 	void HandleAssetPostImport(class UFactory* InFactory, UObject* InObject);
 
+	/** Called when objects need to be swapped out for new versions, like after a blueprint recompile. */
+	void OnObjectsReplaced(const TMap<UObject*, UObject*>& ReplacementMap);
+
 	/** Dockable tab for properties */
 	TSharedPtr< SDockableTab > PropertiesTab;
 
