Fix non-unity
#rnx

diff --git a/Engine/Source/Editor/PropertyEditor/Private/Presentation/PropertyTable/PropertyTable.cpp b/Engine/Source/Editor/PropertyEditor/Private/Presentation/PropertyTable/PropertyTable.cpp
--- a/Engine/Source/Editor/PropertyEditor/Private/Presentation/PropertyTable/PropertyTable.cpp
+++ b/Engine/Source/Editor/PropertyEditor/Private/Presentation/PropertyTable/PropertyTable.cpp
@@ -3,7 +3,7 @@
 #include "Presentation/PropertyTable/PropertyTable.h"
 #include "Misc/FeedbackContext.h"
 #include "Editor/EditorPerProjectUserSettings.h"
-
+#include "Editor.h"
 
 #include "Presentation/PropertyTable/PropertyTableColumn.h"
 #include "Presentation/PropertyTable/PropertyTableRow.h"
@@ -11,7 +11,6 @@
 #include "Presentation/PropertyTable/PropertyTableObjectNameColumn.h"
 #include "Presentation/PropertyTable/PropertyTablePropertyNameColumn.h"
 
-
 #define LOCTEXT_NAMESPACE "PropertyTable"
 
 FPropertyTable::FPropertyTable()
