Also fix TSoftClassPtr to have an add new pointer like the other two and remove the duplicate reset to defaults button on soft object path
Based on PR #5406

diff --git a/Engine/Source/Editor/DetailCustomizations/Private/SoftObjectPathCustomization.cpp b/Engine/Source/Editor/DetailCustomizations/Private/SoftObjectPathCustomization.cpp
--- a/Engine/Source/Editor/DetailCustomizations/Private/SoftObjectPathCustomization.cpp
+++ b/Engine/Source/Editor/DetailCustomizations/Private/SoftObjectPathCustomization.cpp
@@ -93,6 +93,9 @@ void FSoftObjectPathCustomization::CustomizeHeader( TSharedRef<IPropertyHandle>
 		.OnShouldFilterAsset( AssetFilter )
 		.AllowClear( bAllowClear )
 	];
+
+	// This avoids making duplicate reset boxes
+	StructPropertyHandle->MarkResetToDefaultCustomized();
 }
 
 void FSoftObjectPathCustomization::CustomizeChildren(TSharedRef<IPropertyHandle> InStructPropertyHandle, IDetailChildrenBuilder& StructBuilder, IPropertyTypeCustomizationUtils& StructCustomizationUtils)
diff --git a/Engine/Source/Editor/PropertyEditor/Private/PropertyEditorHelpers.cpp b/Engine/Source/Editor/PropertyEditor/Private/PropertyEditorHelpers.cpp
--- a/Engine/Source/Editor/PropertyEditor/Private/PropertyEditorHelpers.cpp
+++ b/Engine/Source/Editor/PropertyEditor/Private/PropertyEditorHelpers.cpp
@@ -621,13 +621,13 @@ namespace PropertyEditorHelpers
 		return (NodeProperty->IsA<UObjectPropertyBase>() || NodeProperty->IsA<UInterfaceProperty>()) && (!bUsingAssetPicker || !SPropertyEditorAsset::Supports(NodeProperty));
 	}
 	
-	static bool IsSoftObjectPath( const UProperty* Property )
+	bool IsSoftObjectPath( const UProperty* Property )
 	{
 		const UStructProperty* StructProp = Cast<const UStructProperty>( Property );
 		return StructProp && StructProp->Struct == TBaseStructure<FSoftObjectPath>::Get();
 	}
 
-	static bool IsSoftClassPath( const UProperty* Property )
+	bool IsSoftClassPath( const UProperty* Property )
 	{
 		const UStructProperty* StructProp = Cast<const UStructProperty>(Property);
 		return StructProp && StructProp->Struct == TBaseStructure<FSoftClassPath>::Get();
@@ -733,35 +733,37 @@ namespace PropertyEditorHelpers
 		// Handle a class property.
 
 		UClassProperty* ClassProp = Cast<UClassProperty>(NodeProperty);
-		if( ClassProp || IsSoftClassPath(NodeProperty))
+		USoftClassProperty* SoftClassProp = Cast<USoftClassProperty>(NodeProperty);
+		if( ClassProp || SoftClassProp || IsSoftClassPath(NodeProperty))
 		{
 			OutRequiredButtons.Add( EPropertyButton::Use );			
 			OutRequiredButtons.Add( EPropertyButton::Browse );
 
-			UClass* Class = (ClassProp ? ClassProp->MetaClass : FEditorClassUtils::GetClassFromString(NodeProperty->GetMetaData("MetaClass")));
-
-			if (Class && FKismetEditorUtilities::CanCreateBlueprintOfClass(Class) && !NodeProperty->HasMetaData("DisallowCreateNew"))
+			UClass* Class = nullptr;
+			if (ClassProp)
+			{
+				Class = ClassProp->MetaClass;
+			}
+			else if (SoftClassProp)
+			{
+				Class = SoftClassProp->MetaClass;
+			}
+			else
 			{
-				OutRequiredButtons.Add( EPropertyButton::NewBlueprint );
+				Class = NodeProperty->GetOwnerProperty()->GetClassMetaData(TEXT("MetaClass"));
 			}
 
-			if( !(NodeProperty->PropertyFlags & CPF_NoClear) )
+			if (Class && FKismetEditorUtilities::CanCreateBlueprintOfClass(Class) && !NodeProperty->HasMetaData("DisallowCreateNew"))
 			{
-				OutRequiredButtons.Add( EPropertyButton::Clear );
+				OutRequiredButtons.Add(EPropertyButton::NewBlueprint);
 			}
-		}
-		else if (NodeProperty->IsA<USoftClassProperty>() )
-		{
-			OutRequiredButtons.Add( EPropertyButton::Use );
-			
-			OutRequiredButtons.Add( EPropertyButton::Browse );
 
 			if( !(NodeProperty->PropertyFlags & CPF_NoClear) )
 			{
 				OutRequiredButtons.Add( EPropertyButton::Clear );
 			}
 		}
-
+		
 		if( OuterArrayProp )
 		{
 			if( PropertyNode->HasNodeFlags(EPropertyNodeFlags::SingleSelectOnly) && !(OuterArrayProp->PropertyFlags & CPF_EditFixedSize) )
@@ -795,6 +797,11 @@ namespace PropertyEditorHelpers
 		PropertyEditorHelpers::MakeRequiredPropertyButtons( PropertyEditor, OutButtons, ButtonsToIgnore, bUsingAssetPicker );
 	}
 
+	static bool IsPropertyButtonEnabled(TWeakPtr<FPropertyNode> PropertyNode)
+	{
+		return PropertyNode.IsValid() ? !PropertyNode.Pin()->IsEditConst() : false;
+	}
+
 	TSharedRef<SWidget> MakePropertyReorderHandle(const TSharedRef<FPropertyNode>& PropertyNode, TSharedPtr<SDetailSingleItemRow> InParentRow)
 	{
 		TSharedRef<SArrayRowHandle> Handle = SNew(SArrayRowHandle)
@@ -893,12 +900,6 @@ namespace PropertyEditorHelpers
 		return SelectionPathName;
 	}
 
-	
-	static bool IsPropertyButtonEnabled( TWeakPtr<FPropertyNode> PropertyNode )
-	{
-		return PropertyNode.IsValid() ? !PropertyNode.Pin()->IsEditConst() : false;
-	}
-
 	/**
 	 * A helper method that checks to see if the editor's current selection is 
 	 * compatible with the specified property.
diff --git a/Engine/Source/Editor/PropertyEditor/Private/PropertyEditorHelpers.h b/Engine/Source/Editor/PropertyEditor/Private/PropertyEditorHelpers.h
--- a/Engine/Source/Editor/PropertyEditor/Private/PropertyEditorHelpers.h
+++ b/Engine/Source/Editor/PropertyEditor/Private/PropertyEditorHelpers.h
@@ -116,8 +116,6 @@ private:
 
 namespace PropertyEditorHelpers
 {
-
-	static bool IsPropertyButtonEnabled(TWeakPtr<FPropertyNode> PropertyNode);
 	/**
 	 * Returns whether or not a property is a built in struct property like a vector or color
 	 *
@@ -148,6 +146,7 @@ namespace PropertyEditorHelpers
 	 * @param InPropertyNode	The property node containing the property to check
 	 */
 	bool IsStaticArray( const FPropertyNode& InPropertyNode );
+	
 	/**
 	 * Returns whether or not a property a static array
 	 *
@@ -155,6 +154,16 @@ namespace PropertyEditorHelpers
 	 */
 	bool IsDynamicArray( const FPropertyNode& InPropertyNode );
 
+	/**
+	 * Returns true if this is an FSoftObjectPath and should be treated like a TSoftObjectPtr
+	 */
+	bool IsSoftObjectPath(const UProperty* Property);
+
+	/**
+	 * Returns true if this is an FSoftClassPath and should be treated like a TSoftClassPtr
+	 */
+	bool IsSoftClassPath(const UProperty* Property);
+
 	/**
 	 * Gets the array parent of a property if it is in a dynamic or static array
 	 */
diff --git a/Engine/Source/Editor/PropertyEditor/Private/PropertyHandleImpl.cpp b/Engine/Source/Editor/PropertyEditor/Private/PropertyHandleImpl.cpp
--- a/Engine/Source/Editor/PropertyEditor/Private/PropertyHandleImpl.cpp
+++ b/Engine/Source/Editor/PropertyEditor/Private/PropertyHandleImpl.cpp
@@ -24,6 +24,7 @@
 #include "UObject/EnumProperty.h"
 #include "IDetailPropertyRow.h"
 #include "ObjectEditorUtils.h"
+#include "PropertyEditorHelpers.h"
 
 #define LOCTEXT_NAMESPACE "PropertyHandleImplementation"
 
@@ -3596,7 +3597,8 @@ bool FPropertyHandleObject::Supports( TSharedRef<FPropertyNode> PropertyNode )
 		return false;
 	}
 
-	return Property->IsA( UObjectPropertyBase::StaticClass() ) || Property->IsA(UInterfaceProperty::StaticClass());
+	return Property->IsA(UObjectPropertyBase::StaticClass()) || Property->IsA(UInterfaceProperty::StaticClass()) 
+		|| PropertyEditorHelpers::IsSoftClassPath(Property) || PropertyEditorHelpers::IsSoftObjectPath(Property);
 }
 
 FPropertyAccess::Result FPropertyHandleObject::GetValue( UObject*& OutValue ) const
@@ -3623,6 +3625,14 @@ FPropertyAccess::Result FPropertyHandleObject::GetValue( const UObject*& OutValu
 			const FScriptInterface& ScriptInterface = InterfaceProp->GetPropertyValue(PropValue);
 			OutValue = ScriptInterface.GetObject();
 		}
+		else
+		{
+			// This is a struct path, get the path string and search for the object
+			FString ObjectPathString;
+			Res = Implementation->GetValueAsString(ObjectPathString);
+			FSoftObjectPath ObjectPath(ObjectPathString);
+			OutValue = ObjectPath.ResolveObject();
+		}
 	}
 
 	return Res;
@@ -3717,6 +3727,15 @@ FPropertyAccess::Result FPropertyHandleObject::SetValueFromFormattedString(const
 			{
 				InterfaceThatMustBeImplemented = InterfaceProperty->InterfaceClass;
 			}
+			else if (PropertyEditorHelpers::IsSoftClassPath(NodeProperty))
+			{
+				RequiredClass = NodeProperty->GetOwnerProperty()->GetClassMetaData(TEXT("MetaClass"));
+				InterfaceThatMustBeImplemented = NodeProperty->GetOwnerProperty()->GetClassMetaData(TEXT("MustImplement"));
+			}
+			else if (PropertyEditorHelpers::IsSoftObjectPath(NodeProperty))
+			{
+				// No metaclass, allowedclasses is the only filter
+			}
 			else
 			{
 				return FPropertyAccess::Fail;
@@ -3841,13 +3860,27 @@ FPropertyAccess::Result FPropertyHandleObject::SetObjectValueFromSelection()
 		UInterfaceProperty* IntProp = Cast<UInterfaceProperty>(NodeProperty);
 		UClassProperty* ClassProp = Cast<UClassProperty>(NodeProperty);
 		USoftClassProperty* SoftClassProperty = Cast<USoftClassProperty>(NodeProperty);
-		UClass* const InterfaceThatMustBeImplemented = ObjProp ? ObjProp->GetOwnerProperty()->GetClassMetaData(TEXT("MustImplement")) : nullptr;
+		UClass* const InterfaceThatMustBeImplemented = NodeProperty ? NodeProperty->GetOwnerProperty()->GetClassMetaData(TEXT("MustImplement")) : nullptr;
 
-		if (ClassProp || SoftClassProperty)
+		if (ClassProp || SoftClassProperty || PropertyEditorHelpers::IsSoftClassPath(NodeProperty))
 		{
 			FEditorDelegates::LoadSelectedAssetsIfNeeded.Broadcast();
 
-			const UClass* const SelectedClass = GEditor->GetFirstSelectedClass(ClassProp ? ClassProp->MetaClass : SoftClassProperty->MetaClass);
+			UClass* RequiredClass = nullptr;
+			if (ClassProp)
+			{
+				RequiredClass = ClassProp->MetaClass;
+			}
+			else if (SoftClassProperty)
+			{
+				RequiredClass = SoftClassProperty->MetaClass;
+			}
+			else
+			{
+				RequiredClass = NodeProperty->GetOwnerProperty()->GetClassMetaData(TEXT("MetaClass"));
+			}
+
+			const UClass* const SelectedClass = GEditor->GetFirstSelectedClass(RequiredClass);
 			if (SelectedClass)
 			{
 				if (!InterfaceThatMustBeImplemented || SelectedClass->ImplementsInterface(InterfaceThatMustBeImplemented))
