Fix details view so setting highlighted property to empty will clear highlight instead of highlighting a property-less customization

diff --git a/Engine/Source/Editor/PropertyEditor/Private/SDetailsViewBase.cpp b/Engine/Source/Editor/PropertyEditor/Private/SDetailsViewBase.cpp
--- a/Engine/Source/Editor/PropertyEditor/Private/SDetailsViewBase.cpp
+++ b/Engine/Source/Editor/PropertyEditor/Private/SDetailsViewBase.cpp
@@ -145,7 +145,12 @@ TArray< FPropertyPath > SDetailsViewBase::GetPropertiesInOrderDisplayed() const
 // @return populates OutNodes with the leaf node corresponding to property as the first entry in the list (e.g. [leaf, parent, grandparent]):
 static void FindTreeNodeFromPropertyRecursive( const TArray< TSharedRef<FDetailTreeNode> >& Nodes, const FPropertyPath& Property, TArray< TSharedPtr< FDetailTreeNode > >& OutNodes )
 {
-	for (auto& TreeNode : Nodes)
+	if (Property == FPropertyPath())
+	{
+		return;
+	}
+
+	for (const TSharedRef<FDetailTreeNode>& TreeNode : Nodes)
 	{
 		if (TreeNode->IsLeaf())
 		{
@@ -171,7 +176,7 @@ static void FindTreeNodeFromPropertyRecursive( const TArray< TSharedRef<FDetailT
 
 void SDetailsViewBase::HighlightProperty(const FPropertyPath& Property)
 {
-	auto PrevHighlightedNodePtr = CurrentlyHighlightedNode.Pin();
+	TSharedPtr<FDetailTreeNode> PrevHighlightedNodePtr = CurrentlyHighlightedNode.Pin();
 	if (PrevHighlightedNodePtr.IsValid())
 	{
 		PrevHighlightedNodePtr->SetIsHighlighted(false);
