PR #3973: Highlight array drag element (Contributed by [REDACTED_PERSON])

diff --git a/Engine/Source/Editor/EditorStyle/Private/SlateEditorStyle.cpp b/Engine/Source/Editor/EditorStyle/Private/SlateEditorStyle.cpp
--- a/Engine/Source/Editor/EditorStyle/Private/SlateEditorStyle.cpp
+++ b/Engine/Source/Editor/EditorStyle/Private/SlateEditorStyle.cpp
@@ -3634,6 +3634,7 @@ void FSlateEditorStyle::FStyle::SetupPropertyEditorStyles()
 		Set( "DetailsView.CategoryMiddle", new IMAGE_BRUSH( "PropertyView/DetailCategoryMiddle", FVector2D( 16, 16 ) ) );
 		Set( "DetailsView.CategoryMiddle_Hovered", new IMAGE_BRUSH( "PropertyView/DetailCategoryMiddle_Hovered", FVector2D( 16, 16 ) ) );
 		Set( "DetailsView.CategoryMiddle_Highlighted", new BOX_BRUSH( "Common/TextBox_Special_Active", FMargin(8.0f/32.0f) ) );
+		Set( "DetailsView.CategoryMiddle_Active", new BOX_BRUSH( "Common/TextBox_Special_Active", FMargin(8.0f/32.0f), SelectionColor_Pressed ) );
 		
 		Set( "DetailsView.PropertyIsFavorite", new IMAGE_BRUSH("PropertyView/Favorites_Enabled", Icon12x12));
 		Set( "DetailsView.PropertyIsNotFavorite", new IMAGE_BRUSH("PropertyView/Favorites_Disabled", Icon12x12));
diff --git a/Engine/Source/Editor/PropertyEditor/Private/SDetailSingleItemRow.cpp b/Engine/Source/Editor/PropertyEditor/Private/SDetailSingleItemRow.cpp
--- a/Engine/Source/Editor/PropertyEditor/Private/SDetailSingleItemRow.cpp
+++ b/Engine/Source/Editor/PropertyEditor/Private/SDetailSingleItemRow.cpp
@@ -685,6 +685,10 @@ const FSlateBrush* SDetailSingleItemRow::GetBorderImage() const
 	{
 		return FEditorStyle::GetBrush("DetailsView.CategoryMiddle_Highlighted");
 	}
+	else if (bIsDragDropObject)
+	{
+		return FEditorStyle::GetBrush("DetailsView.CategoryMiddle_Active");
+	}
 	else if (IsHovered() && !bIsHoveredDragTarget)
 	{
 		return FEditorStyle::GetBrush("DetailsView.CategoryMiddle_Hovered");
@@ -786,6 +790,11 @@ bool SDetailSingleItemRow::IsHighlighted() const
 	return OwnerTreeNode.Pin()->IsHighlighted();
 }
 
+void SDetailSingleItemRow::SetIsDragDrop(bool bInIsDragDrop)
+{
+	bIsDragDropObject = bInIsDragDrop;
+}
+
 void SArrayRowHandle::Construct(const FArguments& InArgs)
 {
 	ParentRow = InArgs._ParentRow;
@@ -800,12 +809,10 @@ FReply SArrayRowHandle::OnDragDetected(const FGeometry& MyGeometry, const FPoint
 {
 	if (MouseEvent.IsMouseButtonDown(EKeys::LeftMouseButton))
 	{
-		{	
-			TSharedPtr<FDragDropOperation> DragDropOp = CreateDragDropOperation(ParentRow.Pin());
-			if (DragDropOp.IsValid())
-			{
-				return FReply::Handled().BeginDragDrop(DragDropOp.ToSharedRef());
-			}
+		TSharedPtr<FDragDropOperation> DragDropOp = CreateDragDropOperation(ParentRow.Pin());
+		if (DragDropOp.IsValid())
+		{
+			return FReply::Handled().BeginDragDrop(DragDropOp.ToSharedRef());
 		}
 	}
 
@@ -819,3 +826,46 @@ TSharedPtr<FArrayRowDragDropOp> SArrayRowHandle::CreateDragDropOperation(TShared
 
 	return Operation;
 }
+
+FArrayRowDragDropOp::FArrayRowDragDropOp(TSharedPtr<SDetailSingleItemRow> InRow)
+{
+	Row = InRow;
+
+	TSharedPtr<SDetailSingleItemRow> RowPtr = nullptr;
+	if (Row.IsValid())
+	{
+		RowPtr = Row.Pin();
+		// mark row as being used for drag and drop
+		RowPtr->SetIsDragDrop(true);
+	}
+
+	DecoratorWidget = SNew(SBorder)
+		.Padding(8.f)
+		.BorderImage(FEditorStyle::GetBrush("Graph.ConnectorFeedback.Border"))
+		.Content()
+		[
+			SNew(SHorizontalBox)
+			+ SHorizontalBox::Slot()
+			.AutoWidth()
+			.VAlign(VAlign_Center)
+			[
+				SNew(STextBlock)
+				.Text(NSLOCTEXT("ArrayDragDrop", "PlaceRowHere", "Place Row Here"))
+			]
+		];
+
+	Construct();
+}
+
+void FArrayRowDragDropOp::OnDrop(bool bDropWasHandled, const FPointerEvent& MouseEvent)
+{
+	FDecoratedDragDropOp::OnDrop(bDropWasHandled, MouseEvent);
+
+	TSharedPtr<SDetailSingleItemRow> RowPtr = nullptr;
+	if (Row.IsValid())
+	{
+		RowPtr = Row.Pin();
+		// reset value
+		RowPtr->SetIsDragDrop(false);
+	}
+}
\ No newline at end of file
diff --git a/Engine/Source/Editor/PropertyEditor/Private/SDetailSingleItemRow.h b/Engine/Source/Editor/PropertyEditor/Private/SDetailSingleItemRow.h
--- a/Engine/Source/Editor/PropertyEditor/Private/SDetailSingleItemRow.h
+++ b/Engine/Source/Editor/PropertyEditor/Private/SDetailSingleItemRow.h
@@ -80,6 +80,8 @@ public:
 	 */
 	void Construct( const FArguments& InArgs, FDetailLayoutCustomization* InCustomization, bool bHasMultipleColumns, TSharedRef<FDetailTreeNode> InOwnerTreeNode, const TSharedRef<STableViewBase>& InOwnerTableView );
 
+	void SetIsDragDrop(bool bInIsDragDrop);
+
 protected:
 	virtual bool OnContextMenuOpening( FMenuBuilder& MenuBuilder ) override;
 private:
@@ -109,6 +111,7 @@ private:
 	FDetailColumnSizeData ColumnSizeData;
 	bool bAllowFavoriteSystem;
 	bool bIsHoveredDragTarget;
+	bool bIsDragDropObject;
 	TSharedPtr<FPropertyNode> SwappablePropertyNode;
 };
 
@@ -117,28 +120,12 @@ class FArrayRowDragDropOp : public FDecoratedDragDropOp
 public:
 	DRAG_DROP_OPERATOR_TYPE(FArrayRowDragDropOp, FDecoratedDragDropOp)
 
-	FArrayRowDragDropOp(TSharedPtr<SDetailSingleItemRow> InRow)
-	{
-		Row = InRow;
-		DecoratorWidget = SNew(SBorder)
-			.BorderImage(FEditorStyle::GetBrush("Graph.ConnectorFeedback.Border"))
-			.Content()
-			[
-				SNew(SHorizontalBox)
-				+ SHorizontalBox::Slot()
-				.AutoWidth()
-				.VAlign(VAlign_Center)
-				[
-					SNew(STextBlock)
-					.Text(NSLOCTEXT("ArrayDragDrop", "PlaceRowHere", "Place Row Here"))
-				]
-			];
-
-		Construct();
-	};
+	FArrayRowDragDropOp(TSharedPtr<SDetailSingleItemRow> InRow);
 
 	TSharedPtr<SWidget> DecoratorWidget;
 
+	virtual void OnDrop(bool bDropWasHandled, const FPointerEvent& MouseEvent) override;
+
 	virtual TSharedPtr<SWidget> GetDefaultDecorator() const override
 	{
 		return DecoratorWidget;
