Improved layout and appearance of per-platform overrides widget

Also fixed enabled/disabled state when used in static mesh editor.
Also moved layout widget into cpp file & updated dependencies and includes for IWYU.

diff --git a/Engine/Source/Editor/DetailCustomizations/Private/PerPlatformPropertyCustomization.cpp b/Engine/Source/Editor/DetailCustomizations/Private/PerPlatformPropertyCustomization.cpp
--- a/Engine/Source/Editor/DetailCustomizations/Private/PerPlatformPropertyCustomization.cpp
+++ b/Engine/Source/Editor/DetailCustomizations/Private/PerPlatformPropertyCustomization.cpp
@@ -14,6 +14,7 @@
 #include "PropertyHandle.h"
 #include "DetailLayoutBuilder.h"
 #include "SPerPlatformPropertiesWidget.h"
+#include "PlatformInfo.h"
 #include "ScopedTransaction.h"
 #include "IPropertyUtilities.h"
 #include "UObject/MetaData.h"
diff --git a/Engine/Source/Editor/PropertyEditor/Private/SPerPlatformPropertiesWidget.cpp b/Engine/Source/Editor/PropertyEditor/Private/SPerPlatformPropertiesWidget.cpp
new file mode 100644
--- /dev/null
+++ b/Engine/Source/Editor/PropertyEditor/Private/SPerPlatformPropertiesWidget.cpp
@@ -0,0 +1,274 @@
+// Copyright 1998-2019 Epic Games, Inc. All Rights Reserved.
+
+#include "SPerPlatformPropertiesWidget.h"
+#include "Layout/Margin.h"
+#include "Fonts/SlateFontInfo.h"
+#include "Styling/CoreStyle.h"
+#include "Framework/MultiBox/MultiBoxBuilder.h"
+#include "Widgets/Input/SButton.h"
+#include "Widgets/Input/SComboButton.h"
+#include "PlatformInfo.h"
+#include "Widgets/Layout/SWrapBox.h"
+#include "EditorStyleSet.h"
+#include "Widgets/Layout/SBox.h"
+#include "DetailLayoutBuilder.h"
+#include "Widgets/Images/SImage.h"
+
+
+void SPerPlatformPropertiesWidget::Construct(const typename SPerPlatformPropertiesWidget::FArguments& InArgs)
+{
+	this->OnGenerateWidget = InArgs._OnGenerateWidget;
+	this->OnAddPlatform = InArgs._OnAddPlatform;
+	this->OnRemovePlatform = InArgs._OnRemovePlatform;
+	this->PlatformOverrideNames = InArgs._PlatformOverrideNames;
+
+	ConstructChildren();
+
+	// this widget has to check platform count change from outside to ensure the widget represents latest update
+	RegisterActiveTimer(FMath::RandRange(2.f, 5.f), FWidgetActiveTimerDelegate::CreateSP(this, &SPerPlatformPropertiesWidget::CheckPlatformCount));
+}
+
+void SPerPlatformPropertiesWidget::ConstructChildren()
+{
+	TSharedPtr<SWrapBox> WrapBox;
+
+	TArray<FName> PlatformOverrides = PlatformOverrideNames.Get();
+	LastPlatformOverrideNames = PlatformOverrides.Num();
+
+	ChildSlot
+	.VAlign(VAlign_Center)
+	.HAlign(HAlign_Fill)
+	[
+		SAssignNew(WrapBox, SWrapBox)
+		.UseAllottedWidth(true)
+	];
+
+	if (OnGenerateWidget.IsBound())
+	{
+		// Build Platform menu
+		bAddedMenuItem = false;
+		FMenuBuilder AddPlatformMenuBuilder(true, nullptr, nullptr, true);
+
+		// Platform (group) names
+		const TArray<FName>& PlatformGroupNameArray = PlatformInfo::GetAllPlatformGroupNames();
+		const TArray<FName>& VanillaPlatformNameArray = PlatformInfo::GetAllVanillaPlatformNames();
+
+		// Sanitized platform names
+		TArray<FName> BasePlatformNameArray;
+		// Mapping from platform group name to individual platforms
+		TMultiMap<FName, FName> GroupToPlatform;
+						
+		// Create mapping from platform to platform groups and remove postfixes and invalid platform names
+		const TArray<FString> Filters = { TEXT("NoEditor"), TEXT("Client"), TEXT("Server"), TEXT("AllDesktop") };
+		for (const FName& PlatformName : VanillaPlatformNameArray)
+		{
+			const PlatformInfo::FPlatformInfo* PlatformInfo = PlatformInfo::FindPlatformInfo(PlatformName);
+			FString PlatformNameString = PlatformName.ToString();
+			for ( const FString& Filter : Filters)
+			{
+				const int32 Position = PlatformNameString.Find(Filter);
+				if (Position != INDEX_NONE)
+				{
+					PlatformNameString.RemoveAt(Position, Filter.Len());
+					break;
+				}
+			}
+
+			// Add filtered name if it isn't already set, and also add to group mapping
+			const FName FilteredName = FName(*PlatformNameString);				
+			if (PlatformNameString.Len() && !PlatformOverrides.Contains(FilteredName))
+			{	
+				BasePlatformNameArray.AddUnique(FilteredName);
+				GroupToPlatform.AddUnique(PlatformInfo->PlatformGroupName, FilteredName);
+			}
+		}
+
+		// Create section for platform groups 
+		const FName PlatformGroupSection(TEXT("PlatformGroupSection"));
+		AddPlatformMenuBuilder.BeginSection(PlatformGroupSection, FText::FromString(TEXT("Platform Groups")));
+		for (const FName& GroupName : PlatformGroupNameArray)
+		{				
+			if (!PlatformOverrides.Contains(GroupName))
+			{					
+				const FTextFormat Format = NSLOCTEXT("SPerPlatformPropertiesWidget", "AddOverrideGroupFor", "Add Override for Platforms part of the {0} Platform Group");
+				AddPlatformToMenu(GroupName, Format, AddPlatformMenuBuilder);
+				bAddedMenuItem = true;
+			}
+		}				
+		AddPlatformMenuBuilder.EndSection();
+			
+		for (const FName& GroupName : PlatformGroupNameArray)
+		{
+			// Create a section for each platform group and their respective platforms
+			AddPlatformMenuBuilder.BeginSection(GroupName, FText::FromName(GroupName));
+
+			TArray<FName> PlatformNames;
+			GroupToPlatform.MultiFind(GroupName, PlatformNames);
+
+			const FTextFormat Format = NSLOCTEXT("SPerPlatformPropertiesWidget", "AddOverrideFor", "Add Override specifically for {0}");
+			for (const FName& PlatformName : PlatformNames)
+			{
+				AddPlatformToMenu(PlatformName, Format, AddPlatformMenuBuilder);
+			}
+
+			bAddedMenuItem |= PlatformNames.Num() > 0;
+
+			AddPlatformMenuBuilder.EndSection();
+		}
+
+		// Default control
+		WrapBox->AddSlot()
+		[
+			MakePerPlatformWidget(NAME_None, NSLOCTEXT("SPerPlatformPropertiesWidget", "DefaultPlatform", "Default"), PlatformOverrides, AddPlatformMenuBuilder)
+		];
+
+		for (FName PlatformName : PlatformOverrides)
+		{
+			WrapBox->AddSlot()
+			[
+				MakePerPlatformWidget(PlatformName, FText::AsCultureInvariant(PlatformName.ToString()), PlatformOverrides, AddPlatformMenuBuilder)
+			];
+		}
+	}
+	else
+	{
+		WrapBox->AddSlot()
+		[
+			SNew(STextBlock)
+			.Text(NSLOCTEXT("SPerPlatformPropertiesWidget", "OnGenerateWidgetWarning", "No OnGenerateWidget() Provided"))
+			.ColorAndOpacity(FLinearColor::Red)
+		];
+	}
+}
+
+void SPerPlatformPropertiesWidget::AddPlatformToMenu(const FName& PlatformName, const FTextFormat Format, FMenuBuilder &AddPlatformMenuBuilder)
+{		
+	const FText MenuText = FText::Format(FText::FromString(TEXT("{0}")), FText::AsCultureInvariant(PlatformName.ToString()));
+	const FText MenuTooltipText = FText::Format(Format, FText::AsCultureInvariant(PlatformName.ToString()));
+	AddPlatformMenuBuilder.AddMenuEntry(
+		MenuText,
+		MenuTooltipText,
+		FSlateIcon(FEditorStyle::GetStyleSetName(), "PerPlatformWidget.AddPlatform"),
+		FUIAction(FExecuteAction::CreateSP(this, &SPerPlatformPropertiesWidget::AddPlatform, PlatformName))
+	);
+}
+
+void SPerPlatformPropertiesWidget::AddPlatform(FName PlatformName)
+{
+	if (OnAddPlatform.IsBound() && OnAddPlatform.Execute(PlatformName))
+	{
+		ConstructChildren();
+		Invalidate(EInvalidateWidget::LayoutAndVolatility);
+	}
+}
+
+FReply SPerPlatformPropertiesWidget::RemovePlatform(FName PlatformName)
+{
+	if (OnRemovePlatform.IsBound() && OnRemovePlatform.Execute(PlatformName))
+	{
+		ConstructChildren();
+		Invalidate(EInvalidateWidget::LayoutAndVolatility);
+	}
+	return FReply::Handled();
+}
+
+EActiveTimerReturnType SPerPlatformPropertiesWidget::CheckPlatformCount(double InCurrentTime, float InDeltaSeconds)
+{
+	// @fixme: the platform count is fixed and locally cached
+	// so if you change this outside of editor, this widget won't update
+	// this timer is the one checking to see if platform count has changed
+	// if so, it will reconstruct
+	TArray<FName> PlatformOverrides = PlatformOverrideNames.Get();
+	if (LastPlatformOverrideNames != PlatformOverrides.Num())
+	{
+		ConstructChildren();
+	}
+
+	return EActiveTimerReturnType::Continue;
+}
+
+
+TSharedRef<SWidget> SPerPlatformPropertiesWidget::MakePerPlatformWidget(FName InName, FText InDisplayText, const TArray<FName>& InPlatformOverrides, FMenuBuilder& InAddPlatformMenuBuilder)
+{
+	TSharedPtr<SHorizontalBox> HorizontalBox;
+
+	TSharedRef<SWidget> Widget = 
+		SNew(SBox)
+		.ToolTipText(InName == NAME_None ? 
+			NSLOCTEXT("SPerPlatformPropertiesWidget", "DefaultPlatformDesc", "This property can have per-platform or platform group overrides.\nThis is the default value used when no override has been set for a platform or platform group.") : 
+			FText::Format(NSLOCTEXT("SPerPlatformPropertiesWidget", "PerPlatformDesc", "Override for {0}"), InDisplayText))
+		.Padding(FMargin(0.0f, 2.0f, 4.0f, 2.0f))
+		.MinDesiredWidth(50.0f)
+		[
+			SNew(SVerticalBox)
+			+SVerticalBox::Slot()
+			.AutoHeight()
+			[
+				SAssignNew(HorizontalBox, SHorizontalBox)
+				+SHorizontalBox::Slot()
+				.FillWidth(1.0f)
+				.Padding(0.0f, 0.0f, 2.0f, 2.0f)
+				.VAlign(VAlign_Center)
+				[
+					SNew(STextBlock)
+					.Font(IDetailLayoutBuilder::GetDetailFont())
+					.Text(InDisplayText)
+				]
+			]
+			+SVerticalBox::Slot()
+			.AutoHeight()
+			[
+				OnGenerateWidget.Execute(InName)
+			]
+		];
+
+	if(InName != NAME_None)
+	{
+		HorizontalBox->AddSlot()
+			.AutoWidth()
+			.Padding(2.0f, 0.0f, 0.0f, 2.0f)
+			.VAlign(VAlign_Center)
+			[
+				SNew(SButton)
+				.ContentPadding(2.0f)
+				.ButtonStyle(FEditorStyle::Get(), "HoverHintOnly")
+				.OnClicked(this, &SPerPlatformPropertiesWidget::RemovePlatform, InName)
+				.ToolTipText(FText::Format(NSLOCTEXT("SPerPlatformPropertiesWidget", "RemoveOverrideFor", "Remove Override for {0}"), InDisplayText))
+				.ForegroundColor(FSlateColor::UseForeground())
+				.HAlign(HAlign_Center)
+				.VAlign(VAlign_Center)
+				.Content()
+				[
+					SNew(SImage)
+					.Image(FEditorStyle::GetBrush("Icons.Cross"))
+				]
+			];
+	}
+	else
+	{
+		HorizontalBox->AddSlot()
+			.AutoWidth()
+			.Padding(2.0f, 0.0f, 0.0f, 2.0f)
+			.VAlign(VAlign_Center)
+			[
+				SNew(SComboButton)
+				.Visibility_Lambda([this](){ return bAddedMenuItem ? EVisibility::Visible : EVisibility::Hidden; })
+				.ButtonStyle(FEditorStyle::Get(), "HoverHintOnly")
+				.ContentPadding(2.0f)
+				.ForegroundColor(FSlateColor::UseForeground())
+				.HasDownArrow(false)
+				.ButtonContent()
+				[
+					SNew(SImage)
+					.Image(FEditorStyle::GetBrush("PropertyWindow.Button_AddToArray"))
+				]
+				.MenuContent()
+				[
+					InAddPlatformMenuBuilder.MakeWidget()
+				]
+				.ToolTipText(NSLOCTEXT("SPerPlatformPropertiesWidget", "AddOverrideToolTip", "Add an override for a specific platform or platform group"))	
+			];
+	}
+
+	return Widget;
+}
\ No newline at end of file
diff --git a/Engine/Source/Editor/PropertyEditor/PropertyEditor.Build.cs b/Engine/Source/Editor/PropertyEditor/PropertyEditor.Build.cs
--- a/Engine/Source/Editor/PropertyEditor/PropertyEditor.Build.cs
+++ b/Engine/Source/Editor/PropertyEditor/PropertyEditor.Build.cs
@@ -61,6 +61,7 @@ public class PropertyEditor : ModuleRules
                 "RHI",
 				"ConfigEditor",
                 "SceneOutliner",
+				"DesktopPlatform",
             }
         );
 
diff --git a/Engine/Source/Editor/PropertyEditor/Public/SPerPlatformPropertiesWidget.h b/Engine/Source/Editor/PropertyEditor/Public/SPerPlatformPropertiesWidget.h
--- a/Engine/Source/Editor/PropertyEditor/Public/SPerPlatformPropertiesWidget.h
+++ b/Engine/Source/Editor/PropertyEditor/Public/SPerPlatformPropertiesWidget.h
@@ -3,23 +3,17 @@
 #pragma once
 
 #include "CoreMinimal.h"
-#include "Layout/Margin.h"
-#include "Fonts/SlateFontInfo.h"
 #include "Widgets/DeclarativeSyntaxSupport.h"
 #include "Widgets/SWidget.h"
 #include "Widgets/SCompoundWidget.h"
-#include "Styling/CoreStyle.h"
 #include "Framework/SlateDelegates.h"
-#include "Framework/MultiBox/MultiBoxBuilder.h"
-#include "Widgets/Input/SButton.h"
-#include "Widgets/Input/SComboButton.h"
-#include "PlatformInfo.h"
 
+class FMenuBuilder;
 
 /**
 * SPerPlatformPropertiesWidget
 */
-class SPerPlatformPropertiesWidget : public SCompoundWidget
+class PROPERTYEDITOR_API SPerPlatformPropertiesWidget : public SCompoundWidget
 {
 public:
 	typedef typename TSlateDelegates<FName>::FOnGenerateWidget FOnGenerateWidget;
@@ -43,262 +37,26 @@ public:
 	*
 	* @param	InArgs	The declaration data for this widget
 	*/
-	void Construct(const typename SPerPlatformPropertiesWidget::FArguments& InArgs)
-	{
-		this->OnGenerateWidget = InArgs._OnGenerateWidget;
-		this->OnAddPlatform = InArgs._OnAddPlatform;
-		this->OnRemovePlatform = InArgs._OnRemovePlatform;
-		this->PlatformOverrideNames = InArgs._PlatformOverrideNames;
+	void Construct(const typename SPerPlatformPropertiesWidget::FArguments& InArgs);
 
-		ConstructChildren();
-
-		// this widget has to check platform count change from outside to ensure the widget represents latest update
-		RegisterActiveTimer(FMath::RandRange(2.f, 5.f), FWidgetActiveTimerDelegate::CreateSP(this, &SPerPlatformPropertiesWidget::CheckPlatformCount));
-	}
-
-	void ConstructChildren()
-	{
-		TSharedPtr<SHorizontalBox> HorizontalBox;
-		TSharedPtr<SVerticalBox> VerticalBox;
-
-		TArray<FName> PlatformOverrides = PlatformOverrideNames.Get();
-		LastPlatformOverrideNames = PlatformOverrides.Num();
-		ChildSlot
-		.VAlign(VAlign_Fill)
-		.HAlign(HAlign_Fill)
-		[
-			SAssignNew(HorizontalBox, SHorizontalBox)
-			+SHorizontalBox::Slot()
-			.FillWidth(1.0f)
-			[
-				SAssignNew(VerticalBox, SVerticalBox)
-				+ SVerticalBox::Slot()
-				.AutoHeight()
-				[
-					SNew(STextBlock)
-					.Visibility(PlatformOverrides.Num() ? EVisibility::Visible : EVisibility::Collapsed)
-					.Font(IDetailLayoutBuilder::GetDetailFont())
-					.Text(NSLOCTEXT("SPerPlatformPropertiesWidget", "DefaultPlatform", "Default"))
-				]
-			]
-		];
-
-		if (OnGenerateWidget.IsBound())
-		{
-			// Default control
-			VerticalBox->AddSlot()
-			.AutoHeight()
-			.VAlign(EVerticalAlignment::VAlign_Bottom)
-			[
-				OnGenerateWidget.Execute(NAME_None)
-			];
-
-			// Add Platform menu
-			bool bAddedMenuItem = false;
-			FMenuBuilder AddPlatformMenuBuilder(true, nullptr, nullptr, true);
-						
-			// Platform (group) names
-			const TArray<FName>& PlatformGroupNameArray = PlatformInfo::GetAllPlatformGroupNames();
-			const TArray<FName>& VanillaPlatformNameArray = PlatformInfo::GetAllVanillaPlatformNames();			
-
-			// Sanitized platform names
-			TArray<FName> BasePlatformNameArray;
-			// Mapping from platform group name to individual platforms
-			TMultiMap<FName, FName> GroupToPlatform;
-						
-			// Create mapping from platform to platform groups and remove postfixes and invalid platform names
-			const TArray<FString> Filters = { TEXT("NoEditor"), TEXT("Client"), TEXT("Server"), TEXT("AllDesktop") };
-			for (const FName& PlatformName : VanillaPlatformNameArray)
-			{
-				const PlatformInfo::FPlatformInfo* PlatformInfo = PlatformInfo::FindPlatformInfo(PlatformName);
-				FString PlatformNameString = PlatformName.ToString();
-				for ( const FString& Filter : Filters)
-				{
-					const int32 Position = PlatformNameString.Find(Filter);
-					if (Position != INDEX_NONE)
-					{
-						PlatformNameString.RemoveAt(Position, Filter.Len());
-						break;
-					}
-				}
-
-				// Add filtered name if it isn't already set, and also add to group mapping
-				const FName FilteredName = FName(*PlatformNameString);				
-				if (PlatformNameString.Len() && !PlatformOverrides.Contains(FilteredName))
-				{	
-					BasePlatformNameArray.AddUnique(FilteredName);
-					GroupToPlatform.AddUnique(PlatformInfo->PlatformGroupName, FilteredName);
-				}
-			}
-
-			// Create section for platform groups 
-			const FName PlatformGroupSection(TEXT("PlatformGroupSection"));
-			AddPlatformMenuBuilder.BeginSection(PlatformGroupSection, FText::FromString(TEXT("Platform Groups")));
-			for (const FName& GroupName : PlatformGroupNameArray)
-			{				
-				if (!PlatformOverrides.Contains(GroupName))
-				{					
-					const FTextFormat Format = NSLOCTEXT("SPerPlatformPropertiesWidget", "AddOverrideGroupFor", "Add Override for Platforms part of the {0} Platform Group");
-					AddPlatformToMenu(GroupName, Format, AddPlatformMenuBuilder);
-					bAddedMenuItem = true;
-				}
-			}				
-			AddPlatformMenuBuilder.EndSection();
-			
-			for (const FName& GroupName : PlatformGroupNameArray)
-			{
-				// Create a section for each platform group and their respective platforms
-				AddPlatformMenuBuilder.BeginSection(GroupName, FText::FromName(GroupName));
-
-				TArray<FName> PlatformNames;
-				GroupToPlatform.MultiFind(GroupName, PlatformNames);
-
-				const FTextFormat Format = NSLOCTEXT("SPerPlatformPropertiesWidget", "AddOverrideFor", "Add Override specifically for {0}");
-				for (const FName& PlatformName : PlatformNames)
-				{
-					AddPlatformToMenu(PlatformName, Format, AddPlatformMenuBuilder);
-				}
-
-				bAddedMenuItem |= PlatformNames.Num() > 0;
-
-				AddPlatformMenuBuilder.EndSection();
-			}
-			
-			if (bAddedMenuItem)
-			{
-				HorizontalBox->AddSlot()
-				.AutoWidth()
-				[
-					SNew(SComboButton)
-					.VAlign(EVerticalAlignment::VAlign_Bottom)
-					.ButtonStyle(FEditorStyle::Get(), "HoverHintOnly")
-					.ContentPadding(4.0f)
-					.ForegroundColor(FSlateColor::UseForeground())
-					.HasDownArrow(false)
-					.ButtonContent()
-					[
-						SNew(SImage)
-						.Image(FEditorStyle::GetBrush("PropertyWindow.Button_AddToArray"))
-					]
-					.MenuContent()
-					[
-						AddPlatformMenuBuilder.MakeWidget()
-					]
-					.ToolTipText(NSLOCTEXT("SPerPlatformPropertiesWidget", "AddOverrideToolTip", "Add an override for a specific platform or platform group"))
-				];
-			}
-
-			for (FName PlatformName : PlatformOverrides)
-			{
-				HorizontalBox->AddSlot()
-				.Padding(FMargin(5.0f, 0.0f, 0.0f, 0.0f))
-				.AutoWidth()
-				[
-					SNew(SVerticalBox)
-					+ SVerticalBox::Slot()
-					.AutoHeight()
-					.VAlign(EVerticalAlignment::VAlign_Bottom)
-					[
-						SNew(SHorizontalBox)
-						+ SHorizontalBox::Slot()
-						.FillWidth(1.0f)
-						[
-							SNew(STextBlock)
-							.Font(IDetailLayoutBuilder::GetDetailFont())
-							.Text(FText::FromName(PlatformName))
-						]
-
-						+ SHorizontalBox::Slot()
-						.HAlign(EHorizontalAlignment::HAlign_Right)
-						.AutoWidth()
-						[
-							SNew(SButton)
-							.ContentPadding(FMargin(1.0f, 0.0f, 1.0f, 0.0f))
-							.ButtonStyle(FEditorStyle::Get(), "NoBorder")
-							.OnClicked(this, &SPerPlatformPropertiesWidget::RemovePlatform, PlatformName)
-							.ToolTipText(FText::Format(NSLOCTEXT("SPerPlatformPropertiesWidget", "RemoveOverrideFor", "Remove Override for {0}"), FText::AsCultureInvariant(PlatformName.ToString())))
-							.ForegroundColor(FSlateColor::UseForeground())
-							.HAlign(HAlign_Center)
-							.VAlign(VAlign_Center)
-							.Content()
-							[
-								SNew(SImage)
-								.Image(FEditorStyle::GetBrush("Icons.Cross"))
-							]
-						]
-					]
-
-					+ SVerticalBox::Slot()
-					.AutoHeight()
-					[
-						OnGenerateWidget.Execute(PlatformName)
-					]
-				];
-			}
-		}
-		else
-		{
-			VerticalBox->AddSlot()
-			.AutoHeight()
-			[
-				SNew(STextBlock)
-				.Text(NSLOCTEXT("SPerPlatformPropertiesWidget", "OnGenerateWidgetWarning", "No OnGenerateWidget() Provided"))
-				.ColorAndOpacity(FLinearColor::Red)
-			];
-		}
-	}
+protected:
+	void ConstructChildren();
 
-	void AddPlatformToMenu(const FName& PlatformName, const FTextFormat Format, FMenuBuilder &AddPlatformMenuBuilder)
-	{		
-		const FText MenuText = FText::Format(FText::FromString(TEXT("{0}")), FText::AsCultureInvariant(PlatformName.ToString()));
-		const FText MenuTooltipText = FText::Format(Format, FText::AsCultureInvariant(PlatformName.ToString()));
-		AddPlatformMenuBuilder.AddMenuEntry(
-			MenuText,
-			MenuTooltipText,
-			FSlateIcon(FEditorStyle::GetStyleSetName(), "PerPlatformWidget.AddPlatform"),
-			FUIAction(FExecuteAction::CreateSP(this, &SPerPlatformPropertiesWidget::AddPlatform, PlatformName))
-		);
-	}
+	TSharedRef<SWidget> MakePerPlatformWidget(FName InName, FText InDisplayText, const TArray<FName>& InPlatformOverrides, FMenuBuilder& InAddPlatformMenuBuilder);
 
-protected:
-	void AddPlatform(FName PlatformName)
-	{
-		if (OnAddPlatform.IsBound() && OnAddPlatform.Execute(PlatformName))
-		{
-			ConstructChildren();
-			Invalidate(EInvalidateWidget::LayoutAndVolatility);
-		}
-	}
+	void AddPlatform(FName PlatformName);
 
-	FReply RemovePlatform(FName PlatformName)
-	{
-		if (OnRemovePlatform.IsBound() && OnRemovePlatform.Execute(PlatformName))
-		{
-			ConstructChildren();
-			Invalidate(EInvalidateWidget::LayoutAndVolatility);
-		}
-		return FReply::Handled();
-	}
+	FReply RemovePlatform(FName PlatformName);
 
-	EActiveTimerReturnType CheckPlatformCount(double InCurrentTime, float InDeltaSeconds)
-	{
-		// @fixme: the platform count is fixed and locally cached
-		// so if you change this outside of editor, this widget won't update
-		// this timer is the one checking to see if platform count has changed
-		// if so, it will reconstruct
-		TArray<FName> PlatformOverrides = PlatformOverrideNames.Get();
-		if (LastPlatformOverrideNames != PlatformOverrides.Num())
-		{
-			ConstructChildren();
-		}
+	EActiveTimerReturnType CheckPlatformCount(double InCurrentTime, float InDeltaSeconds);
 
-		return EActiveTimerReturnType::Continue;
-	}
+	void AddPlatformToMenu(const FName& PlatformName, const FTextFormat Format, FMenuBuilder& AddPlatformMenuBuilder);
 
 	FOnGenerateWidget OnGenerateWidget;
 	FOnPlatformAction OnAddPlatform;
 	FOnPlatformAction OnRemovePlatform;
 	TAttribute<TArray<FName>> PlatformOverrideNames;
 	int32 LastPlatformOverrideNames;
+	bool bAddedMenuItem;
 };
 
diff --git a/Engine/Source/Editor/StaticMeshEditor/Private/StaticMeshEditorTools.cpp b/Engine/Source/Editor/StaticMeshEditor/Private/StaticMeshEditorTools.cpp
--- a/Engine/Source/Editor/StaticMeshEditor/Private/StaticMeshEditorTools.cpp
+++ b/Engine/Source/Editor/StaticMeshEditor/Private/StaticMeshEditorTools.cpp
@@ -28,6 +28,7 @@
 #include "Widgets/Input/SVectorInputBox.h"
 #include "Widgets/Input/SNumericEntryBox.h"
 #include "SPerPlatformPropertiesWidget.h"
+#include "PlatformInfo.h"
 
 #include "Runtime/Analytics/Analytics/Public/Interfaces/IAnalyticsProvider.h"
 #include "EngineAnalytics.h"
@@ -3065,6 +3066,7 @@ void FLevelOfDetailSettingsLayout::AddToDetailsPanel( IDetailLayoutBuilder& Deta
 	.MaxDesiredWidth((float)(PlatformNumber + 1)*125.0f)
 	[
 		SNew(SPerPlatformPropertiesWidget)
+		.IsEnabled(FLevelOfDetailSettingsLayout::GetLODCount() > 1)
 		.OnGenerateWidget(this, &FLevelOfDetailSettingsLayout::GetMinLODWidget)
 		.OnAddPlatform(this, &FLevelOfDetailSettingsLayout::AddMinLODPlatformOverride)
 		.OnRemovePlatform(this, &FLevelOfDetailSettingsLayout::RemoveMinLODPlatformOverride)
@@ -3083,6 +3085,7 @@ void FLevelOfDetailSettingsLayout::AddToDetailsPanel( IDetailLayoutBuilder& Deta
 	.MaxDesiredWidth((float)(PlatformNumber + 1)*125.0f)
 	[
 		SNew(SPerPlatformPropertiesWidget)
+		.IsEnabled(FLevelOfDetailSettingsLayout::GetLODCount() > 1)
 		.OnGenerateWidget(this, &FLevelOfDetailSettingsLayout::GetNumStreamedLODsWidget)
 		.OnAddPlatform(this, &FLevelOfDetailSettingsLayout::AddNumStreamedLODsPlatformOverride)
 		.OnRemovePlatform(this, &FLevelOfDetailSettingsLayout::RemoveNumStreamedLODsPlatformOverride)
@@ -3386,6 +3389,7 @@ void FLevelOfDetailSettingsLayout::AddLODLevelCategories( IDetailLayoutBuilder&
 			.MaxDesiredWidth((float)(PlatformNumber + 1)*125.0f)
 			[
 				SNew(SPerPlatformPropertiesWidget)
+				.IsEnabled(this, &FLevelOfDetailSettingsLayout::CanChangeLODScreenSize)
 				.OnGenerateWidget(this, &FLevelOfDetailSettingsLayout::GetLODScreenSizeWidget, LODIndex)
 				.OnAddPlatform(this, &FLevelOfDetailSettingsLayout::AddLODScreenSizePlatformOverride, LODIndex)
 				.OnRemovePlatform(this, &FLevelOfDetailSettingsLayout::RemoveLODScreenSizePlatformOverride, LODIndex)
