Filter custom rows in UMG Details View

diff --git a/Engine/Source/Editor/PropertyEditor/Private/CustomChildBuilder.cpp b/Engine/Source/Editor/PropertyEditor/Private/CustomChildBuilder.cpp
--- a/Engine/Source/Editor/PropertyEditor/Private/CustomChildBuilder.cpp
+++ b/Engine/Source/Editor/PropertyEditor/Private/CustomChildBuilder.cpp
@@ -32,6 +32,12 @@ FDetailWidgetRow& FCustomChildrenBuilder::AddCustomRow( const FText& SearchStrin
 	FDetailLayoutCustomization NewCustomization;
 
 	NewRow->FilterString( SearchString );
+	
+	IDetailsViewPrivate* DetailsView = ParentCategory.Pin()->GetDetailsView();
+	if (DetailsView && DetailsView->IsCustomRowVisibilityFiltered() && !DetailsView->IsCustomRowVisible(FName(*SearchString.ToString()), FName(*GetParentCategory().GetDisplayName().ToString())))
+	{
+		NewRow->Visibility(TAttribute<EVisibility>(EVisibility::Collapsed));
+	}
 
 	NewCustomization.WidgetDecl = NewRow;
 
diff --git a/Engine/Source/Editor/PropertyEditor/Private/DetailCategoryBuilderImpl.cpp b/Engine/Source/Editor/PropertyEditor/Private/DetailCategoryBuilderImpl.cpp
--- a/Engine/Source/Editor/PropertyEditor/Private/DetailCategoryBuilderImpl.cpp
+++ b/Engine/Source/Editor/PropertyEditor/Private/DetailCategoryBuilderImpl.cpp
@@ -139,6 +139,12 @@ FDetailWidgetRow& FDetailCategoryImpl::AddCustomRow(const FText& FilterString, b
 	NewCustomization.WidgetDecl = MakeShareable(new FDetailWidgetRow);
 	NewCustomization.WidgetDecl->FilterString(FilterString);
 
+	IDetailsViewPrivate* DetailsView = GetDetailsView();
+	if (DetailsView && DetailsView->IsCustomRowVisibilityFiltered() && !GetDetailsView()->IsCustomRowVisible(FName(*FilterString.ToString()), FName(*DisplayName.ToString())))
+	{
+		NewCustomization.WidgetDecl->Visibility(TAttribute<EVisibility>(EVisibility::Collapsed));
+	}
+
 	AddCustomLayout(NewCustomization, bForAdvanced);
 
 	return *NewCustomization.WidgetDecl;
diff --git a/Engine/Source/Editor/PropertyEditor/Private/IDetailsViewPrivate.h b/Engine/Source/Editor/PropertyEditor/Private/IDetailsViewPrivate.h
--- a/Engine/Source/Editor/PropertyEditor/Private/IDetailsViewPrivate.h
+++ b/Engine/Source/Editor/PropertyEditor/Private/IDetailsViewPrivate.h
@@ -80,6 +80,16 @@ public:
 	 */
 	virtual bool IsPropertyReadOnly( const struct FPropertyAndParent& PropertyAndParent ) const = 0;
 
+	/**
+	 * @return Whether the IsCustomRowVisible check is pertinent, i.e. always return true no matter the specified row and parent names.
+	 */
+	virtual bool IsCustomRowVisibilityFiltered() const = 0;
+
+	/**
+	 * @return Whether a custom row with the specified name and parent name is visible.
+	 */
+	virtual bool IsCustomRowVisible(FName InRowName, FName InParentName) const = 0;
+
 	/**
 	 * @return The thumbnail pool that should be used for thumbnails being rendered in this view
 	 */
diff --git a/Engine/Source/Editor/PropertyEditor/Private/SDetailsViewBase.cpp b/Engine/Source/Editor/PropertyEditor/Private/SDetailsViewBase.cpp
--- a/Engine/Source/Editor/PropertyEditor/Private/SDetailsViewBase.cpp
+++ b/Engine/Source/Editor/PropertyEditor/Private/SDetailsViewBase.cpp
@@ -431,6 +431,16 @@ void SDetailsViewBase::SetIsPropertyReadOnlyDelegate(FIsPropertyReadOnly InIsPro
 	IsPropertyReadOnlyDelegate = InIsPropertyReadOnly;
 }
 
+void SDetailsViewBase::SetIsCustomRowVisibilityFilteredDelegate(FIsCustomRowVisibilityFiltered InIsCustomRowVisibilityFilteredDelegate)
+{
+	IsCustomRowVisibilityFilteredDelegate = InIsCustomRowVisibilityFilteredDelegate;
+};
+
+void SDetailsViewBase::SetIsCustomRowVisibleDelegate(FIsCustomRowVisible InIsCustomRowVisible)
+{
+	IsCustomRowVisibleDelegate = InIsCustomRowVisible;
+}
+
 void SDetailsViewBase::SetIsPropertyEditingEnabledDelegate(FIsPropertyEditingEnabled IsPropertyEditingEnabled)
 {
 	IsPropertyEditingEnabledDelegate = IsPropertyEditingEnabled;
@@ -544,6 +554,20 @@ bool SDetailsViewBase::IsPropertyReadOnly( const FPropertyAndParent& PropertyAnd
 	return IsPropertyReadOnlyDelegate.IsBound() ? IsPropertyReadOnlyDelegate.Execute(PropertyAndParent) : false;
 }
 
+bool SDetailsViewBase::IsCustomRowVisibilityFiltered() const
+{
+	if (!IsCustomRowVisibleDelegate.IsBound())
+	{
+		return false;
+	}
+	return IsCustomRowVisibilityFilteredDelegate.IsBound() ? IsCustomRowVisibilityFilteredDelegate.Execute() : true;
+}
+
+bool SDetailsViewBase::IsCustomRowVisible(FName InRowName, FName InParentName) const
+{
+	return IsCustomRowVisibleDelegate.IsBound() ? IsCustomRowVisibleDelegate.Execute(InRowName, InParentName) : true;
+}
+
 TSharedPtr<IPropertyUtilities> SDetailsViewBase::GetPropertyUtilities()
 {
 	return PropertyUtilities;
diff --git a/Engine/Source/Editor/PropertyEditor/Private/SDetailsViewBase.h b/Engine/Source/Editor/PropertyEditor/Private/SDetailsViewBase.h
--- a/Engine/Source/Editor/PropertyEditor/Private/SDetailsViewBase.h
+++ b/Engine/Source/Editor/PropertyEditor/Private/SDetailsViewBase.h
@@ -131,7 +131,11 @@ public:
 	virtual FOnDisplayedPropertiesChanged& GetOnDisplayedPropertiesChanged() override { return OnDisplayedPropertiesChangedDelegate; }
 	virtual void SetDisableCustomDetailLayouts( bool bInDisableCustomDetailLayouts ) override { bDisableCustomDetailLayouts = bInDisableCustomDetailLayouts; }
 	virtual void SetIsPropertyVisibleDelegate(FIsPropertyVisible InIsPropertyVisible) override;
+	virtual FIsCustomRowVisibilityFiltered& GetIsCustomRowVisibilityFilteredDelegate() override { return IsCustomRowVisibilityFilteredDelegate; };
+	virtual void SetIsCustomRowVisibilityFilteredDelegate(FIsCustomRowVisibilityFiltered InIsCustomRowVisibilityFilteredDelegate) override;
 	virtual FIsPropertyVisible& GetIsPropertyVisibleDelegate() override { return IsPropertyVisibleDelegate; }
+	virtual void SetIsCustomRowVisibleDelegate(FIsCustomRowVisible InIsCustomRowVisible) override;
+	virtual FIsCustomRowVisible& GetIsCustomRowVisibleDelegate() override { return IsCustomRowVisibleDelegate; }
 	virtual void SetIsPropertyReadOnlyDelegate(FIsPropertyReadOnly InIsPropertyReadOnly) override;
 	virtual FIsPropertyReadOnly& GetIsPropertyReadOnlyDelegate() override { return IsPropertyReadOnlyDelegate; }
 	virtual void SetIsPropertyEditingEnabledDelegate(FIsPropertyEditingEnabled IsPropertyEditingEnabled) override;
@@ -143,6 +147,8 @@ public:
 	virtual TSharedPtr<IDetailPropertyExtensionHandler> GetExtensionHandler() override;
 	virtual bool IsPropertyVisible(const struct FPropertyAndParent& PropertyAndParent) const override;
 	virtual bool IsPropertyReadOnly(const struct FPropertyAndParent& PropertyAndParent) const override;
+    virtual bool IsCustomRowVisibilityFiltered() const override;
+	virtual bool IsCustomRowVisible(FName InRowName, FName InParentName) const override;
 	virtual void SetGenericLayoutDetailsDelegate(FOnGetDetailCustomizationInstance OnGetGenericDetails) override;
 	virtual FOnGetDetailCustomizationInstance& GetGenericLayoutDetailsDelegate() override { return GenericLayoutDelegate; }
 	virtual bool IsLocked() const override { return bIsLocked; }
@@ -375,6 +381,10 @@ protected:
 	FDetailNodeList RootTreeNodes;
 	/** Delegate executed to determine if a property should be visible */
 	FIsPropertyVisible IsPropertyVisibleDelegate;
+	/** Delegate executed to check if custom row visibility is filtered. */
+	FIsCustomRowVisibilityFiltered IsCustomRowVisibilityFilteredDelegate;
+	/** Delegate executed to determine if a custom row should be visible. */
+	FIsCustomRowVisible IsCustomRowVisibleDelegate;
 	/** Delegate executed to determine if a property should be read-only */
 	FIsPropertyReadOnly IsPropertyReadOnlyDelegate;
 	/** Delegate called to see if a property editing is enabled */
diff --git a/Engine/Source/Editor/PropertyEditor/Public/IDetailsView.h b/Engine/Source/Editor/PropertyEditor/Public/IDetailsView.h
--- a/Engine/Source/Editor/PropertyEditor/Public/IDetailsView.h
+++ b/Engine/Source/Editor/PropertyEditor/Public/IDetailsView.h
@@ -249,6 +249,19 @@ public:
 	virtual void SetIsPropertyReadOnlyDelegate( FIsPropertyReadOnly InIsPropertyReadOnly ) = 0;
 	virtual FIsPropertyReadOnly& GetIsPropertyReadOnlyDelegate() = 0;
 
+	/**
+	 * Sets a delegate to call to check if custom row visibility is filtered,
+	 * i.e. whether the FIsCustomRowVisible delegate will always return true no matter the parameters.
+	 */
+	virtual void SetIsCustomRowVisibilityFilteredDelegate(FIsCustomRowVisibilityFiltered InIsCustomRowVisibilityFiltered) = 0;
+	virtual FIsCustomRowVisibilityFiltered& GetIsCustomRowVisibilityFilteredDelegate() = 0;
+
+	/**
+	 * Sets a delegate to call to determine if a specific custom row should be visible in this instance of the details view
+	 */
+	virtual void SetIsCustomRowVisibleDelegate(FIsCustomRowVisible InIsCustomRowVisible) = 0;
+	virtual FIsCustomRowVisible& GetIsCustomRowVisibleDelegate() = 0;
+
 	/**
 	 * Sets a delegate to call to layout generic details not specific to an object being viewed
 	 */ 
diff --git a/Engine/Source/Editor/PropertyEditor/Public/PropertyEditorDelegates.h b/Engine/Source/Editor/PropertyEditor/Public/PropertyEditorDelegates.h
--- a/Engine/Source/Editor/PropertyEditor/Public/PropertyEditorDelegates.h
+++ b/Engine/Source/Editor/PropertyEditor/Public/PropertyEditorDelegates.h
@@ -29,9 +29,18 @@ struct FPropertyAndParent
 /** Delegate called to see if a property should be visible */
 DECLARE_DELEGATE_RetVal_OneParam( bool, FIsPropertyVisible, const FPropertyAndParent& );
 
-/** Delegate called to see if a property should be visible */
+/** Delegate called to see if a property should be read-only */
 DECLARE_DELEGATE_RetVal_OneParam( bool, FIsPropertyReadOnly, const FPropertyAndParent& );
 
+/** 
+ * Delegate called to check if custom row visibility is filtered, 
+ * i.e. whether FIsCustomRowVisible delegate will always return true no matter the parameters. 
+ */
+DECLARE_DELEGATE_RetVal(bool, FIsCustomRowVisibilityFiltered);
+
+/** Delegate called to determine if a custom row should be visible. */
+DECLARE_DELEGATE_RetVal_TwoParams(bool, FIsCustomRowVisible, FName /*InRowName*/, FName /*InParentName*/);
+
 /** Delegate called to get a detail layout for a specific object class */
 DECLARE_DELEGATE_RetVal( TSharedRef<class IDetailCustomization>, FOnGetDetailCustomizationInstance );
 
diff --git a/Engine/Source/Editor/UMGEditor/Private/Components/DetailsView.cpp b/Engine/Source/Editor/UMGEditor/Private/Components/DetailsView.cpp
--- a/Engine/Source/Editor/UMGEditor/Private/Components/DetailsView.cpp
+++ b/Engine/Source/Editor/UMGEditor/Private/Components/DetailsView.cpp
@@ -81,6 +81,8 @@ void UDetailsView::BuildContentWidget()
 			DetailViewWidget->SetCustomFilterDelegate(FSimpleDelegate::CreateUObject(this, &UDetailsView::ToggleWhitelistedProperties));
 
 			DetailViewWidget->SetIsPropertyVisibleDelegate(FIsPropertyVisible::CreateUObject(this, &UDetailsView::GetIsPropertyVisible));
+			DetailViewWidget->SetIsCustomRowVisibilityFilteredDelegate(FIsCustomRowVisibilityFiltered::CreateUObject(this, &UDetailsView::IsRowVisibilityFiltered));
+			DetailViewWidget->SetIsCustomRowVisibleDelegate(FIsCustomRowVisible::CreateUObject(this, &UDetailsView::GetIsRowVisible));
 			DetailViewWidget->SetObject(ViewedObject);
 			if (DetailViewWidget.IsValid())
 			{
@@ -160,14 +162,17 @@ void UDetailsView::ToggleWhitelistedProperties()
 }
 
 
+bool UDetailsView::IsRowVisibilityFiltered() const
+{
+	return bShowOnlyWhitelisted && (PropertiesToShow.Num() > 0 || CategoriesToShow.Num() > 0);
+}
+
 bool UDetailsView::GetIsPropertyVisible(const FPropertyAndParent& PropertyAndParent) const
 {
-	if (!bShowOnlyWhitelisted
-		|| (PropertiesToShow.Num() == 0 && CategoriesToShow.Num() == 0))
+    if (!IsRowVisibilityFiltered())
 	{
 		return true;
 	}
-
 	if (PropertiesToShow.Contains(PropertyAndParent.Property.GetFName()))
 	{
 		return true;
@@ -179,6 +184,23 @@ bool UDetailsView::GetIsPropertyVisible(const FPropertyAndParent& PropertyAndPar
 	return false;
 }
 
+bool UDetailsView::GetIsRowVisible(FName InRowName, FName InParentName) const
+{
+    if (!IsRowVisibilityFiltered())
+	{
+		return true;
+	}
+	if (PropertiesToShow.Contains(InRowName))
+	{
+		return true;
+	}
+	if (CategoriesToShow.Contains(InParentName))
+	{
+		return true;
+	}
+	return false;
+}
+
 /////////////////////////////////////////////////////
 
 #undef LOCTEXT_NAMESPACE
diff --git a/Engine/Source/Editor/UMGEditor/Public/Components/DetailsView.h b/Engine/Source/Editor/UMGEditor/Public/Components/DetailsView.h
--- a/Engine/Source/Editor/UMGEditor/Public/Components/DetailsView.h
+++ b/Engine/Source/Editor/UMGEditor/Public/Components/DetailsView.h
@@ -84,7 +84,9 @@ public:
 
 protected:
 	void ToggleWhitelistedProperties();
+	bool IsRowVisibilityFiltered() const;
 	bool GetIsPropertyVisible(const FPropertyAndParent& PropertyAndParent) const;
+	bool GetIsRowVisible(FName InRowName, FName InParentName) const;
 
 private:
 	TSharedPtr<IDetailsView> DetailViewWidget;
