Fix details panel regressions caused by CL 7273612 (fix for UE-76064).

diff --git a/Engine/Source/Editor/PropertyEditor/Private/DetailCategoryBuilderImpl.cpp b/Engine/Source/Editor/PropertyEditor/Private/DetailCategoryBuilderImpl.cpp
--- a/Engine/Source/Editor/PropertyEditor/Private/DetailCategoryBuilderImpl.cpp
+++ b/Engine/Source/Editor/PropertyEditor/Private/DetailCategoryBuilderImpl.cpp
@@ -72,8 +72,9 @@ bool FDetailLayoutCustomization::IsHidden() const
 {
 	return !IsValidCustomization()
 		|| (HasCustomWidget() && WidgetDecl->VisibilityAttr.Get() != EVisibility::Visible)
-		|| (HasPropertyNode() && PropertyRow->GetPropertyVisibility() != EVisibility::Visible)
-		|| (HasCustomBuilder() && CustomBuilderRow->AreChildCustomizationsHidden());
+		|| (HasPropertyNode() && PropertyRow->GetPropertyVisibility() != EVisibility::Visible);
+		/** Partial revert of CL 7273612 (fix for UE-76064) that caused a bunch of regressions in the details panel (UE-77377,UE-77376,UE-77451). */
+		//|| (HasCustomBuilder() && CustomBuilderRow->AreChildCustomizationsHidden());
 }
 
 TSharedPtr<FPropertyNode> FDetailLayoutCustomization::GetPropertyNode() const
