Fix detail layout customization visibility for the case of a custom builder (and don't hide custom builders that have no child).

diff --git a/Engine/Source/Editor/PropertyEditor/Private/DetailCategoryBuilderImpl.cpp b/Engine/Source/Editor/PropertyEditor/Private/DetailCategoryBuilderImpl.cpp
--- a/Engine/Source/Editor/PropertyEditor/Private/DetailCategoryBuilderImpl.cpp
+++ b/Engine/Source/Editor/PropertyEditor/Private/DetailCategoryBuilderImpl.cpp
@@ -72,9 +72,8 @@ bool FDetailLayoutCustomization::IsHidden() const
 {
 	return !IsValidCustomization()
 		|| (HasCustomWidget() && WidgetDecl->VisibilityAttr.Get() != EVisibility::Visible)
-		|| (HasPropertyNode() && PropertyRow->GetPropertyVisibility() != EVisibility::Visible);
-		/** Partial revert of CL 7273612 (fix for UE-76064) that caused a bunch of regressions in the details panel (UE-77377,UE-77376,UE-77451). */
-		//|| (HasCustomBuilder() && CustomBuilderRow->AreChildCustomizationsHidden());
+		|| (HasPropertyNode() && PropertyRow->GetPropertyVisibility() != EVisibility::Visible)
+		|| (HasCustomBuilder() && CustomBuilderRow->AreChildCustomizationsHidden());
 }
 
 TSharedPtr<FPropertyNode> FDetailLayoutCustomization::GetPropertyNode() const
diff --git a/Engine/Source/Editor/PropertyEditor/Private/DetailCustomBuilderRow.cpp b/Engine/Source/Editor/PropertyEditor/Private/DetailCustomBuilderRow.cpp
--- a/Engine/Source/Editor/PropertyEditor/Private/DetailCustomBuilderRow.cpp
+++ b/Engine/Source/Editor/PropertyEditor/Private/DetailCustomBuilderRow.cpp
@@ -85,17 +85,16 @@ FDetailWidgetRow FDetailCustomBuilderRow::GetWidgetRow()
 
 bool FDetailCustomBuilderRow::AreChildCustomizationsHidden() const
 {
-	bool bChildCustomizationsHidden = true;
-	if (ChildrenBuilder)
+	if (ChildrenBuilder && ChildrenBuilder->GetChildCustomizations().Num() > 0)
 	{
 		for (const FDetailLayoutCustomization& ChildCustomizations : ChildrenBuilder->GetChildCustomizations())
 		{
 			if (!ChildCustomizations.IsHidden())
 			{
-				bChildCustomizationsHidden = false;
-				break;
+				return false;
 			}
 		}
+		return true;
 	}
-	return bChildCustomizationsHidden;
+	return false;
 }
\ No newline at end of file
