Modifying the property node code to allow Enum's to drive the names of entries of the TArray, much like you can for fixedsize arrays. There are several cases in games where we want to move away from fixed size arrays to TArrays for memory savings.

diff --git a/Engine/Source/Editor/PropertyEditor/Private/ItemPropertyNode.cpp b/Engine/Source/Editor/PropertyEditor/Private/ItemPropertyNode.cpp
--- a/Engine/Source/Editor/PropertyEditor/Private/ItemPropertyNode.cpp
+++ b/Engine/Source/Editor/PropertyEditor/Private/ItemPropertyNode.cpp
@@ -521,6 +521,13 @@ FText FItemPropertyNode::GetDisplayName() const
 			
 			// Sets and maps do not have a display index.
 			UProperty* ParentProperty = ParentNode->GetProperty();
+
+			// Also handle UArray's having the ArraySizeEnum entry...
+			if (ArraySizeEnum == nullptr && Cast<UArrayProperty>(ParentProperty) != nullptr && ParentProperty->HasMetaData(NAME_ArraySizeEnum))
+			{
+				ArraySizeEnum = FindObject<UEnum>(NULL, *ParentProperty->GetMetaData(NAME_ArraySizeEnum));
+			}
+
 			if (Cast<USetProperty>(ParentProperty) == nullptr &&  Cast<UMapProperty>(ParentProperty) == nullptr)
 			{
 				// This item is a member of an array, its display name is its index 
