Adding values to a set or map in a blueprint's defaults no longer fires a check().

diff --git a/Engine/Source/Editor/PropertyEditor/Private/ItemPropertyNode.cpp b/Engine/Source/Editor/PropertyEditor/Private/ItemPropertyNode.cpp
--- a/Engine/Source/Editor/PropertyEditor/Private/ItemPropertyNode.cpp
+++ b/Engine/Source/Editor/PropertyEditor/Private/ItemPropertyNode.cpp
@@ -38,48 +38,49 @@ uint8* FItemPropertyNode::GetValueBaseAddress( uint8* StartAddress )
 
 		uint8* ValueBaseAddress = ParentNode->GetValueBaseAddress(StartAddress);
 
-		if ( OuterArrayProp != NULL )
+		if (OuterArrayProp != nullptr)
 		{
 			FScriptArrayHelper ArrayHelper(OuterArrayProp, ValueBaseAddress);
-			if ( ValueBaseAddress != NULL && ArrayIndex < ArrayHelper.Num() )
+			if (ValueBaseAddress != nullptr && ArrayHelper.Num() > 0)
 			{
 				return ArrayHelper.GetRawPtr() + ArrayOffset;
 			}
-
-			return NULL;
 		}
-		else if (OuterSetProp != NULL)
+		else if (OuterSetProp != nullptr)
 		{
 			FScriptSetHelper SetHelper(OuterSetProp, ValueBaseAddress);
-			if (ValueBaseAddress != NULL && SetHelper.Num())
+			if (ValueBaseAddress != nullptr)
 			{
 				int32 ActualIndex = SetHelper.FindInternalIndex(ArrayIndex);
-				return SetHelper.GetElementPtr(ActualIndex);
+				if (ActualIndex != INDEX_NONE)
+				{
+					return SetHelper.GetElementPtr(ActualIndex);
+				}
 			}
-
-			return NULL;
 		}
-		else if (OuterMapProp != NULL)
+		else if (OuterMapProp != nullptr)
 		{
 			FScriptMapHelper MapHelper(OuterMapProp, ValueBaseAddress);
-			if (ValueBaseAddress != NULL && MapHelper.Num() > 0)
+			if (ValueBaseAddress != nullptr)
 			{
 				int32 ActualIndex = MapHelper.FindInternalIndex(ArrayIndex);
-				uint8* PairPtr = MapHelper.GetPairPtr(ActualIndex);
-				return MyProperty->ContainerPtrToValuePtr<uint8>(PairPtr);
+				if (ActualIndex != INDEX_NONE)
+				{
+					uint8* PairPtr = MapHelper.GetPairPtr(ActualIndex);
+					return MyProperty->ContainerPtrToValuePtr<uint8>(PairPtr);
+				}
 			}
-
-			return NULL;
 		}
 		else
 		{
 			uint8* ValueAddress = ParentNode->GetValueAddress(StartAddress);
-			if (ValueAddress != NULL && ParentNode->GetProperty() != MyProperty)
+			if (ValueAddress != nullptr && ParentNode->GetProperty() != MyProperty)
 			{
 				// if this is not a fixed size array (in which the parent property and this property are the same), we need to offset from the property (otherwise, the parent already did that for us)
 				ValueAddress = Property->ContainerPtrToValuePtr<uint8>(ValueAddress);
 			}
-			if ( ValueAddress != NULL )
+
+			if (ValueAddress != nullptr)
 			{
 				ValueAddress += ArrayOffset;
 			}
@@ -87,7 +88,7 @@ uint8* FItemPropertyNode::GetValueBaseAddress( uint8* StartAddress )
 		}
 	}
 
-	return NULL;
+	return nullptr;
 }
 
 /**
@@ -113,6 +114,16 @@ uint8* FItemPropertyNode::GetValueAddress( uint8* StartAddress )
 		FScriptArrayHelper ArrayHelper(ArrayProperty,Result);
 		Result = ArrayHelper.GetRawPtr();
 	}
+	else if (Result && SetProperty)
+	{
+		FScriptSetHelper SetHelper(SetProperty, Result);
+		Result = SetHelper.GetElementPtr(0);
+	}
+	else if (Result && MapProperty)
+	{
+		FScriptMapHelper MapHelper(MapProperty, Result);
+		Result = MapHelper.GetPairPtr(0);
+	}
 
 	return Result;
 }
diff --git a/Engine/Source/Runtime/CoreUObject/Public/UObject/UnrealType.h b/Engine/Source/Runtime/CoreUObject/Public/UObject/UnrealType.h
--- a/Engine/Source/Runtime/CoreUObject/Public/UObject/UnrealType.h
+++ b/Engine/Source/Runtime/CoreUObject/Public/UObject/UnrealType.h
@@ -3271,7 +3271,10 @@ public:
 	 */
 	int32 FindInternalIndex(int32 LogicalIdx) const
 	{
-		check(LogicalIdx >= 0 && LogicalIdx < Num());
+		if (LogicalIdx < 0 && LogicalIdx > Num())
+		{
+			return INDEX_NONE;
+		}
 
 		int32 MaxIndex = GetMaxIndex();
 		for (int32 Actual = 0; Actual < MaxIndex; ++Actual)
@@ -3846,7 +3849,10 @@ public:
 	 */
 	int32 FindInternalIndex(int32 LogicalIdx) const
 	{
-		check(LogicalIdx >= 0 && LogicalIdx < Num());
+		if (LogicalIdx < 0 && LogicalIdx > Num())
+		{
+			return INDEX_NONE;
+		}
 
 		int32 MaxIndex = GetMaxIndex();
 		for (int32 Actual = 0; Actual < MaxIndex; ++Actual)
