Redoing fix to pass in a property handle to help filter materials in the static mesh editor
#rnx

diff --git a/Engine/Source/Editor/PropertyEditor/Private/MaterialList.cpp b/Engine/Source/Editor/PropertyEditor/Private/MaterialList.cpp
--- a/Engine/Source/Editor/PropertyEditor/Private/MaterialList.cpp
+++ b/Engine/Source/Editor/PropertyEditor/Private/MaterialList.cpp
@@ -140,7 +140,7 @@ public:
 			];
 	}
 
-	TSharedRef<SWidget> CreateValueContent( const TSharedPtr<FAssetThumbnailPool>& ThumbnailPool )
+	TSharedRef<SWidget> CreateValueContent( const TSharedPtr<FAssetThumbnailPool>& ThumbnailPool, TSharedPtr<IPropertyHandle> InHandle)
 	{
 		FIntPoint ThumbnailSize(64, 64);
 
@@ -171,6 +171,7 @@ public:
 						.ThumbnailPool(ThumbnailPool)
 						.DisplayCompactSize(bDisplayCompactSize)
 						.CustomResetToDefault(ResetToDefaultOverride)
+						.PropertyHandle(InHandle)
 						.CustomContentSlot()
 						[
 							SNew( SBox )
@@ -354,13 +355,14 @@ private:
 };
 
 
-FMaterialList::FMaterialList(IDetailLayoutBuilder& InDetailLayoutBuilder, FMaterialListDelegates& InMaterialListDelegates, bool bInAllowCollapse, bool bInShowUsedTextures, bool bInDisplayCompactSize)
+FMaterialList::FMaterialList(IDetailLayoutBuilder& InDetailLayoutBuilder, FMaterialListDelegates& InMaterialListDelegates, bool bInAllowCollapse, bool bInShowUsedTextures, bool bInDisplayCompactSize, TSharedPtr<class IPropertyHandle> InHandle)
 	: MaterialListDelegates( InMaterialListDelegates )
 	, DetailLayoutBuilder( InDetailLayoutBuilder )
 	, MaterialListBuilder( new FMaterialListBuilder )
 	, bAllowCollpase(bInAllowCollapse)
 	, bShowUsedTextures(bInShowUsedTextures)
 	, bDisplayCompactSize(bInDisplayCompactSize)
+	, MeshChildHandle(InHandle)
 {
 }
 
@@ -566,6 +568,11 @@ void FMaterialList::OnPasteMaterialList()
 	}
 }
 
+TSharedPtr<IPropertyHandle> FMaterialList::GetPropertyHandle() const
+{
+	return MeshChildHandle.IsValid() ? MeshChildHandle : nullptr;
+}
+
 bool FMaterialList::OnCanCopyMaterialItem(int32 CurrentSlot) const
 {
 	if (MaterialListDelegates.OnCanCopyMaterialItem.IsBound())
@@ -618,7 +625,7 @@ void FMaterialList::AddMaterialItem( FDetailWidgetRow& Row, int32 CurrentSlot, c
 	}
 	else
 	{
-		RightSideContent = NewView->CreateValueContent( DetailLayoutBuilder.GetThumbnailPool() );
+		RightSideContent = NewView->CreateValueContent( DetailLayoutBuilder.GetThumbnailPool(), GetPropertyHandle() );
 		ViewedMaterials.Add( NewView );
 	}
 
diff --git a/Engine/Source/Editor/PropertyEditor/Public/MaterialList.h b/Engine/Source/Editor/PropertyEditor/Public/MaterialList.h
--- a/Engine/Source/Editor/PropertyEditor/Public/MaterialList.h
+++ b/Engine/Source/Editor/PropertyEditor/Public/MaterialList.h
@@ -137,7 +137,7 @@ class FMaterialList
 	, public TSharedFromThis<FMaterialList>
 {
 public:
-	PROPERTYEDITOR_API FMaterialList( IDetailLayoutBuilder& InDetailLayoutBuilder, FMaterialListDelegates& MaterialListDelegates, bool bInAllowCollapse = false, bool bInShowUsedTextures = true, bool bInDisplayCompactSize = false);
+	PROPERTYEDITOR_API FMaterialList( IDetailLayoutBuilder& InDetailLayoutBuilder, FMaterialListDelegates& MaterialListDelegates, bool bInAllowCollapse = false, bool bInShowUsedTextures = true, bool bInDisplayCompactSize = false, TSharedPtr<class IPropertyHandle> InHandle = nullptr);
 
 	/**
 	 * @return true if materials are being displayed.                                                          
@@ -168,6 +168,7 @@ private:
 	virtual void GenerateChildContent( IDetailChildrenBuilder& ChildrenBuilder ) override;
 	virtual FName GetName() const override { return NAME_None; }
 	virtual bool InitiallyCollapsed() const override { return bAllowCollpase; }
+	virtual TSharedPtr<IPropertyHandle> GetPropertyHandle() const;
 
 	/**
 	 * Adds a new material item to the list
@@ -215,5 +216,7 @@ private:
 	bool bShowUsedTextures;
 	/** Whether or not to display a compact form of material entry*/
 	bool bDisplayCompactSize;
+	/** The property handle for a child property on the mesh that owns these materials */
+	TSharedPtr<class IPropertyHandle> MeshChildHandle;
 };
 
diff --git a/Engine/Source/Editor/StaticMeshEditor/Private/StaticMeshEditorTools.cpp b/Engine/Source/Editor/StaticMeshEditor/Private/StaticMeshEditorTools.cpp
--- a/Engine/Source/Editor/StaticMeshEditor/Private/StaticMeshEditorTools.cpp
+++ b/Engine/Source/Editor/StaticMeshEditor/Private/StaticMeshEditorTools.cpp
@@ -2376,7 +2376,7 @@ void FMeshMaterialsLayout::AddToCategory(IDetailCategoryBuilder& CategoryBuilder
 	MaterialListDelegates.OnCanCopyMaterialItem.BindSP(this, &FMeshMaterialsLayout::OnCanCopyMaterialItem);
 	MaterialListDelegates.OnPasteMaterialItem.BindSP(this, &FMeshMaterialsLayout::OnPasteMaterialItem);
 
-	CategoryBuilder.AddCustomBuilder(MakeShareable(new FMaterialList(CategoryBuilder.GetParentLayout(), MaterialListDelegates, false, true, true)));
+	CategoryBuilder.AddCustomBuilder(MakeShareable(new FMaterialList(CategoryBuilder.GetParentLayout(), MaterialListDelegates, false, true, true, InHandle)));
 }
 
 void FMeshMaterialsLayout::OnCopyMaterialList()
