Fix crash in FPropertyNode::GetDiffersFromDefaultForObject when value tracker has null property default base address.

diff --git a/Engine/Source/Editor/PropertyEditor/Private/PropertyNode.cpp b/Engine/Source/Editor/PropertyEditor/Private/PropertyNode.cpp
--- a/Engine/Source/Editor/PropertyEditor/Private/PropertyNode.cpp
+++ b/Engine/Source/Editor/PropertyEditor/Private/PropertyNode.cpp
@@ -1557,39 +1557,40 @@ bool FPropertyNode::GetDiffersFromDefaultForObject( FPropertyItemValueDataTracke
 
 	if ( ValueTracker.IsValidTracker() && ValueTracker.HasDefaultValue() && GetParentNode() != NULL )
 	{
-		//////////////////////////
-		// Check the property against its default.
-		// If the property is an object property, we have to take special measures.
-		UArrayProperty* OuterArrayProperty = Cast<UArrayProperty>(InProperty->GetOuter());
-		USetProperty* OuterSetProperty = Cast<USetProperty>(InProperty->GetOuter());
-		UMapProperty* OuterMapProperty = Cast<UMapProperty>(InProperty->GetOuter());
-
-		if ( OuterArrayProperty != NULL )
+		if (ValueTracker.GetPropertyDefaultBaseAddress() != NULL)
 		{
-			// make sure we're not trying to compare against an element that doesn't exist
-			if ( ValueTracker.GetPropertyDefaultBaseAddress() != NULL && GetArrayIndex() >= FScriptArrayHelper::Num(ValueTracker.GetPropertyDefaultBaseAddress()) )
+			//////////////////////////
+			// Check the property against its default.
+			// If the property is an object property, we have to take special measures.
+			UArrayProperty* OuterArrayProperty = Cast<UArrayProperty>(InProperty->GetOuter());
+			USetProperty* OuterSetProperty = Cast<USetProperty>(InProperty->GetOuter());
+			UMapProperty* OuterMapProperty = Cast<UMapProperty>(InProperty->GetOuter());
+
+			if ( OuterArrayProperty != NULL )
 			{
-				bDiffersFromDefaultForObject = true;
+				// make sure we're not trying to compare against an element that doesn't exist
+				if (GetArrayIndex() >= FScriptArrayHelper::Num(ValueTracker.GetPropertyDefaultBaseAddress()) )
+				{
+					bDiffersFromDefaultForObject = true;
+				}
 			}
-		}
-		else if (OuterSetProperty != NULL)
-		{
-			FScriptSetHelper SetHelper(OuterSetProperty, ValueTracker.GetPropertyDefaultBaseAddress());
-
-			bool bIsValidIndex = ArrayIndex >= 0 && ArrayIndex < SetHelper.Num();
-			if ( ValueTracker.GetPropertyDefaultBaseAddress() != NULL && !bIsValidIndex)
+			else if (OuterSetProperty != NULL)
 			{
-				bDiffersFromDefaultForObject = true;
+				FScriptSetHelper SetHelper(OuterSetProperty, ValueTracker.GetPropertyDefaultBaseAddress());
+				bool bIsValidIndex = ArrayIndex >= 0 && ArrayIndex < SetHelper.Num();
+				if (!bIsValidIndex)
+				{
+					bDiffersFromDefaultForObject = true;
+				}
 			}
-		}
-		else if (OuterMapProperty != NULL)
-		{
-			FScriptMapHelper MapHelper(OuterMapProperty, ValueTracker.GetPropertyDefaultBaseAddress());
-
-			bool bIsValidIndex = ArrayIndex >= 0 && ArrayIndex < MapHelper.Num();
-			if ( ValueTracker.GetPropertyDefaultBaseAddress() != NULL && !bIsValidIndex )
+			else if (OuterMapProperty != NULL)
 			{
-				bDiffersFromDefaultForObject = true;
+				FScriptMapHelper MapHelper(OuterMapProperty, ValueTracker.GetPropertyDefaultBaseAddress());
+				bool bIsValidIndex = ArrayIndex >= 0 && ArrayIndex < MapHelper.Num();
+				if (!bIsValidIndex)
+				{
+					bDiffersFromDefaultForObject = true;
+				}
 			}
 		}
 
