Modify the asset picker filtering to pass an array of FAssetData has aditionnal references for the filtering

diff --git a/Engine/Source/Editor/DetailCustomizations/Private/ComponentMaterialCategory.cpp b/Engine/Source/Editor/DetailCustomizations/Private/ComponentMaterialCategory.cpp
--- a/Engine/Source/Editor/DetailCustomizations/Private/ComponentMaterialCategory.cpp
+++ b/Engine/Source/Editor/DetailCustomizations/Private/ComponentMaterialCategory.cpp
@@ -165,7 +165,11 @@ void FComponentMaterialCategory::Create( IDetailLayoutBuilder& DetailBuilder )
 	FMaterialListDelegates MaterialListDelegates;
 	MaterialListDelegates.OnGetMaterials.BindSP( this, &FComponentMaterialCategory::OnGetMaterialsForView );
 	MaterialListDelegates.OnMaterialChanged.BindSP( this, &FComponentMaterialCategory::OnMaterialChanged );
-	TSharedRef<FMaterialList> MaterialList = MakeShareable( new FMaterialList( DetailBuilder, MaterialListDelegates ) );
+	
+	//Pass an empty material list owner (owner can be use by the asset picker filter. In this case we do not need it)
+	TArray<FAssetData> MaterialListOwner;
+
+	TSharedRef<FMaterialList> MaterialList = MakeShareable( new FMaterialList( DetailBuilder, MaterialListDelegates, MaterialListOwner) );
 
 	bool bAnyMaterialsToDisplay = false;
 
diff --git a/Engine/Source/Editor/LevelEditor/Private/LevelEditorGenericDetails.cpp b/Engine/Source/Editor/LevelEditor/Private/LevelEditorGenericDetails.cpp
--- a/Engine/Source/Editor/LevelEditor/Private/LevelEditorGenericDetails.cpp
+++ b/Engine/Source/Editor/LevelEditor/Private/LevelEditorGenericDetails.cpp
@@ -217,7 +217,10 @@ void FLevelEditorGenericDetails::AddSurfaceDetails( IDetailLayoutBuilder& Detail
 		MaterialListDelegates.OnGetMaterials.BindSP(this, &FLevelEditorGenericDetails::GetSelectedSurfaceMaterials);
 		MaterialListDelegates.OnMaterialChanged.BindSP(this, &FLevelEditorGenericDetails::OnMaterialChanged);
 
-		TSharedRef<FMaterialList> MaterialList = MakeShareable(new FMaterialList(DetailBuilder, MaterialListDelegates));
+		//Pass an empty material list owner (owner can be use by the asset picker filter. In this case we do not need it)
+		TArray<FAssetData> MaterialListOwner;
+
+		TSharedRef<FMaterialList> MaterialList = MakeShareable(new FMaterialList(DetailBuilder, MaterialListDelegates, MaterialListOwner));
 
 		MaterialsCategory.AddCustomBuilder( MaterialList );
 	}
diff --git a/Engine/Source/Editor/Persona/Private/PersonaMeshDetails.cpp b/Engine/Source/Editor/Persona/Private/PersonaMeshDetails.cpp
--- a/Engine/Source/Editor/Persona/Private/PersonaMeshDetails.cpp
+++ b/Engine/Source/Editor/Persona/Private/PersonaMeshDetails.cpp
@@ -801,7 +801,10 @@ void FPersonaMeshDetails::AddLODLevelCategories(IDetailLayoutBuilder& DetailLayo
 				MaterialListDelegates.OnCanCopyMaterialItem.BindSP(this, &FPersonaMeshDetails::OnCanCopyMaterialItem);
 				MaterialListDelegates.OnPasteMaterialItem.BindSP(this, &FPersonaMeshDetails::OnPasteMaterialItem);
 
-				MaterialCategory.AddCustomBuilder(MakeShareable(new FMaterialList(MaterialCategory.GetParentLayout(), MaterialListDelegates, false, true, true)));
+				//Pass an empty material list owner (owner can be use by the asset picker filter. In this case we do not need it)
+				TArray<FAssetData> MaterialListOwner;
+				MaterialListOwner.Add(SkelMesh);
+				MaterialCategory.AddCustomBuilder(MakeShareable(new FMaterialList(MaterialCategory.GetParentLayout(), MaterialListDelegates, MaterialListOwner, false, true, true)));
 			}
 		}
 
diff --git a/Engine/Source/Editor/PropertyEditor/Private/MaterialList.cpp b/Engine/Source/Editor/PropertyEditor/Private/MaterialList.cpp
--- a/Engine/Source/Editor/PropertyEditor/Private/MaterialList.cpp
+++ b/Engine/Source/Editor/PropertyEditor/Private/MaterialList.cpp
@@ -140,7 +140,7 @@ public:
 			];
 	}
 
-	TSharedRef<SWidget> CreateValueContent( const TSharedPtr<FAssetThumbnailPool>& ThumbnailPool, TSharedPtr<IPropertyHandle> InHandle)
+	TSharedRef<SWidget> CreateValueContent( const TSharedPtr<FAssetThumbnailPool>& ThumbnailPool, const TArray<FAssetData>& OwnerAssetDataArray)
 	{
 		FIntPoint ThumbnailSize(64, 64);
 
@@ -171,7 +171,7 @@ public:
 						.ThumbnailPool(ThumbnailPool)
 						.DisplayCompactSize(bDisplayCompactSize)
 						.CustomResetToDefault(ResetToDefaultOverride)
-						.PropertyHandle(InHandle)
+						.OwnerAssetDataArray(OwnerAssetDataArray)
 						.CustomContentSlot()
 						[
 							SNew( SBox )
@@ -355,14 +355,14 @@ private:
 };
 
 
-FMaterialList::FMaterialList(IDetailLayoutBuilder& InDetailLayoutBuilder, FMaterialListDelegates& InMaterialListDelegates, bool bInAllowCollapse, bool bInShowUsedTextures, bool bInDisplayCompactSize, TSharedPtr<class IPropertyHandle> InHandle)
+FMaterialList::FMaterialList(IDetailLayoutBuilder& InDetailLayoutBuilder, FMaterialListDelegates& InMaterialListDelegates, const TArray<FAssetData>& InOwnerAssetDataArray, bool bInAllowCollapse, bool bInShowUsedTextures, bool bInDisplayCompactSize)
 	: MaterialListDelegates( InMaterialListDelegates )
 	, DetailLayoutBuilder( InDetailLayoutBuilder )
 	, MaterialListBuilder( new FMaterialListBuilder )
 	, bAllowCollpase(bInAllowCollapse)
 	, bShowUsedTextures(bInShowUsedTextures)
 	, bDisplayCompactSize(bInDisplayCompactSize)
-	, MeshChildHandle(InHandle)
+	, OwnerAssetDataArray(InOwnerAssetDataArray)
 {
 }
 
@@ -568,11 +568,6 @@ void FMaterialList::OnPasteMaterialList()
 	}
 }
 
-TSharedPtr<IPropertyHandle> FMaterialList::GetPropertyHandle() const
-{
-	return MeshChildHandle.IsValid() ? MeshChildHandle : nullptr;
-}
-
 bool FMaterialList::OnCanCopyMaterialItem(int32 CurrentSlot) const
 {
 	if (MaterialListDelegates.OnCanCopyMaterialItem.IsBound())
@@ -625,7 +620,7 @@ void FMaterialList::AddMaterialItem( FDetailWidgetRow& Row, int32 CurrentSlot, c
 	}
 	else
 	{
-		RightSideContent = NewView->CreateValueContent( DetailLayoutBuilder.GetThumbnailPool(), GetPropertyHandle() );
+		RightSideContent = NewView->CreateValueContent( DetailLayoutBuilder.GetThumbnailPool(), OwnerAssetDataArray );
 		ViewedMaterials.Add( NewView );
 	}
 
diff --git a/Engine/Source/Editor/PropertyEditor/Private/PropertyCustomizationHelpers.cpp b/Engine/Source/Editor/PropertyEditor/Private/PropertyCustomizationHelpers.cpp
--- a/Engine/Source/Editor/PropertyEditor/Private/PropertyCustomizationHelpers.cpp
+++ b/Engine/Source/Editor/PropertyEditor/Private/PropertyCustomizationHelpers.cpp
@@ -270,27 +270,28 @@ namespace PropertyCustomizationHelpers
 
 	TArray<const UClass*> EmptyClassArray;
 
-	TSharedRef<SWidget> MakeAssetPickerWithMenu(const FAssetData& InitialObject, const bool AllowClear, const TArray<const UClass*>& AllowedClasses, const TArray<UFactory*>& NewAssetFactories, FOnShouldFilterAsset OnShouldFilterAsset, FOnAssetSelected OnSet, FSimpleDelegate OnClose, const TSharedPtr<IPropertyHandle>& PropertyHandle)
+	TSharedRef<SWidget> MakeAssetPickerWithMenu(const FAssetData& InitialObject, const bool AllowClear, const TArray<const UClass*>& AllowedClasses, const TArray<UFactory*>& NewAssetFactories, FOnShouldFilterAsset OnShouldFilterAsset, FOnAssetSelected OnSet, FSimpleDelegate OnClose, const TSharedPtr<IPropertyHandle>& PropertyHandle, const TArray<FAssetData>& OwnerAssetArray)
 	{
-		return MakeAssetPickerWithMenu(InitialObject, AllowClear, true, AllowedClasses, EmptyClassArray, NewAssetFactories, OnShouldFilterAsset, OnSet, OnClose, PropertyHandle);
+		return MakeAssetPickerWithMenu(InitialObject, AllowClear, true, AllowedClasses, EmptyClassArray, NewAssetFactories, OnShouldFilterAsset, OnSet, OnClose, PropertyHandle, OwnerAssetArray);
 	}
 
-	TSharedRef<SWidget> MakeAssetPickerWithMenu(const FAssetData& InitialObject, const bool AllowClear, const TArray<const UClass*>& AllowedClasses, const TArray<const UClass*>& DisallowedClasses, const TArray<UFactory*>& NewAssetFactories, FOnShouldFilterAsset OnShouldFilterAsset, FOnAssetSelected OnSet, FSimpleDelegate OnClose, const TSharedPtr<IPropertyHandle>& PropertyHandle)
+	TSharedRef<SWidget> MakeAssetPickerWithMenu(const FAssetData& InitialObject, const bool AllowClear, const TArray<const UClass*>& AllowedClasses, const TArray<const UClass*>& DisallowedClasses, const TArray<UFactory*>& NewAssetFactories, FOnShouldFilterAsset OnShouldFilterAsset, FOnAssetSelected OnSet, FSimpleDelegate OnClose, const TSharedPtr<IPropertyHandle>& PropertyHandle, const TArray<FAssetData>& OwnerAssetArray)
 	{
-		return MakeAssetPickerWithMenu(InitialObject, AllowClear, true, AllowedClasses, DisallowedClasses, NewAssetFactories, OnShouldFilterAsset, OnSet, OnClose, PropertyHandle);
+		return MakeAssetPickerWithMenu(InitialObject, AllowClear, true, AllowedClasses, DisallowedClasses, NewAssetFactories, OnShouldFilterAsset, OnSet, OnClose, PropertyHandle, OwnerAssetArray);
 	}
 
-	TSharedRef<SWidget> MakeAssetPickerWithMenu(const FAssetData& InitialObject, const bool AllowClear, const bool AllowCopyPaste, const TArray<const UClass*>& AllowedClasses, const TArray<UFactory*>& NewAssetFactories, FOnShouldFilterAsset OnShouldFilterAsset, FOnAssetSelected OnSet, FSimpleDelegate OnClose, const TSharedPtr<IPropertyHandle>& PropertyHandle)
+	TSharedRef<SWidget> MakeAssetPickerWithMenu(const FAssetData& InitialObject, const bool AllowClear, const bool AllowCopyPaste, const TArray<const UClass*>& AllowedClasses, const TArray<UFactory*>& NewAssetFactories, FOnShouldFilterAsset OnShouldFilterAsset, FOnAssetSelected OnSet, FSimpleDelegate OnClose, const TSharedPtr<IPropertyHandle>& PropertyHandle, const TArray<FAssetData>& OwnerAssetArray)
 	{
-		return MakeAssetPickerWithMenu(InitialObject, AllowClear, AllowCopyPaste, AllowedClasses, EmptyClassArray, NewAssetFactories, OnShouldFilterAsset, OnSet, OnClose, PropertyHandle);
+		return MakeAssetPickerWithMenu(InitialObject, AllowClear, AllowCopyPaste, AllowedClasses, EmptyClassArray, NewAssetFactories, OnShouldFilterAsset, OnSet, OnClose, PropertyHandle, OwnerAssetArray);
 	}
 
-	TSharedRef<SWidget> MakeAssetPickerWithMenu( const FAssetData& InitialObject, const bool AllowClear, const bool AllowCopyPaste, const TArray<const UClass*>& AllowedClasses, const TArray<const UClass*>& DisallowedClasses, const TArray<UFactory*>& NewAssetFactories, FOnShouldFilterAsset OnShouldFilterAsset, FOnAssetSelected OnSet, FSimpleDelegate OnClose, const TSharedPtr<IPropertyHandle>& PropertyHandle)
+	TSharedRef<SWidget> MakeAssetPickerWithMenu( const FAssetData& InitialObject, const bool AllowClear, const bool AllowCopyPaste, const TArray<const UClass*>& AllowedClasses, const TArray<const UClass*>& DisallowedClasses, const TArray<UFactory*>& NewAssetFactories, FOnShouldFilterAsset OnShouldFilterAsset, FOnAssetSelected OnSet, FSimpleDelegate OnClose, const TSharedPtr<IPropertyHandle>& PropertyHandle, const TArray<FAssetData>& OwnerAssetArray)
 	{
 		return
 			SNew(SPropertyMenuAssetPicker)
 			.InitialObject(InitialObject)
 			.PropertyHandle(PropertyHandle)
+			.OwnerAssetArray(OwnerAssetArray)
 			.AllowClear(AllowClear)
 			.AllowCopyPaste(AllowCopyPaste)
 			.AllowedClasses(AllowedClasses)
@@ -455,6 +456,8 @@ void SObjectPropertyEntryBox::Construct( const FArguments& InArgs )
 	OnObjectChanged = InArgs._OnObjectChanged;
 	OnShouldSetAsset = InArgs._OnShouldSetAsset;
 
+	const TArray<FAssetData>& OwnerAssetDataArray = InArgs._OwnerAssetDataArray;
+
 	bool bDisplayThumbnail = InArgs._DisplayThumbnail;
 	FIntPoint ThumbnailSize(64, 64);
 	if (InArgs._ThumbnailSizeOverride.IsSet())
@@ -525,6 +528,7 @@ void SObjectPropertyEntryBox::Construct( const FArguments& InArgs )
 				.DisplayBrowse(InArgs._DisplayBrowse)
 				.EnableContentPicker(InArgs._EnableContentPicker)
 				.PropertyHandle(PropertyHandle)
+				.OwnerAssetDataArray(OwnerAssetDataArray)
 				.ThumbnailSize(ThumbnailSize)
 				.DisplayCompactSize(InArgs._DisplayCompactSize)
 				.CustomContentSlot()
diff --git a/Engine/Source/Editor/PropertyEditor/Private/UserInterface/PropertyEditor/SPropertyEditorAsset.cpp b/Engine/Source/Editor/PropertyEditor/Private/UserInterface/PropertyEditor/SPropertyEditorAsset.cpp
--- a/Engine/Source/Editor/PropertyEditor/Private/UserInterface/PropertyEditor/SPropertyEditorAsset.cpp
+++ b/Engine/Source/Editor/PropertyEditor/Private/UserInterface/PropertyEditor/SPropertyEditorAsset.cpp
@@ -78,6 +78,7 @@ void SPropertyEditorAsset::Construct( const FArguments& InArgs, const TSharedPtr
 {
 	PropertyEditor = InPropertyEditor;
 	PropertyHandle = InArgs._PropertyHandle;
+	OwnerAssetDataArray = InArgs._OwnerAssetDataArray;
 	OnSetObject = InArgs._OnSetObject;
 	OnShouldFilterAsset = InArgs._OnShouldFilterAsset;
 	bool DisplayCompactedSize = InArgs._DisplayCompactSize;
@@ -671,7 +672,8 @@ TSharedRef<SWidget> SPropertyEditorAsset::OnGetMenuContent()
 																	 OnShouldFilterAsset,
 																	 FOnAssetSelected::CreateSP(this, &SPropertyEditorAsset::OnAssetSelected),
 																	 FSimpleDelegate::CreateSP(this, &SPropertyEditorAsset::CloseComboButton),
-																	 GetMostSpecificPropertyHandle());
+																	 GetMostSpecificPropertyHandle(),
+																	 OwnerAssetDataArray);
 	}
 }
 
@@ -1311,6 +1313,19 @@ bool SPropertyEditorAsset::CanSetBasedOnAssetReferenceFilter( const FAssetData&
 				AssetReferenceFilterContext.ReferencingAssets.Add(FAssetData(ReferencingObject));
 			}
 		}
+		
+		if(OwnerAssetDataArray.Num() > 0)
+		{
+			for (const FAssetData& AssetData : OwnerAssetDataArray)
+			{
+				if (AssetData.IsValid())
+				{
+					//Use add unique in case the PropertyHandle as already add the referencing asset
+					AssetReferenceFilterContext.ReferencingAssets.AddUnique(AssetData);
+				}
+			}
+		}
+
 		TSharedPtr<IAssetReferenceFilter> AssetReferenceFilter = GUnrealEd->MakeAssetReferenceFilter(AssetReferenceFilterContext);
 		if (AssetReferenceFilter.IsValid() && !AssetReferenceFilter->PassesFilter(InAssetData, OutOptionalFailureReason))
 		{
diff --git a/Engine/Source/Editor/PropertyEditor/Private/UserInterface/PropertyEditor/SPropertyEditorAsset.h b/Engine/Source/Editor/PropertyEditor/Private/UserInterface/PropertyEditor/SPropertyEditorAsset.h
--- a/Engine/Source/Editor/PropertyEditor/Private/UserInterface/PropertyEditor/SPropertyEditorAsset.h
+++ b/Engine/Source/Editor/PropertyEditor/Private/UserInterface/PropertyEditor/SPropertyEditorAsset.h
@@ -77,6 +77,7 @@ public:
 		SLATE_NAMED_SLOT( FArguments, CustomContentSlot )
 		SLATE_NAMED_SLOT( FArguments, ResetToDefaultSlot )
 		SLATE_ARGUMENT( TSharedPtr<IPropertyHandle>, PropertyHandle )
+		SLATE_ARGUMENT( TArray<FAssetData>, OwnerAssetDataArray)
 
 	SLATE_END_ARGS()
 
@@ -364,4 +365,10 @@ private:
 
 	/** The property handle, if any */
 	TSharedPtr<class IPropertyHandle> PropertyHandle;
+
+	/*
+	 * The reference object on which the picker will assign the picked asset, if any.
+	 * The system will test first the PropertyHandle follow by the PropertyEditor and if nothing this array of assets
+	 */
+	TArray<FAssetData> OwnerAssetDataArray;
 };
diff --git a/Engine/Source/Editor/PropertyEditor/Private/UserInterface/PropertyEditor/SPropertyMenuAssetPicker.cpp b/Engine/Source/Editor/PropertyEditor/Private/UserInterface/PropertyEditor/SPropertyMenuAssetPicker.cpp
--- a/Engine/Source/Editor/PropertyEditor/Private/UserInterface/PropertyEditor/SPropertyMenuAssetPicker.cpp
+++ b/Engine/Source/Editor/PropertyEditor/Private/UserInterface/PropertyEditor/SPropertyMenuAssetPicker.cpp
@@ -21,6 +21,7 @@ void SPropertyMenuAssetPicker::Construct( const FArguments& InArgs )
 {
 	CurrentObject = InArgs._InitialObject;
 	PropertyHandle = InArgs._PropertyHandle;
+	const TArray<FAssetData>& OwnerAssetArray = InArgs._OwnerAssetArray;
 	bAllowClear = InArgs._AllowClear;
 	bAllowCopyPaste = InArgs._AllowCopyPaste;
 	AllowedClasses = InArgs._AllowedClasses;
@@ -146,6 +147,8 @@ void SPropertyMenuAssetPicker::Construct( const FArguments& InArgs )
 		AssetPickerConfig.SaveSettingsName = TEXT("AssetPropertyPicker");
 		// Populate the referencing assets via property handle
 		AssetPickerConfig.PropertyHandle = PropertyHandle;
+		// Populate the additional referencing assets with the Owner asset data
+		AssetPickerConfig.AdditionalReferencingAssets = OwnerAssetArray;
 
 		MenuContent =
 			SNew(SBox)
diff --git a/Engine/Source/Editor/PropertyEditor/Private/UserInterface/PropertyEditor/SPropertyMenuAssetPicker.h b/Engine/Source/Editor/PropertyEditor/Private/UserInterface/PropertyEditor/SPropertyMenuAssetPicker.h
--- a/Engine/Source/Editor/PropertyEditor/Private/UserInterface/PropertyEditor/SPropertyMenuAssetPicker.h
+++ b/Engine/Source/Editor/PropertyEditor/Private/UserInterface/PropertyEditor/SPropertyMenuAssetPicker.h
@@ -19,6 +19,7 @@ public:
 	{}
 		SLATE_ARGUMENT( FAssetData, InitialObject )
 		SLATE_ARGUMENT( TSharedPtr<IPropertyHandle>, PropertyHandle )
+		SLATE_ARGUMENT( TArray<FAssetData>, OwnerAssetArray)
 		SLATE_ARGUMENT( bool, AllowClear )
 		SLATE_ARGUMENT( bool, AllowCopyPaste )
 		SLATE_ARGUMENT( TArray<const UClass*>, AllowedClasses )
@@ -89,6 +90,8 @@ private:
 
 private:
 	FAssetData CurrentObject;
+	
+	/** The property this asset picker will modify (if any) */
 	TSharedPtr<IPropertyHandle> PropertyHandle;
 
 	/** Whether the asset can be 'None' in this case */
diff --git a/Engine/Source/Editor/PropertyEditor/Public/MaterialList.h b/Engine/Source/Editor/PropertyEditor/Public/MaterialList.h
--- a/Engine/Source/Editor/PropertyEditor/Public/MaterialList.h
+++ b/Engine/Source/Editor/PropertyEditor/Public/MaterialList.h
@@ -137,7 +137,7 @@ class FMaterialList
 	, public TSharedFromThis<FMaterialList>
 {
 public:
-	PROPERTYEDITOR_API FMaterialList( IDetailLayoutBuilder& InDetailLayoutBuilder, FMaterialListDelegates& MaterialListDelegates, bool bInAllowCollapse = false, bool bInShowUsedTextures = true, bool bInDisplayCompactSize = false, TSharedPtr<class IPropertyHandle> InHandle = nullptr);
+	PROPERTYEDITOR_API FMaterialList( IDetailLayoutBuilder& InDetailLayoutBuilder, FMaterialListDelegates& MaterialListDelegates, const TArray<FAssetData>& InOwnerAssetDataArray, bool bInAllowCollapse = false, bool bInShowUsedTextures = true, bool bInDisplayCompactSize = false);
 
 	/**
 	 * @return true if materials are being displayed.                                                          
@@ -168,7 +168,6 @@ private:
 	virtual void GenerateChildContent( IDetailChildrenBuilder& ChildrenBuilder ) override;
 	virtual FName GetName() const override { return NAME_None; }
 	virtual bool InitiallyCollapsed() const override { return bAllowCollpase; }
-	virtual TSharedPtr<IPropertyHandle> GetPropertyHandle() const;
 
 	/**
 	 * Adds a new material item to the list
@@ -216,7 +215,7 @@ private:
 	bool bShowUsedTextures;
 	/** Whether or not to display a compact form of material entry*/
 	bool bDisplayCompactSize;
-	/** The property handle for a child property on the mesh that owns these materials */
-	TSharedPtr<class IPropertyHandle> MeshChildHandle;
+	/** The mesh asset that owns these materials */
+	TArray<FAssetData> OwnerAssetDataArray;
 };
 
diff --git a/Engine/Source/Editor/PropertyEditor/Public/PropertyCustomizationHelpers.h b/Engine/Source/Editor/PropertyEditor/Public/PropertyCustomizationHelpers.h
--- a/Engine/Source/Editor/PropertyEditor/Public/PropertyCustomizationHelpers.h
+++ b/Engine/Source/Editor/PropertyEditor/Public/PropertyCustomizationHelpers.h
@@ -66,10 +66,10 @@ namespace PropertyCustomizationHelpers
 	PROPERTYEDITOR_API TSharedRef<SWidget> MakeUseSelectedButton( FSimpleDelegate OnUseSelectedClicked, TAttribute<FText> OptionalToolTipText = FText(), TAttribute<bool> IsEnabled = true );
 	PROPERTYEDITOR_API TSharedRef<SWidget> MakeBrowseButton( FSimpleDelegate OnClearClicked, TAttribute<FText> OptionalToolTipText = FText(), TAttribute<bool> IsEnabled = true );
 	PROPERTYEDITOR_API TSharedRef<SWidget> MakeAssetPickerAnchorButton( FOnGetAllowedClasses OnGetAllowedClasses, FOnAssetSelected OnAssetSelectedFromPicker, const TSharedPtr<IPropertyHandle>& PropertyHandle = TSharedPtr<IPropertyHandle>());
-	PROPERTYEDITOR_API TSharedRef<SWidget> MakeAssetPickerWithMenu( const FAssetData& InitialObject, const bool AllowClear, const TArray<const UClass*>& AllowedClasses, const TArray<UFactory*>& NewAssetFactories, FOnShouldFilterAsset OnShouldFilterAsset, FOnAssetSelected OnSet, FSimpleDelegate OnClose, const TSharedPtr<IPropertyHandle>& PropertyHandle = TSharedPtr<IPropertyHandle>());
-	PROPERTYEDITOR_API TSharedRef<SWidget> MakeAssetPickerWithMenu( const FAssetData& InitialObject, const bool AllowClear, const TArray<const UClass*>& AllowedClasses, const TArray<const UClass*>& DisallowedClasses, const TArray<UFactory*>& NewAssetFactories, FOnShouldFilterAsset OnShouldFilterAsset, FOnAssetSelected OnSet, FSimpleDelegate OnClose, const TSharedPtr<IPropertyHandle>& PropertyHandle = TSharedPtr<IPropertyHandle>());
-	PROPERTYEDITOR_API TSharedRef<SWidget> MakeAssetPickerWithMenu( const FAssetData& InitialObject, const bool AllowClear, const bool AllowCopyPaste, const TArray<const UClass*>& AllowedClasses, const TArray<UFactory*>& NewAssetFactories, FOnShouldFilterAsset OnShouldFilterAsset, FOnAssetSelected OnSet, FSimpleDelegate OnClose, const TSharedPtr<IPropertyHandle>& PropertyHandle = TSharedPtr<IPropertyHandle>());
-	PROPERTYEDITOR_API TSharedRef<SWidget> MakeAssetPickerWithMenu( const FAssetData& InitialObject, const bool AllowClear, const bool AllowCopyPaste, const TArray<const UClass*>& AllowedClasses, const TArray<const UClass*>& DisallowedClasses, const TArray<UFactory*>& NewAssetFactories, FOnShouldFilterAsset OnShouldFilterAsset, FOnAssetSelected OnSet, FSimpleDelegate OnClose, const TSharedPtr<IPropertyHandle>& PropertyHandle = TSharedPtr<IPropertyHandle>());
+	PROPERTYEDITOR_API TSharedRef<SWidget> MakeAssetPickerWithMenu( const FAssetData& InitialObject, const bool AllowClear, const TArray<const UClass*>& AllowedClasses, const TArray<UFactory*>& NewAssetFactories, FOnShouldFilterAsset OnShouldFilterAsset, FOnAssetSelected OnSet, FSimpleDelegate OnClose, const TSharedPtr<IPropertyHandle>& PropertyHandle = TSharedPtr<IPropertyHandle>(), const TArray<FAssetData>& OwnerAssetArray = TArray<FAssetData>());
+	PROPERTYEDITOR_API TSharedRef<SWidget> MakeAssetPickerWithMenu( const FAssetData& InitialObject, const bool AllowClear, const TArray<const UClass*>& AllowedClasses, const TArray<const UClass*>& DisallowedClasses, const TArray<UFactory*>& NewAssetFactories, FOnShouldFilterAsset OnShouldFilterAsset, FOnAssetSelected OnSet, FSimpleDelegate OnClose, const TSharedPtr<IPropertyHandle>& PropertyHandle = TSharedPtr<IPropertyHandle>(), const TArray<FAssetData>& OwnerAssetArray = TArray<FAssetData>());
+	PROPERTYEDITOR_API TSharedRef<SWidget> MakeAssetPickerWithMenu( const FAssetData& InitialObject, const bool AllowClear, const bool AllowCopyPaste, const TArray<const UClass*>& AllowedClasses, const TArray<UFactory*>& NewAssetFactories, FOnShouldFilterAsset OnShouldFilterAsset, FOnAssetSelected OnSet, FSimpleDelegate OnClose, const TSharedPtr<IPropertyHandle>& PropertyHandle = TSharedPtr<IPropertyHandle>(), const TArray<FAssetData>& OwnerAssetArray = TArray<FAssetData>());
+	PROPERTYEDITOR_API TSharedRef<SWidget> MakeAssetPickerWithMenu( const FAssetData& InitialObject, const bool AllowClear, const bool AllowCopyPaste, const TArray<const UClass*>& AllowedClasses, const TArray<const UClass*>& DisallowedClasses, const TArray<UFactory*>& NewAssetFactories, FOnShouldFilterAsset OnShouldFilterAsset, FOnAssetSelected OnSet, FSimpleDelegate OnClose, const TSharedPtr<IPropertyHandle>& PropertyHandle = TSharedPtr<IPropertyHandle>(), const TArray<FAssetData>& OwnerAssetArray = TArray<FAssetData>());
 	PROPERTYEDITOR_API TSharedRef<SWidget> MakeActorPickerAnchorButton( FOnGetActorFilters OnGetActorFilters, FOnActorSelected OnActorSelectedFromPicker );
 	PROPERTYEDITOR_API TSharedRef<SWidget> MakeActorPickerWithMenu( AActor* const InitialActor, const bool AllowClear, FOnShouldFilterActor ActorFilter, FOnActorSelected OnSet, FSimpleDelegate OnClose, FSimpleDelegate OnUseSelected );
 	PROPERTYEDITOR_API TSharedRef<SWidget> MakeComponentPickerWithMenu( UActorComponent* const InitialComponent, const bool AllowClear, FOnShouldFilterActor ActorFilter, FOnShouldFilterComponent ComponentFilter, FOnComponentSelected OnSet, FSimpleDelegate OnClose );
@@ -136,6 +136,8 @@ public:
 		SLATE_ATTRIBUTE( FString, ObjectPath )
 		/** Optional property handle that can be used instead of the object path */
 		SLATE_ARGUMENT( TSharedPtr<IPropertyHandle>, PropertyHandle )
+		/** Optional, array of the objects path, in case the property handle is not valid we will use this one to pass additional object to the picker config*/
+		SLATE_ARGUMENT(TArray<FAssetData>, OwnerAssetDataArray)
 		/** Thumbnail pool */
 		SLATE_ARGUMENT( TSharedPtr<FAssetThumbnailPool>, ThumbnailPool )
 		/** Class that is allowed in the asset picker */
diff --git a/Engine/Source/Editor/StaticMeshEditor/Private/StaticMeshEditorTools.cpp b/Engine/Source/Editor/StaticMeshEditor/Private/StaticMeshEditorTools.cpp
--- a/Engine/Source/Editor/StaticMeshEditor/Private/StaticMeshEditorTools.cpp
+++ b/Engine/Source/Editor/StaticMeshEditor/Private/StaticMeshEditorTools.cpp
@@ -2310,7 +2310,7 @@ UStaticMesh& FMeshMaterialsLayout::GetStaticMesh() const
 	return *StaticMesh;
 }
 
-void FMeshMaterialsLayout::AddToCategory(IDetailCategoryBuilder& CategoryBuilder, TSharedPtr<class IPropertyHandle> InHandle)
+void FMeshMaterialsLayout::AddToCategory(IDetailCategoryBuilder& CategoryBuilder, const TArray<FAssetData>& AssetDataArray)
 {
 	CategoryBuilder.AddCustomRow(LOCTEXT("AddLODLevelCategories_MaterialArrayOperationAdd", "Add Material Slot"))
 		.CopyAction(FUIAction(FExecuteAction::CreateSP(this, &FMeshMaterialsLayout::OnCopyMaterialList), FCanExecuteAction::CreateSP(this, &FMeshMaterialsLayout::OnCanCopyMaterialList)))
@@ -2376,7 +2376,7 @@ void FMeshMaterialsLayout::AddToCategory(IDetailCategoryBuilder& CategoryBuilder
 	MaterialListDelegates.OnCanCopyMaterialItem.BindSP(this, &FMeshMaterialsLayout::OnCanCopyMaterialItem);
 	MaterialListDelegates.OnPasteMaterialItem.BindSP(this, &FMeshMaterialsLayout::OnPasteMaterialItem);
 
-	CategoryBuilder.AddCustomBuilder(MakeShareable(new FMaterialList(CategoryBuilder.GetParentLayout(), MaterialListDelegates, false, true, true, InHandle)));
+	CategoryBuilder.AddCustomBuilder(MakeShareable(new FMaterialList(CategoryBuilder.GetParentLayout(), MaterialListDelegates, AssetDataArray, false, true, true)));
 }
 
 void FMeshMaterialsLayout::OnCopyMaterialList()
@@ -3268,9 +3268,10 @@ void FLevelOfDetailSettingsLayout::AddLODLevelCategories( IDetailLayoutBuilder&
 			FString CategoryName = FString(TEXT("StaticMeshMaterials"));
 
 			IDetailCategoryBuilder& MaterialsCategory = DetailBuilder.EditCategory(*CategoryName, LOCTEXT("StaticMeshMaterialsLabel", "Material Slots"), ECategoryPriority::Important);
-			TSharedPtr<IPropertyHandle> MeshLODHandle = DetailBuilder.GetProperty(GET_MEMBER_NAME_CHECKED(UStaticMesh, LODGroup));
 			MaterialsLayoutWidget = MakeShareable(new FMeshMaterialsLayout(StaticMeshEditor));
-			MaterialsLayoutWidget->AddToCategory(MaterialsCategory, MeshLODHandle);
+			TArray<FAssetData> AssetDataArray;
+			AssetDataArray.Add(FAssetData(StaticMesh, false));
+			MaterialsLayoutWidget->AddToCategory(MaterialsCategory, AssetDataArray);
 		}
 
 		int32 CurrentLodIndex = 0;
diff --git a/Engine/Source/Editor/StaticMeshEditor/Private/StaticMeshEditorTools.h b/Engine/Source/Editor/StaticMeshEditor/Private/StaticMeshEditorTools.h
--- a/Engine/Source/Editor/StaticMeshEditor/Private/StaticMeshEditorTools.h
+++ b/Engine/Source/Editor/StaticMeshEditor/Private/StaticMeshEditorTools.h
@@ -433,7 +433,7 @@ public:
 
 	virtual ~FMeshMaterialsLayout();
 
-	void AddToCategory(IDetailCategoryBuilder& CategoryBuilder, TSharedPtr<class IPropertyHandle> InHandle);
+	void AddToCategory(IDetailCategoryBuilder& CategoryBuilder, const TArray<FAssetData>& AssetDataArray);
 
 private:
 	UStaticMesh& GetStaticMesh() const;
