Detail item node filtering supports custom builder nodes with a valid filter string.
Category filtering no longer overrides non-text filters.

diff --git a/Engine/Source/Editor/DetailCustomizations/Private/InputSettingsDetails.cpp b/Engine/Source/Editor/DetailCustomizations/Private/InputSettingsDetails.cpp
--- a/Engine/Source/Editor/DetailCustomizations/Private/InputSettingsDetails.cpp
+++ b/Engine/Source/Editor/DetailCustomizations/Private/InputSettingsDetails.cpp
@@ -51,6 +51,7 @@ void FActionMappingsNodeBuilder::GenerateHeaderRowContent( FDetailWidgetRow& Nod
 		LOCTEXT("ClearActionMappingToolTip", "Removes all Action Mappings") );
 
 	NodeRow
+	.FilterString(ActionMappingsPropertyHandle->GetPropertyDisplayName())
 	[
 		SNew( SHorizontalBox )
 		+SHorizontalBox::Slot()
diff --git a/Engine/Source/Editor/PropertyEditor/Private/DetailItemNode.cpp b/Engine/Source/Editor/PropertyEditor/Private/DetailItemNode.cpp
--- a/Engine/Source/Editor/PropertyEditor/Private/DetailItemNode.cpp
+++ b/Engine/Source/Editor/PropertyEditor/Private/DetailItemNode.cpp
@@ -495,7 +495,7 @@ static bool PassesAllFilters( FDetailItemNode* ItemNode, const FDetailLayoutCust
 			const bool bIsSeenDueToFiltering = PropertyNodePin->HasNodeFlags(EPropertyNodeFlags::IsSeenDueToFiltering) != 0;
 			const bool bIsParentSeenDueToFiltering = PropertyNodePin->HasNodeFlags(EPropertyNodeFlags::IsParentSeenDueToFiltering) != 0;
 
-			const bool bPassesSearchFilter = bSearchFilterIsEmpty || ( bIsNotBeingFiltered || bIsSeenDueToFiltering || bIsParentSeenDueToFiltering );
+			const bool bPassesSearchFilter = bPassesCategoryFilter || bSearchFilterIsEmpty || ( bIsNotBeingFiltered || bIsSeenDueToFiltering || bIsParentSeenDueToFiltering );
 			const bool bPassesModifiedFilter = bPassesSearchFilter && ( InFilter.bShowOnlyModifiedProperties == false || PropertyNodePin->GetDiffersFromDefault() == true );
 			const bool bPassesDifferingFilter = InFilter.bShowOnlyDiffering ? InFilter.WhitelistedProperties.Find(*FPropertyNode::CreatePropertyPath(PropertyNodePin.ToSharedRef())) != nullptr : true;
 
@@ -516,16 +516,25 @@ static bool PassesAllFilters( FDetailItemNode* ItemNode, const FDetailLayoutCust
 			const bool bPassesAnimatedFilter = (InFilter.bShowAnimated == false || Local::ItemIsAnimated(ItemNode, PropertyNodePin));
 
 			// The property node is visible (note categories are never visible unless they have a child that is visible )
-			bPassesAllFilters = (bPassesSearchFilter && bPassesModifiedFilter && bPassesDifferingFilter && bPassesKeyableFilter && bPassesAnimatedFilter) || bPassesCategoryFilter;
+			bPassesAllFilters = bPassesSearchFilter && bPassesModifiedFilter && bPassesDifferingFilter && bPassesKeyableFilter && bPassesAnimatedFilter;
 		}
 		else if (InCustomization.HasCustomWidget())
 		{
-			const bool bPassesTextFilter = Local::StringPassesFilter(InFilter, InCustomization.WidgetDecl->FilterTextString.ToString());
+			const bool bPassesTextFilter = bPassesCategoryFilter || Local::StringPassesFilter(InFilter, InCustomization.WidgetDecl->FilterTextString.ToString());
 			const bool bPassesModifiedFilter = (InFilter.bShowOnlyModifiedProperties == false || InCustomization.WidgetDecl->DiffersFromDefaultAttr.Get() == true);
 			//@todo we need to support custom widgets for keyable,animated, in particularly for transforms(ComponentTransformDetails).
 			const bool bPassesKeyableFilter = (InFilter.bShowKeyable == false);
 			const bool bPassesAnimatedFilter = (InFilter.bShowAnimated == false);
-			bPassesAllFilters = (bPassesTextFilter && bPassesModifiedFilter && bPassesKeyableFilter && bPassesAnimatedFilter) || bPassesCategoryFilter;
+			bPassesAllFilters = bPassesTextFilter && bPassesModifiedFilter && bPassesKeyableFilter && bPassesAnimatedFilter;
+		}
+		else if (InCustomization.HasCustomBuilder())
+		{
+			const bool bPassesTextFilter = bPassesCategoryFilter || Local::StringPassesFilter(InFilter, InCustomization.CustomBuilderRow->GetWidgetRow().FilterTextString.ToString());
+			//@todo we need to support custom builders for modified, keyable, animated, in particularly for transforms(ComponentTransformDetails).
+			const bool bPassesModifiedFilter = (InFilter.bShowOnlyModifiedProperties == false);
+			const bool bPassesKeyableFilter = (InFilter.bShowKeyable == false);
+			const bool bPassesAnimatedFilter = (InFilter.bShowAnimated == false);
+			bPassesAllFilters = bPassesTextFilter && bPassesModifiedFilter && bPassesKeyableFilter && bPassesAnimatedFilter;
 		}
 	}
 
