Fix crash when a Property doesn't have a child node at the requested index. It would return no property so that the user knows that the path was wrong. Tested with A[Index] and A.B[Index].C

diff --git a/Engine/Source/Editor/PropertyEditor/Private/DetailLayoutBuilderImpl.cpp b/Engine/Source/Editor/PropertyEditor/Private/DetailLayoutBuilderImpl.cpp
--- a/Engine/Source/Editor/PropertyEditor/Private/DetailLayoutBuilderImpl.cpp
+++ b/Engine/Source/Editor/PropertyEditor/Private/DetailLayoutBuilderImpl.cpp
@@ -523,6 +523,11 @@ TSharedPtr<FPropertyNode> FDetailLayoutBuilderImpl::GetPropertyNodeInternal( con
 					{
 						if( Index != INDEX_NONE )
 						{
+							if (Index >= PropertyNode->GetNumChildNodes())
+							{
+								return nullptr;
+							}
+
 							// The parent is the actual array, its children are array elements
 							PropertyNode = PropertyNode->GetChildNode( Index );
 						}
