SPropertyEditorEditInline: Fixed not being able to edit properties whose editability is not constant.

diff --git a/Engine/Source/Editor/PropertyEditor/Private/PropertyNode.cpp b/Engine/Source/Editor/PropertyEditor/Private/PropertyNode.cpp
--- a/Engine/Source/Editor/PropertyEditor/Private/PropertyNode.cpp
+++ b/Engine/Source/Editor/PropertyEditor/Private/PropertyNode.cpp
@@ -725,6 +725,19 @@ TSharedPtr<FPropertyNode> FPropertyNode::FindChildPropertyNode( const FName InPr
 	return nullptr;
 }
 
+/**
+ * Returns whether this window's property is read only or has the CPF_EditConst flag.
+ */
+bool FPropertyNode::IsPropertyConst() const
+{
+	bool bIsPropertyConst = (HasNodeFlags(EPropertyNodeFlags::IsReadOnly) != 0);
+	if (!bIsPropertyConst && Property != nullptr)
+	{
+		bIsPropertyConst = (Property->PropertyFlags & CPF_EditConst) ? true : false;	
+	}
+
+	return bIsPropertyConst;
+}
 
 /** @return whether this window's property is constant (can't be edited by the user) */
 bool FPropertyNode::IsEditConst() const
@@ -734,23 +747,19 @@ bool FPropertyNode::IsEditConst() const
 		// Ask the objects whether this property can be changed
 		const FObjectPropertyNode* ObjectPropertyNode = FindObjectItemParent();
 
-		bIsEditConst = (HasNodeFlags(EPropertyNodeFlags::IsReadOnly) != 0);
+		bIsEditConst = IsPropertyConst();
 		if(!bIsEditConst && Property != nullptr && ObjectPropertyNode)
 		{
-			bIsEditConst = (Property->PropertyFlags & CPF_EditConst) ? true : false;
-			if(!bIsEditConst)
+			// travel up the chain to see if this property's owner struct is editconst - if it is, so is this property
+			FPropertyNode* NextParent = ParentNode;
+			while(NextParent != nullptr && Cast<UStructProperty>(NextParent->GetProperty()) != NULL)
 			{
-				// travel up the chain to see if this property's owner struct is editconst - if it is, so is this property
-				FPropertyNode* NextParent = ParentNode;
-				while(NextParent != nullptr && Cast<UStructProperty>(NextParent->GetProperty()) != NULL)
+				if(NextParent->IsEditConst())
 				{
-					if(NextParent->IsEditConst())
-					{
-						bIsEditConst = true;
-						break;
-					}
-					NextParent = NextParent->ParentNode;
+					bIsEditConst = true;
+					break;
 				}
+				NextParent = NextParent->ParentNode;
 			}
 
 			if(!bIsEditConst)
diff --git a/Engine/Source/Editor/PropertyEditor/Private/PropertyNode.h b/Engine/Source/Editor/PropertyEditor/Private/PropertyNode.h
--- a/Engine/Source/Editor/PropertyEditor/Private/PropertyNode.h
+++ b/Engine/Source/Editor/PropertyEditor/Private/PropertyNode.h
@@ -401,6 +401,11 @@ public:
 	*/
 	bool GetChildNode(const int32 ChildArrayIndex, TSharedPtr<FPropertyNode>& OutChildNode) const;
 
+	/**
+	 * Returns whether this window's property is read only or has the CPF_EditConst flag.
+	 */
+	bool IsPropertyConst() const;
+
 	/** @return whether this window's property is constant (can't be edited by the user) */
 	bool IsEditConst() const;
 
diff --git a/Engine/Source/Editor/PropertyEditor/Private/UserInterface/PropertyEditor/SPropertyEditorEditInline.cpp b/Engine/Source/Editor/PropertyEditor/Private/UserInterface/PropertyEditor/SPropertyEditorEditInline.cpp
--- a/Engine/Source/Editor/PropertyEditor/Private/UserInterface/PropertyEditor/SPropertyEditorEditInline.cpp
+++ b/Engine/Source/Editor/PropertyEditor/Private/UserInterface/PropertyEditor/SPropertyEditorEditInline.cpp
@@ -149,7 +149,7 @@ bool SPropertyEditorEditInline::Supports( const FPropertyNode* InTreeNode, int32
 	return InTreeNode
 		&& InTreeNode->HasNodeFlags(EPropertyNodeFlags::EditInlineNew)
 		&& InTreeNode->FindObjectItemParent()
-		&& !InTreeNode->IsEditConst();
+		&& !InTreeNode->IsPropertyConst();
 }
 
 bool SPropertyEditorEditInline::Supports( const TSharedRef< class FPropertyEditor >& InPropertyEditor )
