Fixed an invalid bounds check in ItemPropertyNode::GetValueBaseAddress when fetching the base address for a property in an array. For example, the CDO of an object may not contain a value at a given array index, but still contain some elements.

diff --git a/Engine/Source/Editor/PropertyEditor/Private/ItemPropertyNode.cpp b/Engine/Source/Editor/PropertyEditor/Private/ItemPropertyNode.cpp
--- a/Engine/Source/Editor/PropertyEditor/Private/ItemPropertyNode.cpp
+++ b/Engine/Source/Editor/PropertyEditor/Private/ItemPropertyNode.cpp
@@ -41,9 +41,9 @@ uint8* FItemPropertyNode::GetValueBaseAddress( uint8* StartAddress )
 		if (OuterArrayProp != nullptr)
 		{
 			FScriptArrayHelper ArrayHelper(OuterArrayProp, ValueBaseAddress);
-			if (ValueBaseAddress != nullptr && ArrayHelper.Num() > 0)
+			if (ValueBaseAddress != nullptr && ArrayHelper.IsValidIndex(ArrayIndex))
 			{
-				return ArrayHelper.GetRawPtr() + ArrayOffset;
+				return ArrayHelper.GetRawPtr(ArrayIndex);
 			}
 		}
 		else if (OuterSetProp != nullptr)
