Added a sidecar data structure to hold data that is constant for all members of the same class. This data will now be stored per class instead of being stored per instance. This can result in substantial memory savings for classes that have many instances and constant data. Good candidates for identify this type of data are member variables that are marked as EditDefaultsOnly and BlueprintReadOnly.

diff --git a/Engine/Source/Editor/BlueprintGraph/Classes/K2Node_Variable.h b/Engine/Source/Editor/BlueprintGraph/Classes/K2Node_Variable.h
--- a/Engine/Source/Editor/BlueprintGraph/Classes/K2Node_Variable.h
+++ b/Engine/Source/Editor/BlueprintGraph/Classes/K2Node_Variable.h
@@ -145,6 +145,9 @@ public:
 	static bool DoesRenamedVariableMatch(FName OldVariableName, FName NewVariableName, UStruct* StructType);
 
 private:
+	/** 
+	 * Gets the property for the variable on the owning class or on the owning class's sparse class data structure.
+	 */
 	UProperty* GetPropertyForVariable_Internal(UClass* OwningClass) const;
 
 public:
diff --git a/Engine/Source/Editor/BlueprintGraph/Private/BlueprintActionDatabase.cpp b/Engine/Source/Editor/BlueprintGraph/Private/BlueprintActionDatabase.cpp
--- a/Engine/Source/Editor/BlueprintGraph/Private/BlueprintActionDatabase.cpp
+++ b/Engine/Source/Editor/BlueprintGraph/Private/BlueprintActionDatabase.cpp
@@ -354,6 +354,15 @@ namespace BlueprintActionDatabaseImpl
 	 */
 	static void AddClassPropertyActions(UClass const* const Class, FActionList& ActionListOut);
 
+	/**
+	 * Loops over all of the class's data object's properties and creates node-spawners for
+	 * any that are viable for blueprint use.
+	 *
+	 * @param  Class			The class whose properties you want node-spawners for.
+	 * @param  ActionListOut	The list you want populated with the new spawners.
+	 */
+	static void AddClassDataObjectActions(UClass const* const Class, FActionList& ActionListOut);
+
 	/**
 	 * Evolved from FClassDynamicCastHelper::GetClassDynamicCastNodes(). If the
 	 * specified class is a viable blueprint variable type, then two cast nodes
@@ -599,6 +608,7 @@ static void BlueprintActionDatabaseImpl::GetClassMemberActions(UClass* const Cla
 	{
 		AddClassFunctionActions(Class, ActionListOut);
 		AddClassPropertyActions(Class, ActionListOut);
+		AddClassDataObjectActions(Class, ActionListOut);
 		// class UEnum actions are added by individual nodes via GetNodeSpecificActions()
 		// class UScriptStruct actions are added by individual nodes via GetNodeSpecificActions()
 	}
@@ -664,7 +674,7 @@ static void BlueprintActionDatabaseImpl::AddClassPropertyActions(UClass const* c
 	bool const bIsActorClass = Class->IsChildOf<AActor>();
 	
 	// loop over all the properties in the specified class; exclude-super because 
-	// we can always get the super properties by looking up that class separateHavely 
+	// we can always get the super properties by looking up that class separately 
 	for (TFieldIterator<UProperty> PropertyIt(Class, EFieldIteratorFlags::ExcludeSuper); PropertyIt; ++PropertyIt)
 	{
 		UProperty* Property = *PropertyIt;
@@ -716,6 +726,33 @@ static void BlueprintActionDatabaseImpl::AddClassPropertyActions(UClass const* c
 	}
 }
 
+//------------------------------------------------------------------------------
+static void BlueprintActionDatabaseImpl::AddClassDataObjectActions(UClass const* const Class, FActionList& ActionListOut)
+{
+	using namespace FBlueprintNodeSpawnerFactory; // for MakeDelegateNodeSpawner()
+
+	// loop over all the properties in the specified class; exclude-super because 
+	// we can always get the super properties by looking up that class separately 
+	const UStruct* SparseDataStruct = Class->GetSparseClassDataStruct();
+	const UStruct* ParentSparseDataStruct = Class->GetSuperClass() ? Class->GetSuperClass()->GetSparseClassDataStruct() : nullptr;
+	if (ParentSparseDataStruct != SparseDataStruct)
+	{
+		for (TFieldIterator<UProperty> PropertyIt(SparseDataStruct, EFieldIteratorFlags::ExcludeSuper); PropertyIt; ++PropertyIt)
+		{
+			UProperty* Property = *PropertyIt;
+			if (!IsPropertyBlueprintVisible(Property))
+			{
+				continue;
+			}
+
+			UBlueprintVariableNodeSpawner* GetterSpawner = UBlueprintVariableNodeSpawner::CreateFromMemberOrParam(UK2Node_VariableGet::StaticClass(), Property);
+			ActionListOut.Add(GetterSpawner);
+//			UBlueprintVariableNodeSpawner* SetterSpawner = UBlueprintVariableNodeSpawner::CreateFromMemberOrParam(UK2Node_VariableSet::StaticClass(), Property);
+//			ActionListOut.Add(SetterSpawner);
+		}
+	}
+}
+
 //------------------------------------------------------------------------------
 static void BlueprintActionDatabaseImpl::AddClassCastActions(UClass* Class, FActionList& ActionListOut)
 {
diff --git a/Engine/Source/Editor/BlueprintGraph/Private/BlueprintVariableNodeSpawner.cpp b/Engine/Source/Editor/BlueprintGraph/Private/BlueprintVariableNodeSpawner.cpp
--- a/Engine/Source/Editor/BlueprintGraph/Private/BlueprintVariableNodeSpawner.cpp
+++ b/Engine/Source/Editor/BlueprintGraph/Private/BlueprintVariableNodeSpawner.cpp
@@ -19,7 +19,7 @@
  ******************************************************************************/
 
 //------------------------------------------------------------------------------
-UBlueprintVariableNodeSpawner* UBlueprintVariableNodeSpawner::CreateFromMemberOrParam(TSubclassOf<UK2Node_Variable> NodeClass, UProperty const* VarProperty, UEdGraph* VarContext, UObject* Outer/*= nullptr*/)
+UBlueprintVariableNodeSpawner* UBlueprintVariableNodeSpawner::CreateFromMemberOrParam(TSubclassOf<UK2Node_Variable> NodeClass, UProperty const* VarProperty, UEdGraph* VarContext, UObject* Outer)
 {
 	check(VarProperty != nullptr);
 	if (Outer == nullptr)
diff --git a/Engine/Source/Editor/BlueprintGraph/Private/K2Node_SetVariableOnPersistentFrame.cpp b/Engine/Source/Editor/BlueprintGraph/Private/K2Node_SetVariableOnPersistentFrame.cpp
--- a/Engine/Source/Editor/BlueprintGraph/Private/K2Node_SetVariableOnPersistentFrame.cpp
+++ b/Engine/Source/Editor/BlueprintGraph/Private/K2Node_SetVariableOnPersistentFrame.cpp
@@ -34,7 +34,10 @@ public:
 				return;
 			}
 
-			UProperty* BoundProperty = FKismetCompilerUtilities::FindPropertyInScope(Context.NewClass->UberGraphFunction, Pin, CompilerContext.MessageLog, CompilerContext.GetSchema(), Context.NewClass);
+			bool bIsSparseProperty;
+			UProperty* BoundProperty = FKismetCompilerUtilities::FindPropertyInScope(Context.NewClass->UberGraphFunction, Pin, CompilerContext.MessageLog, CompilerContext.GetSchema(), Context.NewClass, bIsSparseProperty);
+			// no setters on sparse properties for now
+			check(!bIsSparseProperty);
 			if (!BoundProperty || (BoundProperty->GetOuter() != Context.NewClass->UberGraphFunction))
 			{
 				CompilerContext.MessageLog.Error(*LOCTEXT("SetVariableOnPersistentFrame_IceErrorNoProperty", "ICE SetVariableOnPersistentFrame - No property found. @@").ToString(), Pin);
diff --git a/Engine/Source/Editor/BlueprintGraph/Private/K2Node_Variable.cpp b/Engine/Source/Editor/BlueprintGraph/Private/K2Node_Variable.cpp
--- a/Engine/Source/Editor/BlueprintGraph/Private/K2Node_Variable.cpp
+++ b/Engine/Source/Editor/BlueprintGraph/Private/K2Node_Variable.cpp
@@ -110,8 +110,8 @@ void UK2Node_Variable::CreatePinForSelf()
 	// Create the self pin
 	if (!K2Schema->FindSelfPin(*this, EGPD_Input))
 	{
-		// Do not create a self pin for locally scoped variables
-		if( !VariableReference.IsLocalScope() )
+		// Do not create a self pin for locally scoped variables or sparse class data
+		if (!VariableReference.IsLocalScope() && !VariableReference.IsSparseClassData(GetBlueprintClassFromNode()))
 		{
 			bool bSelfTarget = VariableReference.IsSelfContext() && (ESelfContextInfo::NotSelfContext != SelfContextInfo);
 			UClass* MemberParentClass = VariableReference.GetMemberParentClass(GetBlueprintClassFromNode());
@@ -132,18 +132,18 @@ void UK2Node_Variable::CreatePinForSelf()
 
 			// Self Target pins should always make the class be the owning class of the property,
 			// so if the node is from a Macro Blueprint, it will hook up as self in any placed Blueprint
-			if(bSelfTarget)
+			if (bSelfTarget)
 			{
 				if(UProperty* Property = VariableReference.ResolveMember<UProperty>(GetBlueprintClassFromNode()))
 				{
 					TargetClass = Property->GetOwnerClass()->GetAuthoritativeClass();
 				}
-				else if(GetBlueprint()->SkeletonGeneratedClass)
+				else if (GetBlueprint()->SkeletonGeneratedClass)
 				{
 					TargetClass = GetBlueprint()->SkeletonGeneratedClass->GetAuthoritativeClass();
 				}
 			}
-			else if(TargetClass && TargetClass->ClassGeneratedBy)
+			else if (TargetClass && TargetClass->ClassGeneratedBy)
 			{
 				TargetClass = TargetClass->GetAuthoritativeClass();
 			}
@@ -376,13 +376,24 @@ UClass* UK2Node_Variable::GetVariableSourceClass() const
 UProperty* UK2Node_Variable::GetPropertyForVariable_Internal(UClass* OwningClass) const
 {
 	const FName VarName = GetVarName();
-	UEdGraphPin* VariablePin = FindPin(VarName);
 
-	UProperty* VariableProperty = VariableReference.ResolveMember<UProperty>(OwningClass);
+	// Look in the sparse class data first
+	UProperty* VariableProperty = nullptr;
+	UScriptStruct* SparseClassDataStruct = OwningClass ? OwningClass->GetSparseClassDataStruct() : nullptr;
+	if (SparseClassDataStruct)
+	{
+		VariableProperty = FindField<UProperty>(SparseClassDataStruct, VarName);
+	}
+	if (!VariableProperty)
+	{
+		VariableProperty = VariableReference.ResolveMember<UProperty>(OwningClass);
+	}
 
 	// if the variable has been deprecated, don't use it
 	if (VariableProperty != nullptr)
 	{
+		UEdGraphPin* VariablePin = FindPin(VarName);
+
 		// If the variable has been remapped update the pin
 		if (VariablePin && VarName != GetVarName())
 		{
diff --git a/Engine/Source/Editor/BlueprintGraph/Private/StructMemberNodeHandlers.cpp b/Engine/Source/Editor/BlueprintGraph/Private/StructMemberNodeHandlers.cpp
--- a/Engine/Source/Editor/BlueprintGraph/Private/StructMemberNodeHandlers.cpp
+++ b/Engine/Source/Editor/BlueprintGraph/Private/StructMemberNodeHandlers.cpp
@@ -20,7 +20,8 @@ static FBPTerminal* RegisterStructVar(FCompilerResultsLog& MessageLog, FKismetFu
 	UStruct* SearchScope = (SelfPin != NULL) ? Context.GetScopeFromPinType(SelfPin->PinType, Context.NewClass) : Context.Function;
 
 	// Now find the variable
-	if (UProperty* BoundProperty = FKismetCompilerUtilities::FindNamedPropertyInScope(SearchScope, MemberSetNode->GetVarName()))
+	bool bIsSparseProperty;
+	if (UProperty* BoundProperty = FKismetCompilerUtilities::FindNamedPropertyInScope(SearchScope, MemberSetNode->GetVarName(), bIsSparseProperty))
 	{
 		// Create the term in the list
 		FBPTerminal* Term = new FBPTerminal();
diff --git a/Engine/Source/Editor/Kismet/Private/BlueprintDetailsCustomization.cpp b/Engine/Source/Editor/Kismet/Private/BlueprintDetailsCustomization.cpp
--- a/Engine/Source/Editor/Kismet/Private/BlueprintDetailsCustomization.cpp
+++ b/Engine/Source/Editor/Kismet/Private/BlueprintDetailsCustomization.cpp
@@ -878,7 +878,8 @@ void FBlueprintVarActionDetails::CustomizeDetails( IDetailLayoutBuilder& DetailL
 				else if (CachedVariableProperty.IsValid())
 				{
 					// Capture the non-BP class CDO so we can show the default value
-					TargetBlueprintDefaultObject = CachedVariableProperty->GetOwnerClass()->GetDefaultObject();
+					UClass* OwnerClass = CachedVariableProperty->GetOwnerClass();
+					TargetBlueprintDefaultObject = OwnerClass ? OwnerClass->GetDefaultObject() : nullptr;
 				}
 
 				if (TargetBlueprintDefaultObject != nullptr)
diff --git a/Engine/Source/Editor/KismetCompiler/Private/KismetCompilerMisc.cpp b/Engine/Source/Editor/KismetCompiler/Private/KismetCompilerMisc.cpp
--- a/Engine/Source/Editor/KismetCompiler/Private/KismetCompilerMisc.cpp
+++ b/Engine/Source/Editor/KismetCompiler/Private/KismetCompilerMisc.cpp
@@ -47,7 +47,6 @@
 DECLARE_CYCLE_STAT(TEXT("Choose Terminal Scope"), EKismetCompilerStats_ChooseTerminalScope, STATGROUP_KismetCompiler);
 DECLARE_CYCLE_STAT(TEXT("Resolve compiled statements"), EKismetCompilerStats_ResolveCompiledStatements, STATGROUP_KismetCompiler );
 
-
 //////////////////////////////////////////////////////////////////////////
 // FKismetCompilerUtilities
 
@@ -519,8 +518,9 @@ void FKismetCompilerUtilities::RemoveObjectRedirectorIfPresent(UObject* Package,
 }
 
 /** Finds a property by name, starting in the specified scope; Validates property type and returns NULL along with emitting an error if there is a mismatch. */
-UProperty* FKismetCompilerUtilities::FindPropertyInScope(UStruct* Scope, UEdGraphPin* Pin, FCompilerResultsLog& MessageLog, const UEdGraphSchema_K2* Schema, UClass* SelfClass)
+UProperty* FKismetCompilerUtilities::FindPropertyInScope(UStruct* Scope, UEdGraphPin* Pin, FCompilerResultsLog& MessageLog, const UEdGraphSchema_K2* Schema, UClass* SelfClass, bool& bIsSparseProperty)
 {
+	bIsSparseProperty = false;
 	UStruct* InitialScope = Scope;
 	while (Scope != NULL)
 	{
@@ -542,6 +542,22 @@ UProperty* FKismetCompilerUtilities::FindPropertyInScope(UStruct* Scope, UEdGrap
 			}
 		}
 
+		// If this is a class, check the sparse data for the property
+		UClass* Class = Cast<UClass>(Scope);
+		if (Class)
+		{
+			UStruct* SparseData = Class->GetSparseClassDataStruct();
+			if (SparseData)
+			{
+				UProperty* Prop = FindPropertyInScope(SparseData, Pin, MessageLog, Schema, SelfClass, bIsSparseProperty);
+				if (Prop)
+				{
+					bIsSparseProperty = true;
+					return Prop;
+				}
+			}
+		}
+
 		// Functions don't automatically check their class when using a field iterator
 		UFunction* Function = Cast<UFunction>(Scope);
 		Scope = (Function != NULL) ? Cast<UStruct>(Function->GetOuter()) : NULL;
@@ -556,8 +572,9 @@ UProperty* FKismetCompilerUtilities::FindPropertyInScope(UStruct* Scope, UEdGrap
 }
 
 // Finds a property by name, starting in the specified scope, returning NULL if it's not found
-UProperty* FKismetCompilerUtilities::FindNamedPropertyInScope(UStruct* Scope, FName PropertyName)
+UProperty* FKismetCompilerUtilities::FindNamedPropertyInScope(UStruct* Scope, FName PropertyName, bool& bIsSparseProperty)
 {
+	bIsSparseProperty = false;
 	while (Scope != NULL)
 	{
 		for (TFieldIterator<UProperty> It(Scope, EFieldIteratorFlags::IncludeSuper); It; ++It)
@@ -571,6 +588,22 @@ UProperty* FKismetCompilerUtilities::FindNamedPropertyInScope(UStruct* Scope, FN
 			}
 		}
 
+		// If this is a class, check the sparse data for the property
+		UClass* Class = Cast<UClass>(Scope);
+		if (Class)
+		{
+			UStruct* SparseData = Class->GetSparseClassDataStruct();
+			if (SparseData)
+			{
+				UProperty* Prop = FindNamedPropertyInScope(SparseData, PropertyName, bIsSparseProperty);
+				if (Prop)
+				{
+					bIsSparseProperty = true;
+					return Prop;
+				}
+			}
+		}
+
 		// Functions don't automatically check their class when using a field iterator
 		UFunction* Function = Cast<UFunction>(Scope);
 		Scope = (Function != NULL) ? Cast<UStruct>(Function->GetOuter()) : NULL;
@@ -1594,7 +1627,8 @@ void FNodeHandlingFunctor::ResolveAndRegisterScopedTerm(FKismetFunctionContext&
 	}
 
 	// Find the variable in the search scope
-	UProperty* BoundProperty = FKismetCompilerUtilities::FindPropertyInScope(SearchScope, Net, CompilerContext.MessageLog, CompilerContext.GetSchema(), Context.NewClass);
+	bool bIsSparseProperty;
+	UProperty* BoundProperty = FKismetCompilerUtilities::FindPropertyInScope(SearchScope, Net, CompilerContext.MessageLog, CompilerContext.GetSchema(), Context.NewClass, bIsSparseProperty);
 	if (BoundProperty != NULL)
 	{
 		UBlueprintEditorSettings* Settings = GetMutableDefault<UBlueprintEditorSettings>();
@@ -1617,6 +1651,11 @@ void FNodeHandlingFunctor::ResolveAndRegisterScopedTerm(FKismetFunctionContext&
 			Term->bIsConst = true;
 		}
 
+		if (bIsSparseProperty)
+		{
+			Term->SetVarTypeSparseClassData();
+		}
+
 		// Resolve the context term
 		if (SelfPin != NULL)
 		{
diff --git a/Engine/Source/Editor/KismetCompiler/Private/KismetCompilerVMBackend.cpp b/Engine/Source/Editor/KismetCompiler/Private/KismetCompilerVMBackend.cpp
--- a/Engine/Source/Editor/KismetCompiler/Private/KismetCompilerVMBackend.cpp
+++ b/Engine/Source/Editor/KismetCompiler/Private/KismetCompilerVMBackend.cpp
@@ -1033,6 +1033,10 @@ public:
 			{
 				Writer << (Term->AssociatedVarProperty->HasAnyPropertyFlags(CPF_OutParm) ? EX_LocalOutVariable : EX_LocalVariable);
 			}
+			else if (Term->IsSparseClassDataVarTerm())
+			{
+				Writer << EX_ClassSparseDataVariable;
+			}
 			else
 			{
 				Writer << EX_InstanceVariable;
diff --git a/Engine/Source/Editor/KismetCompiler/Public/BPTerminal.h b/Engine/Source/Editor/KismetCompiler/Public/BPTerminal.h
--- a/Engine/Source/Editor/KismetCompiler/Public/BPTerminal.h
+++ b/Engine/Source/Editor/KismetCompiler/Public/BPTerminal.h
@@ -79,6 +79,16 @@ struct FBPTerminal
 		VarType = bIsLocal ? EVarType_Local : EVarType_Instanced;
 	}
 
+	bool IsSparseClassDataVarTerm() const
+	{
+		return !bIsLiteral && VarType == EVarType_SparseClassData;
+	}
+
+	void SetVarTypeSparseClassData()
+	{
+		VarType = EVarType_SparseClassData;
+	}
+
 	bool IsDefaultVarTerm() const
 	{
 		return !bIsLiteral && VarType == EVarType_Default;
@@ -125,7 +135,8 @@ private:
 	{
 		EVarType_Local,
 		EVarType_Default,
-		EVarType_Instanced
+		EVarType_Instanced,
+		EVarType_SparseClassData,
 	};
 
 	// For non-literal terms, this is set to the type of variable reference
diff --git a/Engine/Source/Editor/KismetCompiler/Public/KismetCompilerMisc.h b/Engine/Source/Editor/KismetCompiler/Public/KismetCompilerMisc.h
--- a/Engine/Source/Editor/KismetCompiler/Public/KismetCompilerMisc.h
+++ b/Engine/Source/Editor/KismetCompiler/Public/KismetCompilerMisc.h
@@ -51,10 +51,10 @@ public:
 	static bool IsTypeCompatibleWithProperty(UEdGraphPin* SourcePin, UProperty* Property, FCompilerResultsLog& MessageLog, const UEdGraphSchema_K2* Schema, UClass* SelfClass);
 
 	/** Finds a property by name, starting in the specified scope; Validates property type and returns NULL along with emitting an error if there is a mismatch. */
-	static UProperty* FindPropertyInScope(UStruct* Scope, UEdGraphPin* Pin, FCompilerResultsLog& MessageLog, const UEdGraphSchema_K2* Schema, UClass* SelfClass);
+	static UProperty* FindPropertyInScope(UStruct* Scope, UEdGraphPin* Pin, FCompilerResultsLog& MessageLog, const UEdGraphSchema_K2* Schema, UClass* SelfClass, bool& bIsSparseProperty);
 
 	// Finds a property by name, starting in the specified scope, returning NULL if it's not found
-	static UProperty* FindNamedPropertyInScope(UStruct* Scope, FName PropertyName);
+	static UProperty* FindNamedPropertyInScope(UStruct* Scope, FName PropertyName, bool& bIsSparseProperty);
 
 	/** return function, that overrides BlueprintImplementableEvent with given name in given class (super-classes are not considered) */
 	static const UFunction* FindOverriddenImplementableEvent(const FName& EventName, const UClass* Class);
diff --git a/Engine/Source/Editor/PropertyEditor/Private/CategoryPropertyNode.cpp b/Engine/Source/Editor/PropertyEditor/Private/CategoryPropertyNode.cpp
--- a/Engine/Source/Editor/PropertyEditor/Private/CategoryPropertyNode.cpp
+++ b/Engine/Source/Editor/PropertyEditor/Private/CategoryPropertyNode.cpp
@@ -5,6 +5,7 @@
 #include "Misc/ConfigCacheIni.h"
 #include "ItemPropertyNode.h"
 #include "ObjectEditorUtils.h"
+#include "ObjectPropertyNode.h"
 #include "PropertyEditorHelpers.h"
 
 FCategoryPropertyNode::FCategoryPropertyNode(void)
@@ -42,38 +43,50 @@ void FCategoryPropertyNode::InitChildNodes()
 	const bool bShouldShowDisableEditOnInstance = !!HasNodeFlags(EPropertyNodeFlags::ShouldShowDisableEditOnInstance);
 
 	TArray<UProperty*> Properties;
+	TSet<UProperty*> SparseProperties;
 	// The parent of a category window has to be an object window.
-	auto ComplexNode = FindComplexParent();
+	FComplexPropertyNode* ComplexNode = FindComplexParent();
 	if (ComplexNode)
 	{
+		FObjectPropertyNode* ObjectNode = FindObjectItemParent();
+		
 		// Get a list of properties that are in the same category
-		for (TFieldIterator<UProperty> It(ComplexNode->GetBaseStructure()); It; ++It)
+		for (const UStruct* Structure : ComplexNode->GetAllStructures())
 		{
-			bool bMetaDataAllowVisible = true;
-			if (!bShowHiddenProperties)
+			const bool bIsSparseStruct = (ObjectNode && ObjectNode->IsSparseDataStruct(Cast<const UScriptStruct>(Structure)));
+			for (TFieldIterator<UProperty> It(Structure); It; ++It)
 			{
-				static const FName Name_bShowOnlyWhenTrue("bShowOnlyWhenTrue");
-				const FString& MetaDataVisibilityCheckString = It->GetMetaData(Name_bShowOnlyWhenTrue);
-				if (MetaDataVisibilityCheckString.Len())
+				bool bMetaDataAllowVisible = true;
+				if (!bShowHiddenProperties)
 				{
-					//ensure that the metadata visibility string is actually set to true in order to show this property
-					// @todo Remove this
-					GConfig->GetBool(TEXT("UnrealEd.PropertyFilters"), *MetaDataVisibilityCheckString, bMetaDataAllowVisible, GEditorPerProjectIni);
+					static const FName Name_bShowOnlyWhenTrue("bShowOnlyWhenTrue");
+					const FString& MetaDataVisibilityCheckString = It->GetMetaData(Name_bShowOnlyWhenTrue);
+					if (MetaDataVisibilityCheckString.Len())
+					{
+						//ensure that the metadata visibility string is actually set to true in order to show this property
+						// @todo Remove this
+						GConfig->GetBool(TEXT("UnrealEd.PropertyFilters"), *MetaDataVisibilityCheckString, bMetaDataAllowVisible, GEditorPerProjectIni);
+					}
 				}
-			}
 
-			if (bMetaDataAllowVisible)
-			{
-				static const FName Name_InlineEditConditionToggle("InlineEditConditionToggle");
-				const bool bOnlyShowAsInlineEditCondition = (*It)->HasMetaData(Name_InlineEditConditionToggle);
-				const bool bShowIfEditableProperty = (*It)->HasAnyPropertyFlags(CPF_Edit);
-				const bool bShowIfDisableEditOnInstance = !(*It)->HasAnyPropertyFlags(CPF_DisableEditOnInstance) || bShouldShowDisableEditOnInstance;
-
-				// Add if we are showing non-editable props and this is the 'None' category, 
-				// or if this is the right category, and we are either showing non-editable
-				if (FObjectEditorUtils::GetCategoryFName(*It) == CategoryName && (bShowHiddenProperties || (bShowIfEditableProperty && !bOnlyShowAsInlineEditCondition && bShowIfDisableEditOnInstance)))
+				if (bMetaDataAllowVisible)
 				{
-					Properties.Add(*It);
+					static const FName Name_InlineEditConditionToggle("InlineEditConditionToggle");
+					const bool bOnlyShowAsInlineEditCondition = (*It)->HasMetaData(Name_InlineEditConditionToggle);
+					const bool bShowIfEditableProperty = (*It)->HasAnyPropertyFlags(CPF_Edit);
+					const bool bShowIfDisableEditOnInstance = !(*It)->HasAnyPropertyFlags(CPF_DisableEditOnInstance) || bShouldShowDisableEditOnInstance;
+
+					// Add if we are showing non-editable props and this is the 'None' category, 
+					// or if this is the right category, and we are either showing non-editable
+					if (FObjectEditorUtils::GetCategoryFName(*It) == CategoryName && (bShowHiddenProperties || (bShowIfEditableProperty && !bOnlyShowAsInlineEditCondition && bShowIfDisableEditOnInstance)))
+					{
+						if (bIsSparseStruct)
+						{
+							SparseProperties.Add(*It);
+						}
+
+						Properties.Add(*It);
+					}
 				}
 			}
 		}
@@ -93,6 +106,7 @@ void FCategoryPropertyNode::InitChildNodes()
 		InitParams.bAllowChildren = true;
 		InitParams.bForceHiddenPropertyVisibility = bShowHiddenProperties;
 		InitParams.bCreateDisableEditOnInstanceNodes = bShouldShowDisableEditOnInstance;
+		InitParams.bIsSparseProperty = SparseProperties.Contains(Properties[PropertyIndex]);
 
 		NewItemNode->InitNode( InitParams );
 
diff --git a/Engine/Source/Editor/PropertyEditor/Private/EditConditionContext.cpp b/Engine/Source/Editor/PropertyEditor/Private/EditConditionContext.cpp
--- a/Engine/Source/Editor/PropertyEditor/Private/EditConditionContext.cpp
+++ b/Engine/Source/Editor/PropertyEditor/Private/EditConditionContext.cpp
@@ -104,8 +104,14 @@ TOptional<bool> FEditConditionContext::GetBoolValue(const FString& PropertyName)
 			return TOptional<bool>();
 		}
 
-		uint8* ParentPtr = ParentNode->GetValueAddress(BasePtr);
-		uint8* ValuePtr = BoolProperty->ContainerPtrToValuePtr<uint8>(ParentPtr);
+		uint8* ParentPtr = ParentNode->GetValueAddress(BasePtr, PinnedNode->HasNodeFlags(EPropertyNodeFlags::IsSparseClassData));
+
+		uint8* ValuePtr = ComplexParentNode->GetValuePtrOfInstance(Index, BoolProperty, ParentNode);
+
+		// SPARSEDATA_TODO: remove the next three lines once we're sure the pointer math is correct
+		uint8* OldValuePtr = BoolProperty->ContainerPtrToValuePtr<uint8>(ParentPtr);
+		check(OldValuePtr == ValuePtr || PinnedNode->HasNodeFlags(EPropertyNodeFlags::IsSparseClassData));
+		ensure(ValuePtr != nullptr);
 
 		bool bValue = BoolProperty->GetPropertyValue(ValuePtr);
 		if (!Result.IsSet())
@@ -145,8 +151,14 @@ TOptional<double> FEditConditionContext::GetNumericValue(const FString& Property
 			return TOptional<double>();
 		}
 
-		uint8* ParentPtr = ParentNode->GetValueAddress(BasePtr);
-		uint8* ValuePtr = NumericProperty->ContainerPtrToValuePtr<uint8>(ParentPtr);
+		uint8* ParentPtr = ParentNode->GetValueAddress(BasePtr, PinnedNode->HasNodeFlags(EPropertyNodeFlags::IsSparseClassData));
+
+		uint8* ValuePtr = ComplexParentNode->GetValuePtrOfInstance(Index, NumericProperty, ParentNode);
+
+		// SPARSEDATA_TODO: remove the next three lines once we're sure the pointer math is correct
+		uint8* OldValuePtr = NumericProperty->ContainerPtrToValuePtr<uint8>(ParentPtr);
+		check(OldValuePtr == ValuePtr || PinnedNode->HasNodeFlags(EPropertyNodeFlags::IsSparseClassData));
+		ensure(ValuePtr != nullptr);
 
 		double Value = 0;
 
@@ -213,8 +225,14 @@ TOptional<FString> FEditConditionContext::GetEnumValue(const FString& PropertyNa
 			return TOptional<FString>();
 		}
 
-		uint8* ParentPtr = ParentNode->GetValueAddress(BasePtr);
-		uint8* ValuePtr = Property->ContainerPtrToValuePtr<uint8>(ParentPtr);
+		uint8* ParentPtr = ParentNode->GetValueAddress(BasePtr, PinnedNode->HasNodeFlags(EPropertyNodeFlags::IsSparseClassData));
+
+		uint8* ValuePtr = ComplexParentNode->GetValuePtrOfInstance(Index, Property, ParentNode);
+
+		// SPARSEDATA_TODO: remove the next three lines once we're sure the pointer math is correct
+		uint8* OldValuePtr = Property->ContainerPtrToValuePtr<uint8>(ParentPtr);
+		check(OldValuePtr == ValuePtr || PinnedNode->HasNodeFlags(EPropertyNodeFlags::IsSparseClassData));
+		ensure(ValuePtr != nullptr);
 
 		int64 Value = NumericProperty->GetSignedIntPropertyValue(ValuePtr);
 		if (!Result.IsSet())
diff --git a/Engine/Source/Editor/PropertyEditor/Private/ItemPropertyNode.cpp b/Engine/Source/Editor/PropertyEditor/Private/ItemPropertyNode.cpp
--- a/Engine/Source/Editor/PropertyEditor/Private/ItemPropertyNode.cpp
+++ b/Engine/Source/Editor/PropertyEditor/Private/ItemPropertyNode.cpp
@@ -20,14 +20,7 @@ FItemPropertyNode::~FItemPropertyNode(void)
 
 }
 
-/**
- * Calculates the memory address for the data associated with this item's property.  This is typically the value of a UProperty or a UObject address.
- *
- * @param	StartAddress	the location to use as the starting point for the calculation; typically the address of the object that contains this property.
- *
- * @return	a pointer to a UProperty value or UObject.  (For dynamic arrays, you'd cast this value to an FArray*)
- */
-uint8* FItemPropertyNode::GetValueBaseAddress( uint8* StartAddress )
+uint8* FItemPropertyNode::GetValueBaseAddress(uint8* StartAddress, bool bIsSparseData)
 {
 	UProperty* MyProperty = GetProperty();
 	if( MyProperty && ParentNodeWeakPtr.IsValid())
@@ -36,7 +29,7 @@ uint8* FItemPropertyNode::GetValueBaseAddress( uint8* StartAddress )
 		USetProperty* OuterSetProp = Cast<USetProperty>(MyProperty->GetOuter());
 		UMapProperty* OuterMapProp = Cast<UMapProperty>(MyProperty->GetOuter());
 
-		uint8* ValueBaseAddress = ParentNode->GetValueBaseAddress(StartAddress);
+		uint8* ValueBaseAddress = ParentNode->GetValueBaseAddress(StartAddress, bIsSparseData);
 
 		if (OuterArrayProp != nullptr)
 		{
@@ -73,7 +66,7 @@ uint8* FItemPropertyNode::GetValueBaseAddress( uint8* StartAddress )
 		}
 		else
 		{
-			uint8* ValueAddress = ParentNode->GetValueAddress(StartAddress);
+			uint8* ValueAddress = ParentNode->GetValueAddress(StartAddress, bIsSparseData);
 			if (ValueAddress != nullptr && ParentNode->GetProperty() != MyProperty)
 			{
 				// if this is not a fixed size array (in which the parent property and this property are the same), we need to offset from the property (otherwise, the parent already did that for us)
@@ -91,17 +84,9 @@ uint8* FItemPropertyNode::GetValueBaseAddress( uint8* StartAddress )
 	return nullptr;
 }
 
-/**
- * Calculates the memory address for the data associated with this item's value.  For most properties, identical to GetValueBaseAddress.  For items corresponding
- * to dynamic array elements, the pointer returned will be the location for that element's data. 
- *
- * @param	StartAddress	the location to use as the starting point for the calculation; typically the address of the object that contains this property.
- *
- * @return	a pointer to a UProperty value or UObject.  (For dynamic arrays, you'd cast this value to whatever type is the Inner for the dynamic array)
- */
-uint8* FItemPropertyNode::GetValueAddress( uint8* StartAddress )
+uint8* FItemPropertyNode::GetValueAddress(uint8* StartAddress, bool bIsSparseData)
 {
-	uint8* Result = GetValueBaseAddress(StartAddress);
+	uint8* Result = GetValueBaseAddress(StartAddress, bIsSparseData);
 
 	UProperty* MyProperty = GetProperty();
 
diff --git a/Engine/Source/Editor/PropertyEditor/Private/ItemPropertyNode.h b/Engine/Source/Editor/PropertyEditor/Private/ItemPropertyNode.h
--- a/Engine/Source/Editor/PropertyEditor/Private/ItemPropertyNode.h
+++ b/Engine/Source/Editor/PropertyEditor/Private/ItemPropertyNode.h
@@ -12,21 +12,8 @@ public:
 	FItemPropertyNode();
 	virtual ~FItemPropertyNode();
 
-	/**
-	 * Calculates the memory address for the data associated with this item's property.  This is typically the value of a UProperty or a UObject address.
-	 *
-	 * @param	StartAddress	the location to use as the starting point for the calculation; typically the address of the object that contains this property.
-	 *
-	 * @return	a pointer to a UProperty value or UObject.  (For dynamic arrays, you'd cast this value to an FArray*)
-	 */
-	virtual uint8* GetValueBaseAddress( uint8* Base ) override;
-	/**
-	 * Calculates the memory address for the data associated with this item's value.  For most properties, identical to GetValueBaseAddress.  For items corresponding
-	 * to dynamic array elements, the pointer returned will be the location for that element's data. 
-	 *
-	 * @return	a pointer to a UProperty value or UObject.  (For dynamic arrays, you'd cast this value to whatever type is the Inner for the dynamic array)
-	 */
-	virtual uint8* GetValueAddress( uint8* Base ) override;
+	virtual uint8* GetValueBaseAddress(uint8* StartAddress, bool bIsSparseData) override;
+	virtual uint8* GetValueAddress(uint8* StartAddress, bool bIsSparseData) override;
 
 	/**
 	 * Overridden function to get the derived object node
diff --git a/Engine/Source/Editor/PropertyEditor/Private/ObjectPropertyNode.cpp b/Engine/Source/Editor/PropertyEditor/Private/ObjectPropertyNode.cpp
--- a/Engine/Source/Editor/PropertyEditor/Private/ObjectPropertyNode.cpp
+++ b/Engine/Source/Editor/PropertyEditor/Private/ObjectPropertyNode.cpp
@@ -116,6 +116,44 @@ void FObjectPropertyNode::Finalize()
 	}
 }
 
+TArray<UStruct*> FObjectPropertyNode::GetAllStructures()
+{
+	TArray<UStruct*> RetVal;
+	RetVal.Reserve(SparseClassDataInstances.Num() + 1);
+	if (UStruct* BaseStruct = GetBaseStructure())
+	{
+		RetVal.Add(BaseStruct);
+	}
+
+	for (const auto& ClassSparseDataPair : SparseClassDataInstances)
+	{
+		const TTuple<UScriptStruct*, void*>& SparseInstance = ClassSparseDataPair.Value;
+		RetVal.AddUnique(SparseInstance.Key);
+	}
+
+	return RetVal;
+}
+
+TArray<const UStruct*> FObjectPropertyNode::GetAllStructures() const
+{
+	TArray<const UStruct*> RetVal;
+	RetVal.Reserve(SparseClassDataInstances.Num() + 1);
+	if (const UStruct* BaseStruct = GetBaseStructure())
+	{
+		RetVal.Add(BaseStruct);
+	}
+
+	for (const auto& ClassSparseDataPair : SparseClassDataInstances)
+	{
+		const TTuple<UScriptStruct*, void*>& SparseInstance = ClassSparseDataPair.Value;
+		RetVal.AddUnique(SparseInstance.Key);
+	}
+	return RetVal;
+}
+
+// Intended for a function that gets the number of elements in a container
+DECLARE_DELEGATE_RetVal_OneParam(int32, FGetNumDelegate, const void*);
+
 bool FObjectPropertyNode::GetReadAddressUncached(FPropertyNode& InNode,
 											   bool InRequiresSingleSelection,
 											   FReadAddressListData* OutAddresses,
@@ -169,7 +207,7 @@ bool FObjectPropertyNode::GetReadAddressUncached(FPropertyNode& InNode,
 		const UObject* TempObject = GetUObject(0);
 		if( TempObject )
 		{
-			uint8* BaseAddr = ParentPropertyNode->GetValueBaseAddress( (uint8*)TempObject );
+			uint8* BaseAddr = ParentPropertyNode->GetValueBaseAddressFromObject(TempObject);
 			if( BaseAddr )
 			{
 				if ( ArrayOuter )
@@ -178,7 +216,7 @@ bool FObjectPropertyNode::GetReadAddressUncached(FPropertyNode& InNode,
 					for (int32 ObjIndex = 1; ObjIndex < GetNumObjects(); ++ObjIndex)
 					{
 						TempObject = GetUObject(ObjIndex);
-						BaseAddr = ParentPropertyNode->GetValueBaseAddress((uint8*)TempObject);
+						BaseAddr = ParentPropertyNode->GetValueBaseAddressFromObject(TempObject);
 
 						if (BaseAddr && Num != FScriptArrayHelper::Num(BaseAddr))
 						{
@@ -192,7 +230,7 @@ bool FObjectPropertyNode::GetReadAddressUncached(FPropertyNode& InNode,
 					for (int32 ObjIndex = 1; ObjIndex < GetNumObjects(); ++ObjIndex)
 					{
 						TempObject = GetUObject(ObjIndex);
-						BaseAddr = ParentPropertyNode->GetValueBaseAddress((uint8*)TempObject);
+						BaseAddr = ParentPropertyNode->GetValueBaseAddressFromObject(TempObject);
 
 						if (BaseAddr && Num != FScriptSetHelper::Num(BaseAddr))
 						{
@@ -206,7 +244,7 @@ bool FObjectPropertyNode::GetReadAddressUncached(FPropertyNode& InNode,
 					for (int32 ObjIndex = 1; ObjIndex < GetNumObjects(); ++ObjIndex)
 					{
 						TempObject = GetUObject(ObjIndex);
-						BaseAddr = ParentPropertyNode->GetValueBaseAddress((uint8*)TempObject);
+						BaseAddr = ParentPropertyNode->GetValueBaseAddressFromObject(TempObject);
 
 						if (BaseAddr && Num != FScriptMapHelper::Num(BaseAddr))
 						{
@@ -218,8 +256,8 @@ bool FObjectPropertyNode::GetReadAddressUncached(FPropertyNode& InNode,
 		}
 	}
 
-	uint8* Base = GetUObject(0) ? InNode.GetValueBaseAddress( (uint8*)(GetUObject(0)) ) : NULL;
-	if (Base)
+	uint8* BaseAddr = InNode.GetValueBaseAddressFromObject(GetUObject(0));
+	if (BaseAddr)
 	{
 		// If the item is an array or set itself, return NULL if there are a different number of
 		// items in the container in different objects, when multi-selecting.
@@ -228,21 +266,21 @@ bool FObjectPropertyNode::GetReadAddressUncached(FPropertyNode& InNode,
 		USetProperty* SetProp = Cast<USetProperty>(InItemProperty);
 		UMapProperty* MapProp = Cast<UMapProperty>(InItemProperty);
 
-		if( ArrayProp || SetProp || MapProp)
+		if (ArrayProp || SetProp || MapProp)
 		{
 			// This flag is an override for array properties which want to display e.g. the "Clear" and "Empty"
 			// buttons, even though the array properties may differ in the number of elements.
-			if ( !bArrayPropertiesCanDifferInSize )
+			if (!bArrayPropertiesCanDifferInSize)
 			{
 				const UObject* TempObject = GetUObject(0);
 
 				if (ArrayProp)
 				{
-					int32 const Num = FScriptArrayHelper::Num(InNode.GetValueBaseAddress((uint8*)TempObject));
+					int32 const Num = FScriptArrayHelper::Num(BaseAddr);
 					for (int32 ObjIndex = 1; ObjIndex < GetNumObjects(); ObjIndex++)
 					{
 						TempObject = GetUObject(ObjIndex);
-						if (TempObject && Num != FScriptArrayHelper::Num(InNode.GetValueBaseAddress((uint8*)TempObject)))
+						if (TempObject && Num != FScriptArrayHelper::Num(BaseAddr))
 						{
 							bAllTheSame = false;
 						}
@@ -250,11 +288,11 @@ bool FObjectPropertyNode::GetReadAddressUncached(FPropertyNode& InNode,
 				}
 				else if (SetProp)
 				{
-					int32 const Num = FScriptSetHelper::Num(InNode.GetValueBaseAddress((uint8*)TempObject));
+					int32 const Num = FScriptSetHelper::Num(BaseAddr);
 					for (int32 ObjIndex = 1; ObjIndex < GetNumObjects(); ++ObjIndex)
 					{
 						TempObject = GetUObject(ObjIndex);
-						if (TempObject && Num != FScriptSetHelper::Num(InNode.GetValueBaseAddress((uint8*)TempObject)))
+						if (TempObject && Num != FScriptSetHelper::Num(BaseAddr))
 						{
 							bAllTheSame = false;
 						}
@@ -262,11 +300,11 @@ bool FObjectPropertyNode::GetReadAddressUncached(FPropertyNode& InNode,
 				}
 				else if (MapProp)
 				{
-					int32 const Num = FScriptMapHelper::Num(InNode.GetValueBaseAddress((uint8*)TempObject));
+					int32 const Num = FScriptMapHelper::Num(BaseAddr);
 					for (int32 ObjIndex = 1; ObjIndex < GetNumObjects(); ++ObjIndex)
 					{
 						TempObject = GetUObject(ObjIndex);
-						if (TempObject && Num != FScriptMapHelper::Num(InNode.GetValueBaseAddress((uint8*)TempObject)))
+						if (TempObject && Num != FScriptMapHelper::Num(BaseAddr))
 						{
 							bAllTheSame = false;
 						}
@@ -281,33 +319,32 @@ bool FObjectPropertyNode::GetReadAddressUncached(FPropertyNode& InNode,
 				// Make sure the value of this InItemProperty is the same in all selected objects.
 				for( int32 ObjIndex = 1 ; ObjIndex < GetNumObjects() ; ObjIndex++ )
 				{
-					const UObject* TempObject = GetUObject(ObjIndex);
-					if( !InItemProperty->Identical( Base, InNode.GetValueBaseAddress( (uint8*)TempObject ) ) )
+					if (!InItemProperty->Identical(BaseAddr, InNode.GetValueBaseAddressFromObject(GetUObject(ObjIndex))))
 					{
 						bAllTheSame = false;
+						break;
 					}
 				}
 			}
 			else
 			{
-				if ( Cast<UObjectPropertyBase>(InItemProperty) )
+				if (Cast<UObjectPropertyBase>(InItemProperty))
 				{
 					// We don't want to exactly compare component properties.  However, we
 					// need to be sure that all references are either valid or invalid.
 
 					// If BaseObj is NON-NULL, all other objects' properties should also be non-NULL.
 					// If BaseObj is NULL, all other objects' properties should also be NULL.
-					UObject* BaseObj = Cast<UObjectPropertyBase>(InItemProperty)->GetObjectPropertyValue(Base);
+					UObject* BaseObj = Cast<UObjectPropertyBase>(InItemProperty)->GetObjectPropertyValue(BaseAddr);
 
-					for( int32 ObjIndex = 1 ; ObjIndex < GetNumObjects() ; ObjIndex++ )
+					for (int32 ObjIndex = 1; ObjIndex < GetNumObjects(); ObjIndex++)
 					{
-						const UObject* TempObject = GetUObject(ObjIndex);
-						UObject* CurObj = Cast<UObjectPropertyBase>(InItemProperty)->GetObjectPropertyValue(InNode.GetValueBaseAddress( (uint8*)TempObject ));
+						UObject* CurObj = Cast<UObjectPropertyBase>(InItemProperty)->GetObjectPropertyValue(InNode.GetValueBaseAddressFromObject(GetUObject(ObjIndex)));
 						if (   ( !BaseObj && CurObj )			// BaseObj is NULL, but this InItemProperty is non-NULL!
 							|| ( BaseObj && !CurObj ) )			// BaseObj is non-NULL, but this InItemProperty is NULL!
 						{
-
 							bAllTheSame = false;
+							break;
 						}
 					}
 				}
@@ -315,15 +352,15 @@ bool FObjectPropertyNode::GetReadAddressUncached(FPropertyNode& InNode,
 		}
 	}
 
-	if(OutAddresses != nullptr)
+	// Write addresses to the output.
+	if (OutAddresses != nullptr)
 	{
-		// Write addresses to the output.
-		for(int32 ObjIndex = 0; ObjIndex < GetNumObjects(); ++ObjIndex)
-		{
+		for (int32 ObjIndex = 0; ObjIndex < GetNumObjects(); ++ObjIndex)
+		{	
 			const UObject* TempObject = GetUObject(ObjIndex);
-			if(TempObject)
+			if (TempObject)
 			{
-				OutAddresses->Add(TempObject, InNode.GetValueBaseAddress((uint8*)(TempObject)));
+				OutAddresses->Add(TempObject, InNode.GetValueBaseAddressFromObject(TempObject));
 			}
 		}
 	}
@@ -354,13 +391,12 @@ bool FObjectPropertyNode::GetReadAddressUncached( FPropertyNode& InNode, FReadAd
 
 
 	// Write addresses to the output.
-	for ( int32 ObjectIdx = 0 ; ObjectIdx < GetNumObjects() ; ++ObjectIdx )
+	for (int32 ObjIndex = 0 ; ObjIndex < GetNumObjects() ; ++ObjIndex)
 	{
-		const UObject* TempObject = GetUObject(ObjectIdx);
-
-		if( TempObject != NULL )
+		const UObject* TempObject = GetUObject(ObjIndex);
+		if (TempObject != nullptr)
 		{
-			OutAddresses.Add( TempObject, InNode.GetValueBaseAddress( (uint8*)(TempObject) ) );
+			OutAddresses.Add(TempObject, InNode.GetValueBaseAddressFromObject(TempObject));
 		}
 	}
 
@@ -368,22 +404,17 @@ bool FObjectPropertyNode::GetReadAddressUncached( FPropertyNode& InNode, FReadAd
 	return true;
 }
 
-
-/**
- * Calculates the memory address for the data associated with this item's property.  This is typically the value of a UProperty or a UObject address.
- *
- * @param	StartAddress	the location to use as the starting point for the calculation; typically the address of the object that contains this property.
- *
- * @return	a pointer to a UProperty value or UObject.  (For dynamic arrays, you'd cast this value to an FArray*)
- */
-uint8* FObjectPropertyNode::GetValueBaseAddress( uint8* StartAddress )
+uint8* FObjectPropertyNode::GetValueBaseAddress(uint8* StartAddress, bool bIsSparseData)
 {
 	uint8* Result = StartAddress;
 
 	UClass* ClassObject;
-	if ( (ClassObject=Cast<UClass>((UObject*)Result)) != NULL )
+	if (!bIsSparseData)
 	{
-		Result = (uint8*)ClassObject->GetDefaultObject();
+		if ( (ClassObject=Cast<UClass>((UObject*)Result)) != NULL )
+		{
+			Result = (uint8*)ClassObject->GetDefaultObject();
+		}
 	}
 
 	return Result;
@@ -403,9 +434,60 @@ void FObjectPropertyNode::InitChildNodes()
 	InternalInitChildNodes();
 }
 
+void FObjectPropertyNode::GetCategoryProperties(const TSet<UClass*>& ClassesToConsider, const UProperty* CurrentProperty, bool bShouldShowDisableEditOnInstance, bool bShouldShowHiddenProperties,
+	const TSet<FName>& CategoriesFromBlueprints, TSet<FName>& CategoriesFromProperties, TArray<FName>& SortedCategories)
+{
+	static const FName Name_InlineEditConditionToggle("InlineEditConditionToggle");
+	static const FName Name_bShowOnlyWhenTrue("bShowOnlyWhenTrue");
+
+	bool bHidden = false;
+
+	FName CategoryName = FObjectEditorUtils::GetCategoryFName(CurrentProperty);
+
+	for (UClass* Class : ClassesToConsider)
+	{
+		if (FEditorCategoryUtils::IsCategoryHiddenFromClass(Class, CategoryName.ToString()))
+		{
+			HiddenCategories.Add(CategoryName);
+
+			bHidden = true;
+			break;
+		}
+	}
+
+	// 
+	if (CurrentProperty->HasMetaData(Name_InlineEditConditionToggle))
+	{
+		bHidden = true;
+	}
+
+	bool bMetaDataAllowVisible = true;
+	const FString& MetaDataVisibilityCheckString = CurrentProperty->GetMetaData(Name_bShowOnlyWhenTrue);
+	if (MetaDataVisibilityCheckString.Len())
+	{
+		//ensure that the metadata visibility string is actually set to true in order to show this property
+		GConfig->GetBool(TEXT("UnrealEd.PropertyFilters"), *MetaDataVisibilityCheckString, bMetaDataAllowVisible, GEditorPerProjectIni);
+	}
+
+	if (bMetaDataAllowVisible)
+	{
+		const bool bShowIfNonHiddenEditableProperty = CurrentProperty->HasAnyPropertyFlags(CPF_Edit) && !bHidden;
+		const bool bShowIfDisableEditOnInstance = !CurrentProperty->HasAnyPropertyFlags(CPF_DisableEditOnInstance) || bShouldShowDisableEditOnInstance;
+		if (bShouldShowHiddenProperties || (bShowIfNonHiddenEditableProperty && bShowIfDisableEditOnInstance))
+		{
+			if (!CategoriesFromBlueprints.Contains(CategoryName) && !CategoriesFromProperties.Contains(CategoryName))
+			{
+				SortedCategories.Add(CategoryName);
+			}
+			CategoriesFromProperties.Add(CategoryName);
+		}
+	}
+}
+
 void FObjectPropertyNode::InternalInitChildNodes( FName SinglePropertyName )
 {
-	HiddenCategories.Empty();
+	HiddenCategories.Reset();
+	SparseClassDataInstances.Reset();
 	// Assemble a list of category names by iterating over all fields of BaseClass.
 
 	// build a list of classes that we need to look at
@@ -443,50 +525,24 @@ void FObjectPropertyNode::InternalInitChildNodes( FName SinglePropertyName )
 	const bool bShouldShowHiddenProperties = !!HasNodeFlags(EPropertyNodeFlags::ShouldShowHiddenProperties);
 	const bool bShouldShowDisableEditOnInstance = !!HasNodeFlags(EPropertyNodeFlags::ShouldShowDisableEditOnInstance);
 
-	for( TFieldIterator<UProperty> It(BaseClass.Get()); It; ++It )
+	for (TFieldIterator<UProperty> It(BaseClass.Get()); It; ++It)
 	{
-		bool bHidden = false;
-
-		FName CategoryName = FObjectEditorUtils::GetCategoryFName(*It);
+		GetCategoryProperties(ClassesToConsider, *It, bShouldShowDisableEditOnInstance, bShouldShowHiddenProperties, CategoriesFromBlueprints, CategoriesFromProperties, SortedCategories);
+	}
 
-		for( UClass* Class : ClassesToConsider )
+	for (UClass* Class : ClassesToConsider)
+	{
+		if (Class)
 		{
-			if( FEditorCategoryUtils::IsCategoryHiddenFromClass(Class, CategoryName.ToString()) )
+			UScriptStruct* SparseClassDataStruct = Class->GetSparseClassDataStruct();
+			if (SparseClassDataStruct)
 			{
-				HiddenCategories.Add( CategoryName );
+				SparseClassDataInstances.Add(Class, TTuple<UScriptStruct*, void*>(SparseClassDataStruct, Class->GetOrCreateSparseClassData()));
 
-				bHidden = true;
-				break;
-			}
-		}
-
-		// 
-		static const FName Name_InlineEditConditionToggle("InlineEditConditionToggle");
-		if (It->HasMetaData(Name_InlineEditConditionToggle))
-		{
-			bHidden = true;
-		}
-
-		bool bMetaDataAllowVisible = true;
-		static const FName Name_bShowOnlyWhenTrue("bShowOnlyWhenTrue");
-		const FString& MetaDataVisibilityCheckString = It->GetMetaData(Name_bShowOnlyWhenTrue);
-		if (MetaDataVisibilityCheckString.Len())
-		{
-			//ensure that the metadata visibility string is actually set to true in order to show this property
-			GConfig->GetBool(TEXT("UnrealEd.PropertyFilters"), *MetaDataVisibilityCheckString, bMetaDataAllowVisible, GEditorPerProjectIni);
-		}
-
-		if (bMetaDataAllowVisible)
-		{
-			const bool bShowIfNonHiddenEditableProperty = (*It)->HasAnyPropertyFlags(CPF_Edit) && !bHidden;
-			const bool bShowIfDisableEditOnInstance = !(*It)->HasAnyPropertyFlags(CPF_DisableEditOnInstance) || bShouldShowDisableEditOnInstance;
-			if( bShouldShowHiddenProperties || (bShowIfNonHiddenEditableProperty && bShowIfDisableEditOnInstance) )
-			{
-				if (!CategoriesFromBlueprints.Contains(CategoryName) && !CategoriesFromProperties.Contains(CategoryName) )
-				{
-					SortedCategories.Add(CategoryName);
-				}
-				CategoriesFromProperties.Add(CategoryName);
+				for (TFieldIterator<UProperty> It(SparseClassDataStruct); It; ++It)
+		 		{
+		 			GetCategoryProperties(ClassesToConsider, *It, bShouldShowDisableEditOnInstance, bShouldShowHiddenProperties, CategoriesFromBlueprints, CategoriesFromProperties, SortedCategories);
+		 		}
 			}
 		}
 	}
@@ -655,6 +711,20 @@ TSharedPtr<FPropertyNode> FObjectPropertyNode::GenerateSingleChild( FName ChildP
 	return NULL;
 }
 
+bool FObjectPropertyNode::IsSparseDataStruct(const UScriptStruct* Struct) const
+{
+	for (const auto& Instance : SparseClassDataInstances)
+	{
+		const TTuple<UScriptStruct*, void*>& SparseData = Instance.Value;
+		if (SparseData.Key == Struct)
+		{
+			return true;
+		}
+	}
+
+	return false;
+}
+
 /**
  * Appends my path, including an array index (where appropriate)
  */
diff --git a/Engine/Source/Editor/PropertyEditor/Private/ObjectPropertyNode.h b/Engine/Source/Editor/PropertyEditor/Private/ObjectPropertyNode.h
--- a/Engine/Source/Editor/PropertyEditor/Private/ObjectPropertyNode.h
+++ b/Engine/Source/Editor/PropertyEditor/Private/ObjectPropertyNode.h
@@ -26,8 +26,6 @@ public:
 	virtual bool GetReadAddressUncached(FPropertyNode& InNode, bool InRequiresSingleSelection, FReadAddressListData* OutAddresses, bool bComparePropertyContents = true, bool bObjectForceCompare = false, bool bArrayPropertiesCanDifferInSize = false) const override;
 	virtual bool GetReadAddressUncached(FPropertyNode& InNode, FReadAddressListData& OutAddresses) const override;
 
-	virtual uint8* GetValueBaseAddress( uint8* Base ) override;
-
 	/**
 	 * Returns the UObject at index "n" of the Objects Array
 	 * @param InIndex - index to read out of the array
@@ -72,19 +70,28 @@ public:
 	void Finalize();
 
 	/** @return		The base-est baseclass for objects in this list. */
-	UClass*			GetObjectBaseClass()       { return BaseClass.IsValid() ? BaseClass.Get() : NULL; }
+	UClass*			GetObjectBaseClass()       { return BaseClass.IsValid() ? BaseClass.Get() : nullptr; }
 	/** @return		The base-est baseclass for objects in this list. */
-	const UClass*	GetObjectBaseClass() const { return BaseClass.IsValid() ? BaseClass.Get() : NULL; }
+	const UClass*	GetObjectBaseClass() const { return BaseClass.IsValid() ? BaseClass.Get() : nullptr; }
 
 
 	// FComplexPropertyNode implementation
 	virtual UStruct* GetBaseStructure() override { return GetObjectBaseClass(); }
 	virtual const UStruct* GetBaseStructure() const override{ return GetObjectBaseClass(); }
+
+	virtual TArray<UStruct*> GetAllStructures() override;
+	virtual TArray<const UStruct*> GetAllStructures() const override;
+
 	virtual int32 GetInstancesNum() const override{ return GetNumObjects(); }
 	virtual uint8* GetMemoryOfInstance(int32 Index) override
 	{
 		return (uint8*)GetUObject(Index);
 	}
+	virtual uint8* GetValuePtrOfInstance(int32 Index, const UProperty* InProperty, FPropertyNode* InParentNode) override
+	{
+		uint8* ParentOffset = InParentNode ? InParentNode->GetValueAddressFromObject(GetUObject(Index)) : nullptr;
+		return InProperty ? InProperty->ContainerPtrToValuePtr<uint8>(ParentOffset) : nullptr;
+	}
 	virtual TWeakObjectPtr<UObject> GetInstanceAsUObject(int32 Index) override
 	{
 		check(Objects.IsValidIndex(Index));
@@ -98,7 +105,7 @@ public:
 
 	//////////////////////////////////////////////////////////////////////////
 	/** @return		The property stored at this node, to be passed to Pre/PostEditChange. */
-	virtual UProperty*		GetStoredProperty()		{ return StoredProperty.IsValid() ? StoredProperty.Get() : NULL; }
+	virtual UProperty*		GetStoredProperty()		{ return StoredProperty.IsValid() ? StoredProperty.Get() : nullptr; }
 
 	TPropObjectIterator			ObjectIterator()			{ return TPropObjectIterator( Objects ); }
 	TPropObjectConstIterator	ObjectConstIterator() const	{ return TPropObjectConstIterator( Objects ); }
@@ -113,12 +120,17 @@ public:
 	const TSet<FName>& GetHiddenCategories() const { return HiddenCategories; }
 
 	bool IsRootNode() const { return ParentNode == nullptr; }
+
+	/**
+	 * @return True if Struct is one of the sparse data structures used by this object
+	 */
+	bool IsSparseDataStruct(const UScriptStruct* Struct) const;
 protected:
 	/** FPropertyNode interface */
 	virtual void InitBeforeNodeFlags() override;
 	virtual void InitChildNodes() override;
 	virtual bool GetQualifiedName( FString& PathPlusIndex, const bool bWithArrayIndex, const FPropertyNode* StopParent = nullptr, bool bIgnoreCategories = false ) const override;
-
+	virtual uint8* GetValueBaseAddress(uint8* Base, bool bIsSparseData) override;
 	/**
 	 * Looks at the Objects array and creates the best base class.  Called by
 	 * Finalize(); that is, when the list of selected objects is being finalized.
@@ -132,6 +144,10 @@ private:
 	 * @param SingleChildName	The property name of a single child to create instead of all childen
 	 */
 	void InternalInitChildNodes( FName SingleChildName = NAME_None );
+
+	/** If CurrentProperty should show up in the ClassesToConsider make sure its category is in SortedCategories and CategoriesFromProperties. */
+	void GetCategoryProperties(const TSet<UClass*>& ClassesToConsider, const UProperty* CurrentProperty, bool bShouldShowDisableEditOnInstance, bool bShouldShowHiddenProperties,
+	const TSet<FName>& CategoriesFromBlueprints, TSet<FName>& CategoriesFromProperties, TArray<FName>& SortedCategories);
 private:
 	/** The list of objects we are editing properties for. */
 	TArray< TWeakObjectPtr<UObject> >		Objects;
@@ -149,6 +165,13 @@ private:
 	 */
 	TSet<FName> HiddenCategories;
 
+
+
+	/* 
+	 * Contains the structure and memory location for storing data of that type. Used to read and write to sidecar structs for the class
+	 */
+	TMap<UClass*, TTuple<UScriptStruct*, void*>> SparseClassDataInstances;
+
 	/** Object -> Package re-mapping */
 	TMap<TWeakObjectPtr<UObject>, TWeakObjectPtr<UPackage>> ObjectToPackageMapping;
 };
diff --git a/Engine/Source/Editor/PropertyEditor/Private/Presentation/PropertyEditor/PropertyEditor.cpp b/Engine/Source/Editor/PropertyEditor/Private/Presentation/PropertyEditor/PropertyEditor.cpp
--- a/Engine/Source/Editor/PropertyEditor/Private/Presentation/PropertyEditor/PropertyEditor.cpp
+++ b/Engine/Source/Editor/PropertyEditor/Private/Presentation/PropertyEditor/PropertyEditor.cpp
@@ -468,8 +468,13 @@ void FPropertyEditor::ToggleEditConditionState()
 		// ParentNode can point to a struct inside that object (which is stored as an FItemPropertyNode)
 		// We need all three pointers to get the value pointer
 		uint8* BaseAddress = ComplexParentNode->GetMemoryOfInstance(Index);
-		uint8* ParentOffset = ParentNode->GetValueAddress(BaseAddress);
-		uint8* ValuePtr = EditConditionProperty->ContainerPtrToValuePtr<uint8>(ParentOffset);
+		uint8* ParentOffset = ParentNode->GetValueAddress(BaseAddress, PropertyNode->HasNodeFlags(EPropertyNodeFlags::IsSparseClassData));
+
+		uint8* ValuePtr = ComplexParentNode->GetValuePtrOfInstance(Index, EditConditionProperty, ParentNode);
+
+		// SPARSEDATA_TODO: these two lines should go away once we're really confident the pointer math is all correct
+		uint8* OldValuePtr = EditConditionProperty->ContainerPtrToValuePtr<uint8>(ParentOffset);
+		check(OldValuePtr == ValuePtr || PropertyNode->HasNodeFlags(EPropertyNodeFlags::IsSparseClassData));
 
 		OldValue &= EditConditionProperty->GetPropertyValue(ValuePtr);
 		EditConditionProperty->SetPropertyValue(ValuePtr, !OldValue);
@@ -489,8 +494,8 @@ void FPropertyEditor::ToggleEditConditionState()
 				Object->GetArchetypeInstances(ArchetypeInstances);
 				for (int32 InstanceIndex = 0; InstanceIndex < ArchetypeInstances.Num(); ++InstanceIndex)
 				{
-					uint8* ArchetypeBaseOffset = ComplexParentNode->GetValueAddress((uint8*) ArchetypeInstances[InstanceIndex]);
-					uint8* ArchetypeParentOffset = ParentNode->GetValueAddress(ArchetypeBaseOffset);
+					uint8* ArchetypeBaseOffset = ComplexParentNode->GetValueAddressFromObject(ArchetypeInstances[InstanceIndex]);
+					uint8* ArchetypeParentOffset = ParentNode->GetValueAddress(ArchetypeBaseOffset, PropertyNode->HasNodeFlags(EPropertyNodeFlags::IsSparseClassData));
 					uint8* ArchetypeValueAddr = EditConditionProperty->ContainerPtrToValuePtr<uint8>(ArchetypeParentOffset);
 
 					// Only propagate if the current value on the instance matches the previous value on the template.
diff --git a/Engine/Source/Editor/PropertyEditor/Private/PropertyChangeListener.cpp b/Engine/Source/Editor/PropertyEditor/Private/PropertyChangeListener.cpp
--- a/Engine/Source/Editor/PropertyEditor/Private/PropertyChangeListener.cpp
+++ b/Engine/Source/Editor/PropertyEditor/Private/PropertyChangeListener.cpp
@@ -39,25 +39,16 @@ void FPropertyChangeListener::SetObject( UObject& Object, const FPropertyListene
 class FValueCache
 {
 public:
-	/**
-	 * A union which allows a single address to be represented as a pointer to a uint8
-	 * or a pointer to a UObject.
-	 */
-	union FPropertyValueRoot
-	{
-		UObject*	OwnerObject;
-		uint8*		ValueAddress;
-	};
 
 	FValueCache( TSharedRef<FPropertyNode> InPropertyNode, UObject* InOwnerObject )
 		: PropertyNode( InPropertyNode )
 	{
-		PropertyValueRoot.OwnerObject = InOwnerObject;
+		OwnerObject = InOwnerObject;
 
 		UProperty* Property = InPropertyNode->GetProperty();
 
 		check(Property);
-		check(PropertyValueRoot.OwnerObject);
+		check(OwnerObject);
 
 		FPropertyNode* ParentNode = InPropertyNode->GetParentNode();
 	}
@@ -196,10 +187,10 @@ private:
 		FPropertyValueAddresses ValueAddresses;
 
 		ValueAddresses.BaseAddress = (OuterArrayProp == NULL && OuterSetProp == NULL && OuterMapProp == NULL)
-			? PropertyNodeRef.GetValueBaseAddress( PropertyValueRoot.ValueAddress )
-			: ParentNode->GetValueBaseAddress( PropertyValueRoot.ValueAddress );
+			? PropertyNodeRef.GetValueBaseAddressFromObject(OwnerObject)
+			: ParentNode->GetValueBaseAddressFromObject(OwnerObject);
 
-		ValueAddresses.Address = PropertyNodeRef.GetValueAddress( PropertyValueRoot.ValueAddress );
+		ValueAddresses.Address = PropertyNodeRef.GetValueAddressFromObject(OwnerObject);
 
 		return ValueAddresses;
 	}
@@ -215,7 +206,7 @@ private:
 	TWeakObjectPtr<UObject> CachedTopLevelObject;
 
 	/** The address of the owning object */
-	FPropertyValueRoot PropertyValueRoot;
+	UObject*	OwnerObject;
 };
 
 
diff --git a/Engine/Source/Editor/PropertyEditor/Private/PropertyHandleImpl.cpp b/Engine/Source/Editor/PropertyEditor/Private/PropertyHandleImpl.cpp
--- a/Engine/Source/Editor/PropertyEditor/Private/PropertyHandleImpl.cpp
+++ b/Engine/Source/Editor/PropertyEditor/Private/PropertyHandleImpl.cpp
@@ -40,6 +40,21 @@ void PropertyToTextHelper(FString& OutString, FPropertyNode* InPropertyNode, UPr
 	}
 }
 
+void PropertyToTextHelper(FString& OutString, FPropertyNode* InPropertyNode, UProperty* Property, const FObjectBaseAddress& ObjectAddress, EPropertyPortFlags PortFlags)
+{
+	if (!InPropertyNode->HasNodeFlags(EPropertyNodeFlags::IsSparseClassData))
+	{
+		PropertyToTextHelper(OutString, InPropertyNode, Property, ObjectAddress.BaseAddress, PortFlags);
+	}
+	else
+	{
+// FRED_TODO: change this to call PropertyToTextHelper with the new base address and verify it is doing the right thing
+		void* BaseAddress = ObjectAddress.Object->GetClass()->GetOrCreateSparseClassData();
+		void* ValueAddress = Property->ContainerPtrToValuePtr<void>(BaseAddress);
+		Property->ExportText_Direct(OutString, ValueAddress, ValueAddress, nullptr, PortFlags);
+	}
+}
+
 void TextToPropertyHelper(const TCHAR* Buffer, FPropertyNode* InPropertyNode, UProperty* Property, uint8* ValueAddress, UObject* Object)
 {
 	if (InPropertyNode->GetArrayIndex() != INDEX_NONE || Property->ArrayDim == 1)
@@ -52,6 +67,13 @@ void TextToPropertyHelper(const TCHAR* Buffer, FPropertyNode* InPropertyNode, UP
 	}
 }
 
+void TextToPropertyHelper(const TCHAR* Buffer, FPropertyNode* InPropertyNode, UProperty* Property, const FObjectBaseAddress& ObjectAddress)
+{
+
+	uint8* BaseAddress = InPropertyNode ? InPropertyNode->GetValueBaseAddressFromObject(ObjectAddress.Object) : ObjectAddress.BaseAddress;
+	TextToPropertyHelper(Buffer, InPropertyNode, Property, BaseAddress, ObjectAddress.Object);
+}
+
 FPropertyValueImpl::FPropertyValueImpl( TSharedPtr<FPropertyNode> InPropertyNode, FNotifyHook* InNotifyHook, TSharedPtr<IPropertyUtilities> InPropertyUtilities )
 	: PropertyNode( InPropertyNode )
 	, PropertyUtilities( InPropertyUtilities )
@@ -72,7 +94,7 @@ void FPropertyValueImpl::EnumerateObjectsToModify( FPropertyNode* InPropertyNode
 		for (int32 Index = 0; Index < NumInstances; ++Index)
 		{
 			UObject*	Object = ComplexNode->GetInstanceAsUObject(Index).Get();
-			uint8*		Addr = InPropertyNode->GetValueBaseAddress(ComplexNode->GetMemoryOfInstance(Index));
+			uint8* Addr = Object ? InPropertyNode->GetValueBaseAddressFromObject(Object) : InPropertyNode->GetValueBaseAddress(ComplexNode->GetMemoryOfInstance(Index), false);
 			if (!InObjectsToModifyCallback(FObjectBaseAddress(Object, Addr, bIsStruct), Index, NumInstances))
 			{
 				break;
@@ -429,7 +451,7 @@ FPropertyAccess::Result FPropertyValueImpl::ImportText( const TArray<FObjectBase
 
 			// Cache the value of the property before modifying it.
 			FString PreviousValue;
-			PropertyToTextHelper(PreviousValue, InPropertyNode, NodeProperty, Cur.BaseAddress, PPF_None);
+			PropertyToTextHelper(PreviousValue, InPropertyNode, NodeProperty, Cur, PPF_None);
 
 			// If this property is the inner-property of a container, cache the current value as well
 			FString PreviousContainerValue;
@@ -451,7 +473,7 @@ FPropertyAccess::Result FPropertyValueImpl::ImportText( const TArray<FObjectBase
 
 					if (bIsInContainer)
 					{
-						FScriptSetHelper SetHelper(SetProp, ParentNode->GetValueBaseAddress((uint8*)Cur.Object));
+						FScriptSetHelper SetHelper(SetProp, ParentNode->GetValueBaseAddressFromObject(Cur.Object));
 
 						if ( SetHelper.HasElement(Cur.BaseAddress, InValues[ObjectIndex]) )
 						{
@@ -468,7 +490,7 @@ FPropertyAccess::Result FPropertyValueImpl::ImportText( const TArray<FObjectBase
 
 					if (bIsInContainer)
 					{
-						FScriptMapHelper MapHelper(MapProperty, ParentNode->GetValueBaseAddress((uint8*)Cur.Object));
+						FScriptMapHelper MapHelper(MapProperty, ParentNode->GetValueBaseAddressFromObject(Cur.Object));
 
 						if ( MapHelper.HasKey(Cur.BaseAddress, InValues[ObjectIndex]) )
 						{
@@ -486,7 +508,7 @@ FPropertyAccess::Result FPropertyValueImpl::ImportText( const TArray<FObjectBase
 
 				if (bIsInContainer)
 				{
-					uint8* Addr = ParentNode->GetValueBaseAddress((uint8*)Cur.Object);
+					uint8* Addr = ParentNode->GetValueBaseAddressFromObject(Cur.Object);
 					Property->ExportText_Direct(PreviousContainerValue, Addr, Addr, nullptr, 0);
 				}
 			}
@@ -516,19 +538,23 @@ FPropertyAccess::Result FPropertyValueImpl::ImportText( const TArray<FObjectBase
 
 			// Set the new value.
 			const TCHAR* NewValue = *InValues[ObjectIndex];
-			TextToPropertyHelper(NewValue, InPropertyNode, NodeProperty, Cur.BaseAddress, Cur.Object);
+			TextToPropertyHelper(NewValue, InPropertyNode, NodeProperty, Cur);
 
 			if (Cur.Object)
 			{
 				// Cache the value of the property after having modified it.
 				FString ValueAfterImport;
-				PropertyToTextHelper(ValueAfterImport, InPropertyNode, NodeProperty, Cur.BaseAddress, PPF_None);
+				PropertyToTextHelper(ValueAfterImport, InPropertyNode, NodeProperty, Cur, PPF_None);
 
 				if ((Cur.Object->HasAnyFlags(RF_ClassDefaultObject | RF_ArchetypeObject) ||
 					(Cur.Object->HasAnyFlags(RF_DefaultSubObject) && Cur.Object->GetOuter()->HasAnyFlags(RF_ClassDefaultObject | RF_ArchetypeObject))) &&
 					!bIsGameWorld)
 				{
-					InPropertyNode->PropagatePropertyChange(Cur.Object, NewValue, PreviousContainerValue.IsEmpty() ? PreviousValue : PreviousContainerValue);
+					// propagate the changes to instances unless we're modifying class shared data
+					if (!(InPropertyNode->HasNodeFlags(EPropertyNodeFlags::IsSparseClassData)))
+					{
+						InPropertyNode->PropagatePropertyChange(Cur.Object, NewValue, PreviousContainerValue.IsEmpty() ? PreviousValue : PreviousContainerValue);
+					}
 				}
 
 				// If the values before and after setting the property differ, mark the object dirty.
@@ -539,13 +565,13 @@ FPropertyAccess::Result FPropertyValueImpl::ImportText( const TArray<FObjectBase
 					// For TMap and TSet, we need to rehash it in case a key was modified
 					if (NodeProperty->GetOuter()->IsA<UMapProperty>())
 					{
-						uint8* Addr = InPropertyNode->GetParentNode()->GetValueBaseAddress((uint8*)Cur.Object);
+						uint8* Addr = InPropertyNode->GetParentNode()->GetValueBaseAddressFromObject(Cur.Object);
 						FScriptMapHelper MapHelper(Cast<UMapProperty>(NodeProperty->GetOuter()), Addr);
 						MapHelper.Rehash();
 					}
 					else if (NodeProperty->GetOuter()->IsA<USetProperty>())
 					{
-						uint8* Addr = InPropertyNode->GetParentNode()->GetValueBaseAddress((uint8*)Cur.Object);
+						uint8* Addr = InPropertyNode->GetParentNode()->GetValueBaseAddressFromObject(Cur.Object);
 						FScriptSetHelper SetHelper(Cast<USetProperty>(NodeProperty->GetOuter()), Addr);
 						SetHelper.Rehash();
 					}
@@ -2651,12 +2677,12 @@ void FPropertyHandleBase::SetToolTipText( const FText& ToolTip )
 	}
 }
 
-uint8* FPropertyHandleBase::GetValueBaseAddress( uint8* Base )
+uint8* FPropertyHandleBase::GetValueBaseAddress(uint8* Base)
 {
 	TSharedPtr<FPropertyNode> PropertyNode = Implementation->GetPropertyNode();
 	if (PropertyNode.IsValid())
 	{
-		return PropertyNode->GetValueBaseAddress(Base);
+		return PropertyNode->GetValueBaseAddress(Base, PropertyNode->HasNodeFlags(EPropertyNodeFlags::IsSparseClassData));
 	}
 
 	return nullptr;
@@ -4505,7 +4531,7 @@ bool FPropertyHandleSet::HasDefaultElement()
 		if (Addresses.Num() > 0)
 		{
 			USetProperty* SetProperty = CastChecked<USetProperty>(PropNode->GetProperty());
-			FScriptSetHelper SetHelper(SetProperty, PropNode->GetValueBaseAddress((uint8*)Addresses[0].Object));
+			FScriptSetHelper SetHelper(SetProperty, PropNode->GetValueBaseAddressFromObject(Addresses[0].Object));
 
 			FDefaultConstructedPropertyElement DefaultElement(SetHelper.ElementProp);
 
@@ -4604,7 +4630,8 @@ bool FPropertyHandleMap::HasDefaultKey()
 		if (Addresses.Num() > 0)
 		{
 			UMapProperty* MapProperty = CastChecked<UMapProperty>(PropNode->GetProperty());
-			FScriptMapHelper MapHelper(MapProperty, PropNode->GetValueBaseAddress((uint8*)Addresses[0].Object));
+
+			FScriptMapHelper MapHelper(MapProperty, PropNode->GetValueBaseAddressFromObject(Addresses[0].Object));
 
 			FDefaultConstructedPropertyElement DefaultKey(MapHelper.KeyProp);
 
diff --git a/Engine/Source/Editor/PropertyEditor/Private/PropertyNode.cpp b/Engine/Source/Editor/PropertyEditor/Private/PropertyNode.cpp
--- a/Engine/Source/Editor/PropertyEditor/Private/PropertyNode.cpp
+++ b/Engine/Source/Editor/PropertyEditor/Private/PropertyNode.cpp
@@ -81,7 +81,7 @@ FPropertyNode::~FPropertyNode(void)
 
 
 
-void FPropertyNode::InitNode( const FPropertyNodeInitParams& InitParams )
+void FPropertyNode::InitNode(const FPropertyNodeInitParams& InitParams)
 {
 	//Dismantle the previous tree
 	DestroyTree();
@@ -91,6 +91,13 @@ void FPropertyNode::InitNode( const FPropertyNodeInitParams& InitParams )
 	ParentNode = InitParams.ParentNode.Get();
 	ParentNodeWeakPtr = InitParams.ParentNode;
 
+	//Property Data
+	Property = InitParams.Property;
+	ArrayOffset = InitParams.ArrayOffset;
+	ArrayIndex = InitParams.ArrayIndex;
+
+	bool bIsSparse = InitParams.bIsSparseProperty;
+
 	if (ParentNode)
 	{
 		//default to parents max child depth
@@ -100,17 +107,16 @@ void FPropertyNode::InitNode( const FPropertyNodeInitParams& InitParams )
 		{
 			--MaxChildDepthAllowed;
 		}
-	}
 
-	//Property Data
-	Property		= InitParams.Property;
-	ArrayOffset		= InitParams.ArrayOffset;
-	ArrayIndex		= InitParams.ArrayIndex;
+		// if the parent node's property is sparse data, our property must be too
+		bIsSparse = bIsSparse || ParentNode->HasNodeFlags(EPropertyNodeFlags::IsSparseClassData);
+	}
 
 	// Property is advanced if it is marked advanced or the entire class is advanced and the property not marked as simple
 	bool bAdvanced = Property.IsValid() ? ( Property->HasAnyPropertyFlags(CPF_AdvancedDisplay) || ( !Property->HasAnyPropertyFlags( CPF_SimpleDisplay ) && Property->GetOwnerClass() && Property->GetOwnerClass()->HasAnyClassFlags( CLASS_AdvancedDisplay ) ) ): false;
 
 	PropertyNodeFlags = EPropertyNodeFlags::NoFlags;
+	SetNodeFlags(EPropertyNodeFlags::IsSparseClassData, bIsSparse);
 
 	//default to copying from the parent
 	if (ParentNode)
@@ -910,36 +916,53 @@ bool FPropertyNode::GetReadAddress( FReadAddressList& OutAddresses )
 	return bSuccess;
 }
 
+uint8* FPropertyNode::GetStartAddress(const UObject* Obj) const
+{
+	if (!Obj)
+	{
+		return nullptr;
+	}
 
-/**
- * Calculates the memory address for the data associated with this item's property.  This is typically the value of a UProperty or a UObject address.
- *
- * @param	StartAddress	the location to use as the starting point for the calculation; typically the address of the object that contains this property.
- *
- * @return	a pointer to a UProperty value or UObject.  (For dynamic arrays, you'd cast this value to an FArray*)
- */
-uint8* FPropertyNode::GetValueBaseAddress( uint8* StartAddress )
+	if (HasNodeFlags(EPropertyNodeFlags::IsSparseClassData))
+	{
+		return (uint8*)Obj->GetClass()->GetOrCreateSparseClassData();
+	}
+
+	return (uint8*)Obj;
+}
+
+uint8* FPropertyNode::GetValueBaseAddressFromObject(const UObject* Obj)
+{
+	return GetValueBaseAddress(GetStartAddress(Obj), HasNodeFlags(EPropertyNodeFlags::IsSparseClassData));
+}
+
+uint8* FPropertyNode::GetValueAddressFromObject(const UObject* Obj)
+{
+	return GetValueAddress(GetStartAddress(Obj), HasNodeFlags(EPropertyNodeFlags::IsSparseClassData));
+}
+
+
+uint8* FPropertyNode::GetValueBaseAddress(uint8* StartAddress, bool bIsSparseData)
 {
 	uint8* Result = NULL;
 
-	if ( ParentNodeWeakPtr.IsValid() )
+	if (bIsSparseData)
 	{
-		Result = ParentNode->GetValueAddress(StartAddress);
+		Result = StartAddress;
+	}
+	else
+	{
+		if (ParentNodeWeakPtr.IsValid())
+		{
+			Result = ParentNode->GetValueAddress(StartAddress, bIsSparseData);
+		}
 	}
 	return Result;
 }
 
-/**
- * Calculates the memory address for the data associated with this item's value.  For most properties, identical to GetValueBaseAddress.  For items corresponding
- * to dynamic array elements, the pointer returned will be the location for that element's data. 
- *
- * @param	StartAddress	the location to use as the starting point for the calculation; typically the address of the object that contains this property.
- *
- * @return	a pointer to a UProperty value or UObject.  (For dynamic arrays, you'd cast this value to whatever type is the Inner for the dynamic array)
- */
-uint8* FPropertyNode::GetValueAddress( uint8* StartAddress )
+uint8* FPropertyNode::GetValueAddress(uint8* StartAddress, bool bIsSparseData)
 {
-	return GetValueBaseAddress( StartAddress );
+	return GetValueBaseAddress(StartAddress, bIsSparseData);
 }
 
 
@@ -973,7 +996,7 @@ public:
 
 	void InnerInitialize()
 	{
-	{
+		{
 			PropertyValueRoot.OwnerObject = NULL;
 			PropertyDefaultValueRoot.OwnerObject = NULL;
 			PropertyValueAddress = NULL;
@@ -996,41 +1019,31 @@ public:
 			PropertyValueRoot.OwnerObject = Cast<UClass>(PropertyValueRoot.OwnerObject)->GetDefaultObject();
 		}
 
-		UArrayProperty* ArrayProp = Cast<UArrayProperty>(Property);
-		UArrayProperty* OuterArrayProp = Cast<UArrayProperty>(Property->GetOuter());
+		const bool bIsContainerProperty = Cast<UArrayProperty>(Property) || Cast<USetProperty>(Property) || Cast<UMapProperty>(Property);
+		const bool bIsInsideContainerProperty = Cast<UArrayProperty>(Property->GetOuter()) || Cast<USetProperty>(Property->GetOuter()) || Cast<UMapProperty>(Property->GetOuter());
 
-		USetProperty* SetProp = Cast<USetProperty>(Property);
-		USetProperty* OuterSetProp = Cast<USetProperty>(Property->GetOuter());
+		FPropertyNode* Node = bIsInsideContainerProperty ? ParentNode : PropertyNode;
 
-		UMapProperty* MapProp = Cast<UMapProperty>(Property);
-		UMapProperty* OuterMapProp = Cast<UMapProperty>(Property->GetOuter());
+		PropertyValueBaseAddress = Node->GetValueBaseAddressFromObject(PropertyValueRoot.OwnerObject);
+		PropertyValueAddress = PropertyNode->GetValueAddressFromObject(PropertyValueRoot.OwnerObject);
 
-		// calculate the values for the current object
-		{
-			PropertyValueBaseAddress = (OuterArrayProp == NULL && OuterSetProp == NULL && OuterMapProp == NULL)
-				? PropertyNode->GetValueBaseAddress(PropertyValueRoot.ValueAddress)
-				: ParentNode->GetValueBaseAddress(PropertyValueRoot.ValueAddress);
-
-			PropertyValueAddress = PropertyNode->GetValueAddress(PropertyValueRoot.ValueAddress);
-		}
-
-		if( IsValidTracker() )
+		if (IsValidTracker())
 		{
 			 bHasDefaultValue = Private_HasDefaultValue();
-			// calculate the values for the default object
-			if ( bHasDefaultValue )
+
+			// calculate the addresses for the default object if it exists
+			if (bHasDefaultValue)
 			{
-				PropertyDefaultValueRoot.OwnerObject = PropertyValueRoot.OwnerObject ? PropertyValueRoot.OwnerObject->GetArchetype() : NULL;
-				PropertyDefaultBaseAddress = (OuterArrayProp == NULL && OuterSetProp == NULL && OuterMapProp == NULL)
-					? PropertyNode->GetValueBaseAddress(PropertyDefaultValueRoot.ValueAddress)
-					: ParentNode->GetValueBaseAddress(PropertyDefaultValueRoot.ValueAddress);
-				PropertyDefaultAddress = PropertyNode->GetValueAddress(PropertyDefaultValueRoot.ValueAddress);
+				PropertyDefaultValueRoot.OwnerObject = PropertyValueRoot.OwnerObject ? PropertyValueRoot.OwnerObject->GetArchetype() : nullptr;
+
+				PropertyDefaultBaseAddress = Node->GetValueBaseAddressFromObject(PropertyDefaultValueRoot.OwnerObject);
+				PropertyDefaultAddress = PropertyNode->GetValueAddressFromObject(PropertyDefaultValueRoot.OwnerObject);
 
 				//////////////////////////
 				// If this is a container property, we must take special measures to use the base address of the property's value; for instance,
 				// the array property's PropertyDefaultBaseAddress points to an FScriptArray*, while PropertyDefaultAddress points to the 
 				// FScriptArray's Data pointer.
-				if ( ArrayProp != NULL || SetProp != NULL || MapProp != NULL )
+				if (bIsContainerProperty)
 				{
 					PropertyValueAddress = PropertyValueBaseAddress;
 					PropertyDefaultAddress = PropertyDefaultBaseAddress;
@@ -1045,8 +1058,8 @@ public:
 	 * @param	InPropItem		the property window item this struct will hold values for
 	 * @param	InOwnerObject	the object which contains the property value
 	 */
-	FPropertyItemValueDataTrackerSlate( FPropertyNode* InPropertyNode, UObject* InOwnerObject )
-		: OwnerObject( InOwnerObject )
+	FPropertyItemValueDataTrackerSlate(FPropertyNode* InPropertyNode, UObject* InOwnerObject)
+		: OwnerObject(InOwnerObject)
 		, PropertyNode(InPropertyNode)
 		, bHasDefaultValue(false)
 	{
@@ -1124,12 +1137,8 @@ public:
 	 */
 	uint8* GetPropertyDefaultAddress() const { return PropertyDefaultAddress; }
 
-	/**
-	 * @return The address of the owning object's archetype
-	 */
-	FPropertyValueRoot GetPropertyValueRoot() const { return PropertyValueRoot; }
 private:
-		/**
+	/**
 	 * Determines whether the property bound to this struct exists in the owning object's archetype.
 	 *
 	 * @return	true if this property exists in the owning object's archetype; false if the archetype is e.g. a
@@ -1139,13 +1148,50 @@ private:
 	{
 		bool bResult = false;
 
-		if( IsValidTracker() )
+		if (IsValidTracker())
 		{
+			UClass* OwnerClass = PropertyValueRoot.OwnerObject->GetClass();
+			if (GetPropertyNode()->HasNodeFlags(EPropertyNodeFlags::IsSparseClassData))
+			{
+				if (OwnerClass)
+				{
+					UStruct* SparseClassDataStruct = OwnerClass->GetSparseClassDataStruct();
+					UStruct* SparseClassDataArchetypeStruct = OwnerClass->GetSparseClassDataArchetypeStruct();
+
+					if (SparseClassDataStruct == SparseClassDataArchetypeStruct)
+					{
+						bResult = true;
+					}
+					else
+					{
+						// Find the member property which contains this item's property
+						FPropertyNode* MemberPropertyNode = PropertyNode;
+						for (; MemberPropertyNode != nullptr; MemberPropertyNode = MemberPropertyNode->GetParentNode())
+						{
+							UProperty* MemberProperty = MemberPropertyNode->GetProperty();
+							if (MemberProperty != nullptr)
+							{
+								if (Cast<UClass>(MemberProperty->GetOuter()) != nullptr)
+								{
+									break;
+								}
+							}
+						}
+						if (MemberPropertyNode != nullptr && MemberPropertyNode->GetProperty())
+						{
+							// we check to see that this property is in the defaults class
+							bResult = MemberPropertyNode->GetProperty()->IsInContainer(SparseClassDataArchetypeStruct);
+						}
+					}
+				}
+
+				return bResult;
+			}
 			check(PropertyValueBaseAddress);
 			check(PropertyValueRoot.OwnerObject);
 			UObject* ParentDefault = PropertyValueRoot.OwnerObject->GetArchetype();
 			check(ParentDefault);
-			if (PropertyValueRoot.OwnerObject->GetClass() == ParentDefault->GetClass())
+			if (OwnerClass == ParentDefault->GetClass())
 			{
 				// if the archetype is of the same class, then we must have a default
 				bResult = true;
@@ -1153,20 +1199,20 @@ private:
 			else
 			{
 				// Find the member property which contains this item's property
-			FPropertyNode* MemberPropertyNode = PropertyNode;
-			for ( ;MemberPropertyNode != NULL; MemberPropertyNode = MemberPropertyNode->GetParentNode() )
-			{
-				UProperty* MemberProperty = MemberPropertyNode->GetProperty();
-				if ( MemberProperty != NULL )
+				FPropertyNode* MemberPropertyNode = PropertyNode;
+				for ( ;MemberPropertyNode != NULL; MemberPropertyNode = MemberPropertyNode->GetParentNode() )
 				{
-					if ( Cast<UClass>(MemberProperty->GetOuter()) != NULL )
+					UProperty* MemberProperty = MemberPropertyNode->GetProperty();
+					if ( MemberProperty != NULL )
 					{
-						break;
+						if ( Cast<UClass>(MemberProperty->GetOuter()) != NULL )
+						{
+							break;
+						}
 					}
 				}
-			}
 				if ( MemberPropertyNode != NULL && MemberPropertyNode->GetProperty())
-			{
+				{
 					// we check to see that this property is in the defaults class
 					bResult = MemberPropertyNode->GetProperty()->IsInContainer(ParentDefault->GetClass());
 				}
@@ -1555,26 +1601,31 @@ bool FPropertyNode::GetDiffersFromDefaultForObject( FPropertyItemValueDataTracke
 
 	bool bDiffersFromDefaultForObject = false;
 
-	if ( ValueTracker.IsValidTracker() && ValueTracker.HasDefaultValue() && GetParentNode() != NULL )
+	bool bIsValidTracker = ValueTracker.IsValidTracker();
+	bool bHasDefaultValue = ValueTracker.HasDefaultValue();
+	bool bHasParent = GetParentNode() != nullptr;
+
+	if (bIsValidTracker && bHasDefaultValue && bHasParent)
 	{
-		if (ValueTracker.GetPropertyDefaultBaseAddress() != NULL)
-		{
-			//////////////////////////
-			// Check the property against its default.
-			// If the property is an object property, we have to take special measures.
-			UArrayProperty* OuterArrayProperty = Cast<UArrayProperty>(InProperty->GetOuter());
-			USetProperty* OuterSetProperty = Cast<USetProperty>(InProperty->GetOuter());
-			UMapProperty* OuterMapProperty = Cast<UMapProperty>(InProperty->GetOuter());
+		//////////////////////////
+		// Check the property against its default.
+		// If the property is an object property, we have to take special measures.
 
-			if ( OuterArrayProperty != NULL )
+		if (UArrayProperty* OuterArrayProperty = Cast<UArrayProperty>(InProperty->GetOuter()))
+		{
+			// make sure we're not trying to compare against an element that doesn't exist
+			if (ValueTracker.GetPropertyDefaultBaseAddress() != nullptr)
 			{
 				// make sure we're not trying to compare against an element that doesn't exist
-				if (GetArrayIndex() >= FScriptArrayHelper::Num(ValueTracker.GetPropertyDefaultBaseAddress()) )
+				if (GetArrayIndex() >= FScriptArrayHelper::Num(ValueTracker.GetPropertyDefaultBaseAddress()))
 				{
 					bDiffersFromDefaultForObject = true;
 				}
 			}
-			else if (OuterSetProperty != NULL)
+		}
+		else if (USetProperty* OuterSetProperty = Cast<USetProperty>(InProperty->GetOuter()))
+		{
+			if (ValueTracker.GetPropertyDefaultBaseAddress() != nullptr)
 			{
 				FScriptSetHelper SetHelper(OuterSetProperty, ValueTracker.GetPropertyDefaultBaseAddress());
 				bool bIsValidIndex = ArrayIndex >= 0 && ArrayIndex < SetHelper.Num();
@@ -1583,7 +1634,10 @@ bool FPropertyNode::GetDiffersFromDefaultForObject( FPropertyItemValueDataTracke
 					bDiffersFromDefaultForObject = true;
 				}
 			}
-			else if (OuterMapProperty != NULL)
+		}
+		else if (UMapProperty* OuterMapProperty = Cast<UMapProperty>(InProperty->GetOuter()))
+		{
+			if (ValueTracker.GetPropertyDefaultBaseAddress() != nullptr)
 			{
 				FScriptMapHelper MapHelper(OuterMapProperty, ValueTracker.GetPropertyDefaultBaseAddress());
 				bool bIsValidIndex = ArrayIndex >= 0 && ArrayIndex < MapHelper.Num();
@@ -1595,14 +1649,13 @@ bool FPropertyNode::GetDiffersFromDefaultForObject( FPropertyItemValueDataTracke
 		}
 
 		// The property is a simple field.  Compare it against the enclosing object's default for that property.
-		if ( !bDiffersFromDefaultForObject)
+		if (!bDiffersFromDefaultForObject)
 		{
 			uint32 PortFlags = 0;
-			UObjectPropertyBase* ObjectProperty = Cast<UObjectPropertyBase>(InProperty);
 			if (InProperty->ContainsInstancedObjectProperty())
 			{
 				// Use PPF_DeepCompareInstances for component objects
-				if (ObjectProperty)
+				if (Cast<UObjectPropertyBase>(InProperty))
 				{
 					PortFlags |= PPF_DeepCompareInstances;
 				}
@@ -1614,9 +1667,10 @@ bool FPropertyNode::GetDiffersFromDefaultForObject( FPropertyItemValueDataTracke
 				// have enough elements in either the default or the object
 				bDiffersFromDefaultForObject = true;
 			}
-			else if ( GetArrayIndex() == INDEX_NONE && InProperty->ArrayDim > 1 )
+			else if (GetArrayIndex() == INDEX_NONE && InProperty->ArrayDim > 1)
 			{
-				for ( int32 Idx = 0; !bDiffersFromDefaultForObject && Idx < InProperty->ArrayDim; Idx++ )
+				// this is a container; loop through all of its elements and see if any of them differ from the default
+				for (int32 Idx = 0; !bDiffersFromDefaultForObject && Idx < InProperty->ArrayDim; Idx++)
 				{
 					bDiffersFromDefaultForObject = !InProperty->Identical(
 						ValueTracker.GetPropertyValueAddress() + Idx * InProperty->ElementSize,
@@ -1627,6 +1681,7 @@ bool FPropertyNode::GetDiffersFromDefaultForObject( FPropertyItemValueDataTracke
 			}
 			else
 			{
+				// try to compare the values at the current and default property addresses
 				uint8* PropertyValueAddr = ValueTracker.GetPropertyValueAddress();
 				uint8* DefaultPropertyValueAddr = ValueTracker.GetPropertyDefaultAddress();
 
@@ -1655,6 +1710,14 @@ bool FPropertyNode::GetDiffersFromDefault()
 		bUpdateDiffersFromDefault = false;
 		bDiffersFromDefault = false;
 
+		UProperty* Prop = GetProperty();
+
+		if (!Prop)
+		{
+			return bDiffersFromDefault;
+		}
+
+
 		FObjectPropertyNode* ObjectNode = FindObjectItemParent();
 		if(ObjectNode && Property.IsValid() && !IsEditConst())
 		{
@@ -1691,50 +1754,39 @@ FString FPropertyNode::GetDefaultValueAsStringForObject( FPropertyItemValueDataT
 	{
 		if ( ValueTracker.IsValidTracker() && ValueTracker.HasDefaultValue() )
 		{
-			//////////////////////////
-			// Check the property against its default.
-			// If the property is an object property, we have to take special measures.
-			UArrayProperty* OuterArrayProperty = Cast<UArrayProperty>(InProperty->GetOuter());
-			USetProperty* OuterSetProperty = Cast<USetProperty>(InProperty->GetOuter());
-			UMapProperty* OuterMapProperty = Cast<UMapProperty>(InProperty->GetOuter());
-
-			// The property is a simple field.  Compare it against the enclosing object's default for that property.
-			if ( !bDiffersFromDefaultForObject)
+			uint32 PortFlags = bUseDisplayName ? PPF_PropertyWindow : PPF_None;
+			
+			if (InProperty->ContainsInstancedObjectProperty())
 			{
-				uint32 PortFlags = bUseDisplayName ? PPF_PropertyWindow : PPF_None;
-				UObjectPropertyBase* ObjectProperty = Cast<UObjectPropertyBase>(InProperty);
-				if (InProperty->ContainsInstancedObjectProperty())
+				// Use PPF_DeepCompareInstances for component objects
+				if (Cast<UObjectPropertyBase>(InProperty))
 				{
-					// Use PPF_DeepCompareInstances for component objects
-					if (ObjectProperty)
-					{
-						PortFlags |= PPF_DeepCompareInstances;
-					}
+					PortFlags |= PPF_DeepCompareInstances;
 				}
+			}
 
-				if ( ValueTracker.GetPropertyDefaultAddress() == NULL )
+			if ( ValueTracker.GetPropertyDefaultAddress() == NULL )
+			{
+				// no default available, fall back on the default value for our primitive:
+				uint8* TempComplexPropAddr = (uint8*)FMemory::Malloc(InProperty->GetSize(), InProperty->GetMinAlignment());
+				InProperty->InitializeValue(TempComplexPropAddr);
+				ON_SCOPE_EXIT
 				{
-					// no default available, fall back on the default value for our primitive:
-					uint8* TempComplexPropAddr = (uint8*)FMemory::Malloc(InProperty->GetSize(), InProperty->GetMinAlignment());
-					InProperty->InitializeValue(TempComplexPropAddr);
-					ON_SCOPE_EXIT
-					{
-						InProperty->DestroyValue(TempComplexPropAddr);
-						FMemory::Free(TempComplexPropAddr);
-					};
+					InProperty->DestroyValue(TempComplexPropAddr);
+					FMemory::Free(TempComplexPropAddr);
+				};
 					
-					InProperty->ExportText_Direct(DefaultValue, TempComplexPropAddr, TempComplexPropAddr, nullptr, PPF_None);
-				}
-				else if ( GetArrayIndex() == INDEX_NONE && InProperty->ArrayDim > 1 )
-				{
-					UArrayProperty::ExportTextInnerItem(DefaultValue, InProperty, ValueTracker.GetPropertyDefaultAddress(), InProperty->ArrayDim,
-														ValueTracker.GetPropertyDefaultAddress(), InProperty->ArrayDim, nullptr, PortFlags);
-				}
-				else
-				{
-					// Port flags will cause enums to display correctly
-					InProperty->ExportTextItem( DefaultValue, ValueTracker.GetPropertyDefaultAddress(), ValueTracker.GetPropertyDefaultAddress(), InObject, PortFlags, NULL );
-				}
+				InProperty->ExportText_Direct(DefaultValue, TempComplexPropAddr, TempComplexPropAddr, nullptr, PPF_None);
+			}
+			else if ( GetArrayIndex() == INDEX_NONE && InProperty->ArrayDim > 1 )
+			{
+				UArrayProperty::ExportTextInnerItem(DefaultValue, InProperty, ValueTracker.GetPropertyDefaultAddress(), InProperty->ArrayDim,
+													ValueTracker.GetPropertyDefaultAddress(), InProperty->ArrayDim, nullptr, PortFlags);
+			}
+			else
+			{
+				// Port flags will cause enums to display correctly
+				InProperty->ExportTextItem( DefaultValue, ValueTracker.GetPropertyDefaultAddress(), ValueTracker.GetPropertyDefaultAddress(), InObject, PortFlags, NULL );
 			}
 		}
 	}
@@ -2347,7 +2399,7 @@ FPropertyChangedEvent& FPropertyNode::FixPropertiesInEvent(FPropertyChangedEvent
 {
 	ensure(Event.Property);
 
-	auto PropertyChain = BuildPropertyChain(Event.Property);
+	TSharedRef<FEditPropertyChain> PropertyChain = BuildPropertyChain(Event.Property);
 	auto MemberProperty = PropertyChain->GetActiveMemberNode() ? PropertyChain->GetActiveMemberNode()->GetValue() : NULL;
 	if (ensure(MemberProperty))
 	{
@@ -2484,7 +2536,7 @@ void FPropertyNode::PropagateContainerPropertyChange( UObject* ModifiedObject, c
 		{
 			// If the original object is a subobject, get the current object's subobject too.
 			// In this case we're not going to modify ObjToChange but its default subobject.
-			ActualObjToChange = *(UObject**)SubobjectPropertyNode->GetValueBaseAddress((uint8*)ObjToChange);
+			ActualObjToChange = *(UObject**)SubobjectPropertyNode->GetValueBaseAddressFromObject(ObjToChange);
 		}
 		else
 		{
@@ -2497,16 +2549,25 @@ void FPropertyNode::PropagateContainerPropertyChange( UObject* ModifiedObject, c
 
 			if (ChangeType == EPropertyArrayChangeType::Add || ChangeType == EPropertyArrayChangeType::Clear)
 			{
-				Addr = GetValueBaseAddress((uint8*)ActualObjToChange);
+				Addr = GetValueBaseAddressFromObject(ActualObjToChange);
 			}
 			else
 			{
-				Addr = ParentPropertyNode->GetValueBaseAddress((uint8*)ActualObjToChange);
+				Addr = ParentPropertyNode->GetValueBaseAddressFromObject(ActualObjToChange);
 			}
 
 			if (Addr != nullptr)
 			{
-				checkf(OriginalContainerAddr != Addr, TEXT("PropagateContainerPropertyChange tried to propagate a change onto itself!"));
+				if (OriginalContainerAddr == Addr)
+				{
+					if (HasNodeFlags(EPropertyNodeFlags::IsSparseClassData))
+					{
+						continue;
+					}
+
+					checkf(false, TEXT("PropagateContainerPropertyChange tried to propagate a change onto itself!"));
+				}
+				
 				const bool bIsDefaultContainerContent = ConvertedProperty->Identical(OriginalContainerAddr, Addr);
 
 				// Return instance changes result to caller
@@ -2696,7 +2757,7 @@ void FPropertyNode::PropagatePropertyChange( UObject* ModifiedObject, const TCHA
 		{
 			// If the original object is a subobject, get the current object's subobject too.
 			// In this case we're not going to modify ObjToChange but its default subobject.
-			ActualObjToChange = *(UObject**)SubobjectPropertyNode->GetValueBaseAddress((uint8*)ObjToChange);
+			ActualObjToChange = *(UObject**)SubobjectPropertyNode->GetValueBaseAddressFromObject(ObjToChange);
 		}
 		else
 		{
@@ -2705,7 +2766,7 @@ void FPropertyNode::PropagatePropertyChange( UObject* ModifiedObject, const TCHA
 
 		if (ActualObjToChange != ModifiedObject)
 		{
-			uint8* DestSimplePropAddr = GetValueBaseAddress( (uint8*)ActualObjToChange );
+			uint8* DestSimplePropAddr = GetValueBaseAddressFromObject(ActualObjToChange);
 			if (DestSimplePropAddr != nullptr)
 			{
 				UProperty* ComplexProperty = Prop;
@@ -2716,8 +2777,8 @@ void FPropertyNode::PropagatePropertyChange( UObject* ModifiedObject, const TCHA
 					ComplexPropertyNode = ParentNode;
 				}
 				
-				uint8* DestComplexPropAddr = ComplexPropertyNode->GetValueBaseAddress((uint8*)ActualObjToChange);
-				uint8* ModifiedComplexPropAddr = ComplexPropertyNode->GetValueBaseAddress((uint8*)ModifiedObject);
+				uint8* DestComplexPropAddr = ComplexPropertyNode->GetValueBaseAddressFromObject(ActualObjToChange);
+				uint8* ModifiedComplexPropAddr = ComplexPropertyNode->GetValueBaseAddressFromObject(ModifiedObject);
 
 				bool bShouldImport = false;
 				{
diff --git a/Engine/Source/Editor/PropertyEditor/Private/PropertyNode.h b/Engine/Source/Editor/PropertyEditor/Private/PropertyNode.h
--- a/Engine/Source/Editor/PropertyEditor/Private/PropertyNode.h
+++ b/Engine/Source/Editor/PropertyEditor/Private/PropertyNode.h
@@ -59,6 +59,8 @@ namespace EPropertyNodeFlags
 
 	const Type	HasCustomResetToDefault			= 1 << 23;	/** true if this node's visual representation of reset to default has been customized*/
 
+	const Type	IsSparseClassData				= 1 << 24;	/** true if the property on this node is part of a sparse class data structure */
+
 	const Type 	NoFlags							= 0;
 
 };
@@ -206,6 +208,8 @@ struct FPropertyNodeInitParams
 	bool bCreateCategoryNodes;
 	/** Whether or not to create nodes for properties marked CPF_DisableEditOnInstance */
 	bool bCreateDisableEditOnInstanceNodes;
+	/** Whether or not this property is sparse data */
+	bool bIsSparseProperty;
 
 	FPropertyNodeInitParams()
 		: ParentNode(nullptr)
@@ -216,6 +220,7 @@ struct FPropertyNodeInitParams
 		, bForceHiddenPropertyVisibility( false )
 		, bCreateCategoryNodes( true )
 		, bCreateDisableEditOnInstanceNodes( true )
+		, bIsSparseProperty( false )
 	{}
 };
 
@@ -443,18 +448,41 @@ public:
 	 * Calculates the memory address for the data associated with this item's property.  This is typically the value of a UProperty or a UObject address.
 	 *
 	 * @param	StartAddress	the location to use as the starting point for the calculation; typically the address of the object that contains this property.
+	 * @param	bIsSparseData	True if StartAddress is pointing to a sidecar structure containing sparse class data, false otherwise
 	 *
 	 * @return	a pointer to a UProperty value or UObject.  (For dynamic arrays, you'd cast this value to an FArray*)
 	 */
-	virtual uint8* GetValueBaseAddress(uint8* Base);
+	virtual uint8* GetValueBaseAddress(uint8* StartAddress, bool bIsSparseData);
 
 	/**
 	 * Calculates the memory address for the data associated with this item's value.  For most properties, identical to GetValueBaseAddress.  For items corresponding
 	 * to dynamic array elements, the pointer returned will be the location for that element's data.
 	 *
+	 * @param	StartAddress	the location to use as the starting point for the calculation; typically the address of the object that contains this property.
+	 * @param	bIsSparseData	True if StartAddress is pointing to a sidecar structure containing sparse class data, false otherwise
+	 *
 	 * @return	a pointer to a UProperty value or UObject.  (For dynamic arrays, you'd cast this value to whatever type is the Inner for the dynamic array)
 	 */
-	virtual uint8* GetValueAddress(uint8* Base);
+	virtual uint8* GetValueAddress(uint8* StartAddress, bool bIsSparseData);
+
+	/**
+	 * Calculates the memory address for the data associated with this item's property.  This is typically the value of a UProperty or a UObject address.
+	 *
+	 * @param	Obj	The object that contains this property; used as the starting point for the calculation
+	 *
+	 * @return	a pointer to a UProperty value or UObject.  (For dynamic arrays, you'd cast this value to an FArray*)
+	 */
+	uint8* GetValueBaseAddressFromObject(const UObject* Obj);
+
+	/**
+	 * Calculates the memory address for the data associated with this item's value.  For most properties, identical to GetValueBaseAddress.  For items corresponding
+	 * to dynamic array elements, the pointer returned will be the location for that element's data.
+	 *
+	 * @param	Obj	The object that contains this property; used as the starting point for the calculation
+	 *
+	 * @return	a pointer to a UProperty value or UObject.  (For dynamic arrays, you'd cast this value to whatever type is the Inner for the dynamic array)
+	 */
+	uint8* GetValueAddressFromObject(const UObject* Obj);
 
 	/**
 	 * Sets the display name override to use instead of the display name
@@ -825,6 +853,10 @@ public:
 	void BroadcastPropertyResetToDefault();
 protected:
 
+	// Returns a pointer to the starting point of the structure that contains the property this node uses.
+	// This will often be Obj but may also point to a sidecar data structure
+	uint8* GetStartAddress(const UObject* Obj) const;
+
 	TSharedRef<FEditPropertyChain> BuildPropertyChain( UProperty* PropertyAboutToChange );
 
 	/**
@@ -1019,8 +1051,17 @@ public:
 	virtual UStruct* GetBaseStructure() = 0;
 	virtual const UStruct* GetBaseStructure() const = 0;
 
+	// Returns the base struct as well as any sidecar data structs
+	virtual TArray<UStruct*> GetAllStructures() = 0;
+	virtual TArray<const UStruct*> GetAllStructures() const = 0;
+
 	virtual int32 GetInstancesNum() const = 0;
 	virtual uint8* GetMemoryOfInstance(int32 Index) = 0;
+
+	/**
+	 * Returns a pointer to the stored value of InProperty on InParentNode's Index'th instance.
+	 */
+	virtual uint8* GetValuePtrOfInstance(int32 Index, const UProperty* InProperty, FPropertyNode* InParentNode) = 0;
 	virtual TWeakObjectPtr<UObject> GetInstanceAsUObject(int32 Index) = 0;
 	virtual EPropertyType GetPropertyType() const = 0;
 
diff --git a/Engine/Source/Editor/PropertyEditor/Private/StructurePropertyNode.h b/Engine/Source/Editor/PropertyEditor/Private/StructurePropertyNode.h
--- a/Engine/Source/Editor/PropertyEditor/Private/StructurePropertyNode.h
+++ b/Engine/Source/Editor/PropertyEditor/Private/StructurePropertyNode.h
@@ -19,13 +19,6 @@ public:
 	virtual FStructurePropertyNode* AsStructureNode() override { return this; }
 	virtual const FStructurePropertyNode* AsStructureNode() const override { return this; }
 
-	/** FPropertyNode Interface */
-	virtual uint8* GetValueBaseAddress(uint8* Base) override
-	{
-		ensure(true);
-		return HasValidStructData() ? StructData->GetStructMemory() : NULL;
-	}
-
 	void SetStructure(TSharedPtr<FStructOnScope> InStructData)
 	{
 		ClearCachedReadAddresses(true);
@@ -65,7 +58,7 @@ public:
 
 		uint8* ReadAddress = StructData->GetStructMemory();
 		check(ReadAddress);
-		OutAddresses.Add(NULL, InPropertyNode.GetValueBaseAddress(ReadAddress), true);
+		OutAddresses.Add(nullptr, InPropertyNode.GetValueBaseAddress(ReadAddress, InPropertyNode.HasNodeFlags(EPropertyNodeFlags::IsSparseClassData)), true);
 		return true;
 	}
 
@@ -95,13 +88,31 @@ public:
 	/** FComplexPropertyNode Interface */
 	virtual const UStruct* GetBaseStructure() const override
 	{ 
-		return HasValidStructData() ? StructData->GetStruct() : NULL;
+		return HasValidStructData() ? StructData->GetStruct() : nullptr;
 	}
 	virtual UStruct* GetBaseStructure() override
 	{
-		const UStruct* Struct = HasValidStructData() ? StructData->GetStruct() : NULL;
+		const UStruct* Struct = HasValidStructData() ? StructData->GetStruct() : nullptr;
 		return const_cast<UStruct*>(Struct);
 	}
+	virtual TArray<UStruct*> GetAllStructures() override
+	{
+		TArray<UStruct*> RetVal;
+		if (UStruct* BaseStruct = GetBaseStructure())
+		{
+			RetVal.Add(BaseStruct);
+		}
+		return RetVal;
+	}
+	virtual TArray<const UStruct*> GetAllStructures() const override
+	{
+		TArray<const UStruct*> RetVal;
+		if (const UStruct* BaseStruct = GetBaseStructure())
+		{
+			RetVal.Add(BaseStruct);
+		}
+		return RetVal;
+	}
 	virtual int32 GetInstancesNum() const override
 	{ 
 		return HasValidStructData() ? 1 : 0;
@@ -111,6 +122,13 @@ public:
 		check(0 == Index);
 		return HasValidStructData() ? StructData->GetStructMemory() : NULL;
 	}
+	virtual uint8* GetValuePtrOfInstance(int32 Index, const UProperty* InProperty, FPropertyNode* InParentNode) override
+	{ 
+		check(0 == Index);
+		uint8* StructBaseAddress = HasValidStructData() ? StructData->GetStructMemory() : nullptr;
+		uint8* ParentBaseAddress = InParentNode ? InParentNode->GetValueAddress(StructBaseAddress, false) : nullptr;
+		return InProperty ? InProperty->ContainerPtrToValuePtr<uint8>(ParentBaseAddress) : nullptr;
+	}
 	virtual TWeakObjectPtr<UObject> GetInstanceAsUObject(int32 Index) override
 	{
 		check(0 == Index);
@@ -131,6 +149,12 @@ protected:
 	/** FPropertyNode interface */
 	virtual void InitChildNodes() override;
 
+	virtual uint8* GetValueBaseAddress(uint8* Base, bool bIsSparseData) override
+	{
+		check(bIsSparseData == false);
+		return HasValidStructData() ? StructData->GetStructMemory() : nullptr;
+	}
+
 	virtual bool GetQualifiedName(FString& PathPlusIndex, const bool bWithArrayIndex, const FPropertyNode* StopParent = nullptr, bool bIgnoreCategories = false) const override
 	{
 		PathPlusIndex += TEXT("Struct");
diff --git a/Engine/Source/Editor/UnrealEd/Private/Kismet2/BlueprintEditorUtils.cpp b/Engine/Source/Editor/UnrealEd/Private/Kismet2/BlueprintEditorUtils.cpp
--- a/Engine/Source/Editor/UnrealEd/Private/Kismet2/BlueprintEditorUtils.cpp
+++ b/Engine/Source/Editor/UnrealEd/Private/Kismet2/BlueprintEditorUtils.cpp
@@ -1233,6 +1233,7 @@ void FBlueprintEditorUtils::RemoveStaleFunctions(UBlueprintGeneratedClass* Class
 		OrphanedClass->ClassAddReferencedObjects = Class->AddReferencedObjects;
 		OrphanedClass->ClassFlags |= CLASS_CompiledFromBlueprint;
 		OrphanedClass->ClassGeneratedBy = Class->ClassGeneratedBy;
+		OrphanedClass->SparseClassData = Class->GetOrCreateSparseClassData();
 
 		const ERenameFlags RenFlags = REN_DontCreateRedirectors | (Blueprint->bIsRegeneratingOnLoad ? REN_ForceNoResetLoaders : 0) | REN_NonTransactional | REN_DoNotDirty;
 
@@ -1541,6 +1542,7 @@ void FBlueprintEditorUtils::RecreateClassMetaData(UBlueprint* Blueprint, UClass*
 		Class->RemoveMetaData("HideFunctions");
 		Class->RemoveMetaData("AutoExpandCategories");
 		Class->RemoveMetaData("AutoCollapseCategories");
+		Class->RemoveMetaData("SparseClassDataTypes");
 		Class->RemoveMetaData("ClassGroupNames");
 		Class->RemoveMetaData("Category");
 		Class->RemoveMetaData(FBlueprintMetadata::MD_AllowableBlueprintVariableType);
@@ -1586,6 +1588,11 @@ void FBlueprintEditorUtils::RecreateClassMetaData(UBlueprint* Blueprint, UClass*
 
 			Class->SetMetaData(NAME_ClassGroupNames, *ClassGroupCategory);
 		}
+
+		if (ParentClass->HasMetaData(TEXT("SparseClassDataTypes")))
+		{
+			Class->SetMetaData(TEXT("SparseClassDataTypes"), *ParentClass->GetMetaData("SparseClassDataTypes"));
+		}
 	}
 
 	// Add a category if one has been specified
diff --git a/Engine/Source/Editor/UnrealEd/Private/ScriptDisassembler.cpp b/Engine/Source/Editor/UnrealEd/Private/ScriptDisassembler.cpp
--- a/Engine/Source/Editor/UnrealEd/Private/ScriptDisassembler.cpp
+++ b/Engine/Source/Editor/UnrealEd/Private/ScriptDisassembler.cpp
@@ -508,6 +508,12 @@ void FKismetBytecodeDisassembler::ProcessCommon(int32& ScriptIndex, EExprToken O
 			Ar.Logf(TEXT("%s $%X: Local out variable named %s"), *Indents, (int32)Opcode, PropertyPtr ? *PropertyPtr->GetName() : TEXT("(null)"));
 			break;
 		}
+	case EX_ClassSparseDataVariable:
+		{
+			UProperty* PropertyPtr = ReadPointer<UProperty>(ScriptIndex);
+			Ar.Logf(TEXT("%s $%X: Class sparse data variable named %s"), *Indents, (int32)Opcode, PropertyPtr ? *PropertyPtr->GetName() : TEXT("(null)"));
+			break;
+		}
 	case EX_InterfaceContext:
 		{
 			Ar.Logf(TEXT("%s $%X: EX_InterfaceContext:"), *Indents, (int32)Opcode);
diff --git a/Engine/Source/Programs/UnrealHeaderTool/Private/BaseParser.cpp b/Engine/Source/Programs/UnrealHeaderTool/Private/BaseParser.cpp
--- a/Engine/Source/Programs/UnrealHeaderTool/Private/BaseParser.cpp
+++ b/Engine/Source/Programs/UnrealHeaderTool/Private/BaseParser.cpp
@@ -123,6 +123,9 @@ namespace
 
 			FMetadataKeyword& DocumentationPolicy = Dictionary.Add(TEXT("DocumentationPolicy"), EMetadataValueArgument::None);
 			DocumentationPolicy.InsertAddAction(TEXT("DocumentationPolicy"), TEXT("Strict"));
+
+			FMetadataKeyword& SparseClassDataType = Dictionary.Add(TEXT("SparseClassDataType"), EMetadataValueArgument::Required);
+			SparseClassDataType.InsertAddAction(TEXT("SparseClassDataType"), TEXT(""));
 		}
 
 		return Dictionary.Find(Keyword);
diff --git a/Engine/Source/Programs/UnrealHeaderTool/Private/ClassDeclarationMetaData.cpp b/Engine/Source/Programs/UnrealHeaderTool/Private/ClassDeclarationMetaData.cpp
--- a/Engine/Source/Programs/UnrealHeaderTool/Private/ClassDeclarationMetaData.cpp
+++ b/Engine/Source/Programs/UnrealHeaderTool/Private/ClassDeclarationMetaData.cpp
@@ -196,6 +196,11 @@ void FClassDeclarationMetaData::ParseClassProperties(const TArray<FPropertySpeci
 				HideFunctions.AddUnique(Value);
 			}
 		}
+		// Currently some code only handles a single sidecar data structure so we enforce that here
+		else if (Specifier == TEXT("SparseClassDataTypes"))
+		{
+			SparseClassDataTypes.AddUnique(FHeaderParser::RequireExactlyOneSpecifierValue(PropSpecifier));
+		}
 		else if (Specifier == TEXT("classGroup"))
 		{
 			FHeaderParser::RequireSpecifierValue(PropSpecifier);
@@ -294,11 +299,13 @@ void FClassDeclarationMetaData::MergeClassCategories(FClass* Class)
 	TArray<FString> ParentHideCategories;
 	TArray<FString> ParentShowSubCatgories;
 	TArray<FString> ParentHideFunctions;
+	TArray<FString> ParentSparseClassDataTypes;
 	TArray<FString> ParentAutoExpandCategories;
 	TArray<FString> ParentAutoCollapseCategories;
 	Class->GetHideCategories(ParentHideCategories);
 	Class->GetShowCategories(ParentShowSubCatgories);
 	Class->GetHideFunctions(ParentHideFunctions);
+	Class->GetSparseClassDataTypes(ParentSparseClassDataTypes);
 	Class->GetAutoExpandCategories(ParentAutoExpandCategories);
 	Class->GetAutoCollapseCategories(ParentAutoCollapseCategories);
 
diff --git a/Engine/Source/Programs/UnrealHeaderTool/Private/ClassDeclarationMetaData.h b/Engine/Source/Programs/UnrealHeaderTool/Private/ClassDeclarationMetaData.h
--- a/Engine/Source/Programs/UnrealHeaderTool/Private/ClassDeclarationMetaData.h
+++ b/Engine/Source/Programs/UnrealHeaderTool/Private/ClassDeclarationMetaData.h
@@ -29,6 +29,7 @@ public:
 	TArray<FString> AutoCollapseCategories;
 	TArray<FString> DependsOn;
 	TArray<FString> ClassGroupNames;
+	TArray<FString> SparseClassDataTypes;
 	
 	/**
 	* Parse Class's properties to generate its declaration data.
diff --git a/Engine/Source/Programs/UnrealHeaderTool/Private/CodeGenerator.cpp b/Engine/Source/Programs/UnrealHeaderTool/Private/CodeGenerator.cpp
--- a/Engine/Source/Programs/UnrealHeaderTool/Private/CodeGenerator.cpp
+++ b/Engine/Source/Programs/UnrealHeaderTool/Private/CodeGenerator.cpp
@@ -2161,6 +2161,14 @@ void FNativeClassHeaderGenerator::ExportNativeGeneratedInitCode(FOutputDevice& O
 		GeneratedClassRegisterFunctionText.Logf(TEXT("\t\t{\r\n"));
 		GeneratedClassRegisterFunctionText.Logf(TEXT("\t\t\tUE4CodeGen_Private::ConstructUClass(OuterClass, %s::ClassParams);\r\n"), *StaticsStructName);
 
+		TArray<FString> SparseClassDataTypes;
+		((FClass*)Class)->GetSparseClassDataTypes(SparseClassDataTypes);
+		ensureMsgf(SparseClassDataTypes.Num() <= 1, TEXT("We don't currently support multiple sparse class data structures on a single object."));
+		for (const FString& SparseClassDataString : SparseClassDataTypes)
+		{
+			GeneratedClassRegisterFunctionText.Logf(TEXT("\t\t\tOuterClass->SetSparseClassDataStruct(F%s::StaticStruct());\r\n"), *SparseClassDataString);
+		}
+
 		if (bIsDynamic)
 		{
 			FString* CustomDynamicClassInitializationMD = MetaDataMap ? MetaDataMap->Find(TEXT("CustomDynamicClassInitialization")) : nullptr;
@@ -4620,6 +4628,8 @@ struct FNativeFunctionStringBuilder
 	FUHTStringBuilder RPCWrappers;
 	FUHTStringBuilder AutogeneratedBlueprintFunctionDeclarations;
 	FUHTStringBuilder AutogeneratedBlueprintFunctionDeclarationsOnlyNotDeclared;
+	FUHTStringBuilder AutogeneratedStaticData;
+	FUHTStringBuilder AutogeneratedStaticDataFuncs;
 };
 
 void FNativeClassHeaderGenerator::ExportNativeFunctions(FOutputDevice& OutGeneratedHeaderText, FOutputDevice& OutMacroCalls, FOutputDevice& OutNoPureDeclsMacroCalls, const FUnrealSourceFile& SourceFile, UClass* Class, FClassMetaData* ClassData)
@@ -4636,6 +4646,81 @@ void FNativeClassHeaderGenerator::ExportNativeFunctions(FOutputDevice& OutGenera
 		ClassRange.Validate();
 	}
 
+	// gather static class data
+	TArray<FString> SparseClassDataTypes;
+	((FClass*)Class)->GetSparseClassDataTypes(SparseClassDataTypes);
+	FString FullClassName = ((FClass*)Class)->GetNameWithPrefix();
+	for (const FString& SparseClassDataString : SparseClassDataTypes)
+	{
+		RuntimeStringBuilders.AutogeneratedStaticData += TEXT("virtual void* CreateSparseClassData() override\r\n");
+		RuntimeStringBuilders.AutogeneratedStaticData += TEXT("{\r\n");
+		RuntimeStringBuilders.AutogeneratedStaticData.Logf(TEXT("\tvoid* SparseClassDataObj = new F%s;\r\n"), *SparseClassDataString);
+		RuntimeStringBuilders.AutogeneratedStaticData.Logf(TEXT("\treturn SparseClassDataObj;\r\n"), *SparseClassDataString);
+		RuntimeStringBuilders.AutogeneratedStaticData += TEXT("}\r\n");
+
+		RuntimeStringBuilders.AutogeneratedStaticData.Logf(TEXT("F%s* Get%s()\r\n"), *SparseClassDataString, *SparseClassDataString);
+		RuntimeStringBuilders.AutogeneratedStaticData += TEXT("{\r\n");
+		RuntimeStringBuilders.AutogeneratedStaticData.Logf(TEXT("\treturn (F%s*)(GetClass()->GetOrCreateSparseClassData());\r\n"), *SparseClassDataString);
+		RuntimeStringBuilders.AutogeneratedStaticData += TEXT("}\r\n");
+
+		RuntimeStringBuilders.AutogeneratedStaticData.Logf(TEXT("F%s* Get%s() const\r\n"), *SparseClassDataString, *SparseClassDataString);
+		RuntimeStringBuilders.AutogeneratedStaticData += TEXT("{\r\n");
+		RuntimeStringBuilders.AutogeneratedStaticData.Logf(TEXT("\treturn (F%s*)(GetClass()->GetOrCreateSparseClassData());\r\n"), *SparseClassDataString);
+		RuntimeStringBuilders.AutogeneratedStaticData += TEXT("}\r\n");
+
+		UScriptStruct* SparseClassDataStruct = FindObjectSafe<UScriptStruct>(ANY_PACKAGE, *SparseClassDataString);
+		while (SparseClassDataStruct != nullptr)
+		{
+			const UProperty* Child = Cast<UProperty>(SparseClassDataStruct->Children);
+			while (Child)
+			{
+				FString ReturnExtendedType;
+				FString VarType = Child->GetCPPType(&ReturnExtendedType, EPropertyExportCPPFlags::CPPF_ArgumentOrReturnValue | EPropertyExportCPPFlags::CPPF_Implementation);
+				if (!ReturnExtendedType.IsEmpty())
+				{
+					VarType.Append(ReturnExtendedType);
+				}
+				FString VarName = Child->GetName();
+				FString CleanVarName = VarName;
+				if (Cast<UBoolProperty>(Child) && VarName.StartsWith(TEXT("b"), ESearchCase::CaseSensitive))
+				{
+					CleanVarName = VarName.RightChop(1);
+				}
+
+				if (!Child->HasMetaData("NoGetter"))
+				{
+					if (Child->HasMetaData("GetByRef"))
+					{
+						RuntimeStringBuilders.AutogeneratedStaticDataFuncs.Logf(TEXT("const %s& Get%s()\r\n"), *VarType, *CleanVarName);
+					}
+					else
+					{
+						RuntimeStringBuilders.AutogeneratedStaticDataFuncs.Logf(TEXT("%s Get%s()\r\n"), *VarType, *CleanVarName);
+					}
+					RuntimeStringBuilders.AutogeneratedStaticDataFuncs.Logf(TEXT("{\r\n"));
+					RuntimeStringBuilders.AutogeneratedStaticDataFuncs.Logf(TEXT("\treturn Get%s()->%s;\r\n"), *SparseClassDataString, *VarName);
+					RuntimeStringBuilders.AutogeneratedStaticDataFuncs.Logf(TEXT("}\r\n"));
+
+					if (Child->HasMetaData("GetByRef"))
+					{
+						RuntimeStringBuilders.AutogeneratedStaticDataFuncs.Logf(TEXT("const %s& Get%s() const\r\n"), *VarType, *CleanVarName);
+					}
+					else
+					{
+						RuntimeStringBuilders.AutogeneratedStaticDataFuncs.Logf(TEXT("%s Get%s() const\r\n"), *VarType, *CleanVarName);
+					}
+					RuntimeStringBuilders.AutogeneratedStaticDataFuncs.Logf(TEXT("{\r\n"));
+					RuntimeStringBuilders.AutogeneratedStaticDataFuncs.Logf(TEXT("\treturn Get%s()->%s;\r\n"), *SparseClassDataString, *VarName);
+					RuntimeStringBuilders.AutogeneratedStaticDataFuncs.Logf(TEXT("}\r\n"));
+				}
+
+				Child = Cast<UProperty>(Child->Next);
+			}
+
+			SparseClassDataStruct = Cast<UScriptStruct>(SparseClassDataStruct->GetSuperStruct());
+		}
+	}
+
 	// export the C++ stubs
 
 	for (UFunction* Function : TFieldRange<UFunction>(Class, EFieldIteratorFlags::ExcludeSuper))
@@ -4737,6 +4822,15 @@ void FNativeClassHeaderGenerator::ExportNativeFunctions(FOutputDevice& OutGenera
 		FuncStringBuilders.RPCWrappers += TEXT("\t}") LINE_TERMINATOR;
 	}
 
+	// static class data
+	{
+		FString MacroName = SourceFile.GetGeneratedMacroName(ClassData, TEXT("_SPARSE_DATA"));
+
+		WriteMacro(OutGeneratedHeaderText, MacroName, RuntimeStringBuilders.AutogeneratedStaticData + RuntimeStringBuilders.AutogeneratedStaticDataFuncs);
+		OutMacroCalls.Logf(TEXT("\t%s\r\n"), *MacroName);
+		OutNoPureDeclsMacroCalls.Logf(TEXT("\t%s\r\n"), *MacroName);
+	}
+
 	// Write runtime wrappers
 	{
 		FString MacroName = SourceFile.GetGeneratedMacroName(ClassData, TEXT("_RPC_WRAPPERS"));
diff --git a/Engine/Source/Programs/UnrealHeaderTool/Private/HeaderParser.cpp b/Engine/Source/Programs/UnrealHeaderTool/Private/HeaderParser.cpp
--- a/Engine/Source/Programs/UnrealHeaderTool/Private/HeaderParser.cpp
+++ b/Engine/Source/Programs/UnrealHeaderTool/Private/HeaderParser.cpp
@@ -5917,6 +5917,7 @@ UClass* FHeaderParser::CompileClassDeclaration(FClasses& AllClasses)
 	if (ClassDeclarationData->AutoCollapseCategories.Num()) { MetaData.Add("AutoCollapseCategories", FString::Join(ClassDeclarationData->AutoCollapseCategories, TEXT(" "))); }
 	if (ClassDeclarationData->HideCategories.Num()) { MetaData.Add("HideCategories", FString::Join(ClassDeclarationData->HideCategories, TEXT(" "))); }
 	if (ClassDeclarationData->ShowSubCatgories.Num()) { MetaData.Add("ShowCategories", FString::Join(ClassDeclarationData->ShowSubCatgories, TEXT(" "))); }
+	if (ClassDeclarationData->SparseClassDataTypes.Num()) { MetaData.Add("SparseClassDataTypes", FString::Join(ClassDeclarationData->SparseClassDataTypes, TEXT(" "))); }
 	if (ClassDeclarationData->HideFunctions.Num()) { MetaData.Add("HideFunctions", FString::Join(ClassDeclarationData->HideFunctions, TEXT(" "))); }
 	if (ClassDeclarationData->AutoExpandCategories.Num()) { MetaData.Add("AutoExpandCategories", FString::Join(ClassDeclarationData->AutoExpandCategories, TEXT(" "))); }
 
@@ -9486,6 +9487,10 @@ void FHeaderParser::ResetClassData()
 		{
 			CurrentClass->SetMetaData(TEXT("ShowCategories"), *SuperClass->GetMetaData("ShowCategories"));
 		}
+		if (SuperClass->HasMetaData(TEXT("SparseClassDataTypes")))
+		{
+			CurrentClass->SetMetaData(TEXT("SparseClassDataTypes"), *SuperClass->GetMetaData("SparseClassDataTypes"));
+		}
 		if (SuperClass->HasMetaData(TEXT("HideFunctions")))
 		{
 			CurrentClass->SetMetaData(TEXT("HideFunctions"), *SuperClass->GetMetaData("HideFunctions"));
diff --git a/Engine/Source/Programs/UnrealHeaderTool/Private/ParserClass.cpp b/Engine/Source/Programs/UnrealHeaderTool/Private/ParserClass.cpp
--- a/Engine/Source/Programs/UnrealHeaderTool/Private/ParserClass.cpp
+++ b/Engine/Source/Programs/UnrealHeaderTool/Private/ParserClass.cpp
@@ -87,6 +87,16 @@ void FClass::GetShowCategories(TArray<FString>& OutShowCategories) const
 	}
 }
 
+void FClass::GetSparseClassDataTypes(TArray<FString>& OutSparseClassDataTypes) const
+{
+	static const FName NAME_SparseClassDataTypes(TEXT("SparseClassDataTypes"));
+	if (HasMetaData(NAME_SparseClassDataTypes))
+	{
+		const FString& SparseClassDataTypes = GetMetaData(NAME_SparseClassDataTypes);
+		SparseClassDataTypes.ParseIntoArray(OutSparseClassDataTypes, TEXT(" "), true);
+	}
+}
+
 bool FClass::IsDynamic(const UField* Field)
 {
 	static const FName NAME_ReplaceConverted(TEXT("ReplaceConverted"));
diff --git a/Engine/Source/Programs/UnrealHeaderTool/Private/ParserClass.h b/Engine/Source/Programs/UnrealHeaderTool/Private/ParserClass.h
--- a/Engine/Source/Programs/UnrealHeaderTool/Private/ParserClass.h
+++ b/Engine/Source/Programs/UnrealHeaderTool/Private/ParserClass.h
@@ -49,6 +49,7 @@ public:
 
 	void GetHideCategories(TArray<FString>& OutHideCategories) const;
 	void GetShowCategories(TArray<FString>& OutShowCategories) const;
+	void GetSparseClassDataTypes(TArray<FString>& OutSparseClassDataTypes) const;
 
 	/** Helper function that checks if the field is a dynamic type (can be constructed post-startup) */
 	static bool IsDynamic(const UField* Field);
diff --git a/Engine/Source/Runtime/CoreUObject/Private/UObject/Class.cpp b/Engine/Source/Runtime/CoreUObject/Private/UObject/Class.cpp
--- a/Engine/Source/Runtime/CoreUObject/Private/UObject/Class.cpp
+++ b/Engine/Source/Runtime/CoreUObject/Private/UObject/Class.cpp
@@ -2995,7 +2995,7 @@ class FRestoreClassInfo: public FRestoreForUObjectOverwrite
 	UClass*			Target;
 	/** Saved ClassWithin **/
 	UClass*			Within;
-	/** Save ClassGeneratedBy */
+	/** Saved ClassGeneratedBy */
 	UObject*		GeneratedBy;
 	/** Saved ClassDefaultObject **/
 	UObject*		DefaultObject;
@@ -3674,6 +3674,7 @@ void UClass::SetSuperStruct(UStruct* NewSuperStruct)
 	UnhashObject(this);
 	ClearFunctionMapsCaches();
 	Super::SetSuperStruct(NewSuperStruct);
+	SetSparseClassDataStruct(GetSparseClassDataArchetypeStruct());
 	HashObject(this);
 }
 
@@ -3974,6 +3975,33 @@ void UClass::SerializeDefaultObject(UObject* Object, FStructuredArchive::FSlot S
 	UnderlyingArchive.StopSerializingDefaults();
 }
 
+void UClass::SerializeSparseClassData(UObject* Object, FStructuredArchive::FSlot Slot)
+{
+	if (!SparseClassDataStruct)
+	{
+		return;
+	}
+
+	// tell the archive that it's allowed to load data for transient properties
+	FArchive& UnderlyingArchive = Slot.GetUnderlyingArchive();
+
+	if (((UnderlyingArchive.IsLoading() || UnderlyingArchive.IsSaving()) && !UnderlyingArchive.WantBinaryPropertySerialization()))
+	{
+		// class default objects do not always have a vtable when saved
+		// so use script serialization as opposed to native serialization to
+		// guarantee that all property data is loaded into the correct location
+		SparseClassDataStruct->SerializeItem(Slot, SparseClassData, GetArchetypeForSparseClassData());
+	}
+	else if (UnderlyingArchive.GetPortFlags() != 0)
+	{
+		SparseClassDataStruct->SerializeBinEx(Slot, (uint8*)SparseClassData, SparseClassDataStruct, GetSparseClassDataArchetypeStruct());
+	}
+	else
+	{
+		SparseClassDataStruct->SerializeBin(Slot, (uint8*)SparseClassData);
+	}
+}
+
 
 FArchive& operator<<(FArchive& Ar, FImplementedInterface& A)
 {
@@ -3984,9 +4012,21 @@ FArchive& operator<<(FArchive& Ar, FImplementedInterface& A)
 	return Ar;
 }
 
+void* UClass::GetArchetypeForSparseClassData() const
+{
+	UClass* SuperClass = GetSuperClass();
+	return SuperClass ? SuperClass->GetOrCreateSparseClassData() : nullptr;
+}
+
+UScriptStruct* UClass::GetSparseClassDataArchetypeStruct() const
+{
+	UClass* SuperClass = GetSuperClass();
+	return SuperClass ? SuperClass->GetSparseClassDataStruct() : nullptr;
+}
+
 UObject* UClass::GetArchetypeForCDO() const
 {
-	auto SuperClass = GetSuperClass();
+	UClass* SuperClass = GetSuperClass();
 	return SuperClass ? SuperClass->GetDefaultObject() : nullptr;
 }
 
@@ -4094,6 +4134,8 @@ UClass::UClass(const FObjectInitializer& ObjectInitializer)
 ,	ClassWithin( UObject::StaticClass() )
 ,	ClassGeneratedBy(nullptr)
 ,	ClassDefaultObject(nullptr)
+,	SparseClassData(nullptr)
+,	SparseClassDataStruct(nullptr)
 {
 	// If you add properties here, please update the other constructors and PurgeClass()
 
@@ -4112,6 +4154,8 @@ UClass::UClass(const FObjectInitializer& ObjectInitializer, UClass* InBaseClass)
 ,	ClassWithin(UObject::StaticClass())
 ,	ClassGeneratedBy(nullptr)
 ,	ClassDefaultObject(nullptr)
+,	SparseClassData(nullptr)
+,	SparseClassDataStruct(nullptr)
 {
 	// If you add properties here, please update the other constructors and PurgeClass()
 
@@ -4167,6 +4211,8 @@ UClass::UClass
 ,	ClassConfigName			()
 ,	NetFields				()
 ,	ClassDefaultObject		( nullptr )
+,	SparseClassData			( nullptr )
+,	SparseClassDataStruct	( nullptr )
 {
 	// If you add properties here, please update the other constructors and PurgeClass()
 
@@ -4179,6 +4225,41 @@ UClass::UClass
 	*(const TCHAR**)&ClassConfigName = InConfigName; //-V580
 }
 
+void* UClass::CreateSparseClassData()
+{
+	check(SparseClassData == nullptr);
+
+	SparseClassData = GetDefaultObject()->CreateSparseClassData();
+
+	if (SparseClassData)
+	{
+		// initialize per class data from the archetype if we have one
+		void* SparseArchetypeData = GetArchetypeForSparseClassData();
+		UStruct* SparseClassDataArchetypeStruct = GetSparseClassDataArchetypeStruct();
+
+		if (SparseArchetypeData)
+		{
+			for (UProperty* P = SparseClassDataArchetypeStruct->PropertyLink; P; P = P->PropertyLinkNext)
+			{
+				P->CopyCompleteValue_InContainer(SparseClassData, SparseArchetypeData);
+			}
+		}
+	}
+
+	return SparseClassData;
+}
+
+UScriptStruct* UClass::GetSparseClassDataStruct() const
+{
+	// this info is specified on the object via code generation so we use it instead of looking at the UClass
+	return SparseClassDataStruct;
+}
+
+void UClass::SetSparseClassDataStruct(UScriptStruct* InSparseClassDataStruct)
+{ 
+	SparseClassDataStruct = InSparseClassDataStruct; 
+}
+
 #if WITH_HOT_RELOAD
 
 bool UClass::HotReloadPrivateStaticClass(
diff --git a/Engine/Source/Runtime/CoreUObject/Private/UObject/Obj.cpp b/Engine/Source/Runtime/CoreUObject/Private/UObject/Obj.cpp
--- a/Engine/Source/Runtime/CoreUObject/Private/UObject/Obj.cpp
+++ b/Engine/Source/Runtime/CoreUObject/Private/UObject/Obj.cpp
@@ -1151,6 +1151,16 @@ void UObject::PostLoadSubobjects( FObjectInstancingGraph* OuterInstanceGraph/*=N
 	}
 }
 
+void* UObject::CreateSparseClassData()
+{
+	return nullptr;
+}
+
+UScriptStruct* UObject::GetSparseClassDataStruct() const
+{
+	UClass* Class = GetClass();
+	return Class ? Class->GetSparseClassDataStruct() : nullptr;
+}
 
 void UObject::ConditionalPostLoadSubobjects( FObjectInstancingGraph* OuterInstanceGraph/*=NULL*/ )
 {
@@ -1405,7 +1415,7 @@ void UObject::SerializeScriptProperties( FStructuredArchive::FSlot Slot ) const
 		FArchive::FScopeAddDebugData S(UnderlyingArchive, ObjClass->GetFName());
 #endif
 
-		ObjClass->SerializeTaggedProperties(Slot, (uint8*)this, HasAnyFlags(RF_ClassDefaultObject) ? ObjClass->GetSuperClass() : ObjClass, (uint8*)DiffObject, bBreakSerializationRecursion ? this : NULL);
+		ObjClass->SerializeTaggedProperties(Slot, (uint8*)this, HasAnyFlags(RF_ClassDefaultObject) ? ObjClass->GetSuperClass() : ObjClass, (uint8*)DiffObject, bBreakSerializationRecursion ? this : nullptr);
 	}
 	else if (UnderlyingArchive.GetPortFlags() != 0 && !UnderlyingArchive.ArUseCustomPropertyList )
 	{
@@ -1422,7 +1432,7 @@ void UObject::SerializeScriptProperties( FStructuredArchive::FSlot Slot ) const
 		ObjClass->SerializeBin(Slot, const_cast<UObject *>(this));
 	}
 
-	if( HasAnyFlags(RF_ClassDefaultObject) )
+	if (HasAnyFlags(RF_ClassDefaultObject))
 	{
 		UnderlyingArchive.StopSerializingDefaults();
 	}
diff --git a/Engine/Source/Runtime/CoreUObject/Private/UObject/ScriptCore.cpp b/Engine/Source/Runtime/CoreUObject/Private/UObject/ScriptCore.cpp
--- a/Engine/Source/Runtime/CoreUObject/Private/UObject/ScriptCore.cpp
+++ b/Engine/Source/Runtime/CoreUObject/Private/UObject/ScriptCore.cpp
@@ -1994,6 +1994,35 @@ DEFINE_FUNCTION(UObject::execInstanceVariable)
 }
 IMPLEMENT_VM_FUNCTION( EX_InstanceVariable, execInstanceVariable );
 
+DEFINE_FUNCTION(UObject::execClassSparseDataVariable)
+{
+	UProperty* VarProperty = (UProperty*)Stack.ReadObject();
+	Stack.MostRecentProperty = VarProperty;
+
+	if (VarProperty == nullptr || P_THIS->GetSparseClassDataStruct() == nullptr)
+	{
+		FBlueprintExceptionInfo ExceptionInfo(EBlueprintExceptionType::AccessViolation, FText::Format(LOCTEXT("MissingProperty", "Attempted to access missing sparse property '{0}' {1}, {2}. If this is a packaged/cooked build, are you attempting to use an editor-only property?"), FText::FromString(GetNameSafe(VarProperty)), FText::FromString(GetNameSafe(P_THIS->GetSparseClassDataStruct())), FText::FromString(GetNameSafe((UClass*)VarProperty->GetOuter()))));
+		FBlueprintCoreDelegates::ThrowScriptException(P_THIS, Stack, ExceptionInfo);
+
+		Stack.MostRecentPropertyAddress = nullptr;
+	}
+	else
+	{
+		void* SparseDataBaseAddress = P_THIS->GetClass()->GetOrCreateSparseClassData();
+		Stack.MostRecentPropertyAddress = VarProperty->ContainerPtrToValuePtr<uint8>(SparseDataBaseAddress);
+
+		// SPARSEDATA_TODO: remove these two lines once we're sure the math is right
+		int32 Offset = VarProperty->GetOffset_ForInternal();
+		check((uint8*)SparseDataBaseAddress + Offset == Stack.MostRecentPropertyAddress);
+
+		if (RESULT_PARAM)
+		{
+			VarProperty->CopyCompleteValueToScriptVM(RESULT_PARAM, Stack.MostRecentPropertyAddress);
+		}
+	}
+}
+IMPLEMENT_VM_FUNCTION(EX_ClassSparseDataVariable, execClassSparseDataVariable);
+
 DEFINE_FUNCTION(UObject::execDefaultVariable)
 {
 	UProperty* VarProperty = (UProperty*)Stack.ReadObject();
diff --git a/Engine/Source/Runtime/CoreUObject/Private/UObject/UObjectGlobals.cpp b/Engine/Source/Runtime/CoreUObject/Private/UObject/UObjectGlobals.cpp
--- a/Engine/Source/Runtime/CoreUObject/Private/UObject/UObjectGlobals.cpp
+++ b/Engine/Source/Runtime/CoreUObject/Private/UObject/UObjectGlobals.cpp
@@ -4574,6 +4574,8 @@ namespace UE4CodeGen_Private
 #endif
 
 		NewClass->StaticLink();
+
+		NewClass->SetSparseClassDataStruct(NewClass->GetSparseClassDataArchetypeStruct());
 	}
 }
 
diff --git a/Engine/Source/Runtime/CoreUObject/Public/UObject/Class.h b/Engine/Source/Runtime/CoreUObject/Public/UObject/Class.h
--- a/Engine/Source/Runtime/CoreUObject/Public/UObject/Class.h
+++ b/Engine/Source/Runtime/CoreUObject/Public/UObject/Class.h
@@ -2233,6 +2233,7 @@ class COREUOBJECT_API UClass : public UStruct
 
 public:
 	friend class FRestoreClassInfo;
+	friend class FBlueprintEditorUtils;
 
 	typedef void		(*ClassConstructorType)				(const FObjectInitializer&);
 	typedef UObject*	(*ClassVTableHelperCtorCallerType)	(FVTableHelper& Helper);
@@ -2309,6 +2310,28 @@ public:
 	/** The class default object; used for delta serialization and object initialization */
 	UObject* ClassDefaultObject;
 
+protected:
+	/** This is where we store the data that is only changed per class instead of per instance */
+	UPROPERTY()
+	void* SparseClassData;
+
+	/** The struct used to store sparse class data. */
+	UPROPERTY()
+	UScriptStruct* SparseClassDataStruct;
+
+public:
+	/**
+	 * Returns a pointer to the sidecar data structure. This function will create an instance of the data structure if one has been specified and it has not yet been created.
+	 */
+	void* GetOrCreateSparseClassData() { return SparseClassData ? SparseClassData : CreateSparseClassData(); }
+
+	/**
+	 * Returns a pointer to the type of the sidecar data structure if one is specified.
+	 */
+	virtual UScriptStruct* GetSparseClassDataStruct() const;
+
+	void SetSparseClassDataStruct(UScriptStruct* InSparseClassDataStruct);
+
 	/** Assemble reference token streams for all classes if they haven't had it assembled already */
 	static void AssembleReferenceTokenStreams();
 
@@ -2320,6 +2343,8 @@ public:
 #endif // WITH_EDITOR
 
 private:
+	void* CreateSparseClassData();
+
 #if WITH_EDITOR
 	/** Provides access to attributes of the underlying C++ class. Should never be unset. */
 	TOptional<FCppClassTypeInfo> CppTypeInfo;
@@ -2731,7 +2756,7 @@ public:
 	 */
 	bool ImplementsInterface(const class UClass* SomeInterface) const;
 
-	/** serializes the passed in object as this class's default object using the given archive
+	/** serializes the passed in object as this class's default object using the given archive slot
 	 * @param Object the object to serialize as default
 	 * @param Slot the structured archive slot to serialize from
 	 */
@@ -2746,6 +2771,12 @@ public:
 		SerializeDefaultObject(Object, FStructuredArchiveFromArchive(Ar).GetSlot());
 	}
 
+	/** serializes the associated sparse class data for the passed in object using the given archive slot. This should only be called if the class has an associated sparse data structure.
+	 * @param Object the object to serialize as default
+	 * @param Slot the structured archive slot to serialize from
+	 */
+	void SerializeSparseClassData(UObject* Object, FStructuredArchive::FSlot Slot);
+
 	/** Wraps the PostLoad() call for the class default object.
 	 * @param Object the default object to call PostLoad() on
 	 */
@@ -2798,6 +2829,12 @@ public:
 	/** Returns archetype object for CDO */
 	virtual UObject* GetArchetypeForCDO() const;
 
+	/** Returns archetype for sparse class data */
+	virtual void* GetArchetypeForSparseClassData() const;
+
+	/** Returns the struct used by the sparse class data archetype */
+	UScriptStruct* GetSparseClassDataArchetypeStruct() const;
+
 	/**
 	* Returns all objects that should be preloaded before the class default object is serialized at load time. Only used by the EDL.
 	*
diff --git a/Engine/Source/Runtime/CoreUObject/Public/UObject/Object.h b/Engine/Source/Runtime/CoreUObject/Public/UObject/Object.h
--- a/Engine/Source/Runtime/CoreUObject/Public/UObject/Object.h
+++ b/Engine/Source/Runtime/CoreUObject/Public/UObject/Object.h
@@ -541,6 +541,16 @@ public:
 	/** Return a one line description of an object for viewing in the thumbnail view of the generic browser */
 	virtual FString GetDesc() { return TEXT( "" ); }
 
+	/** Allocates the sidecar data structure that stores data that is constant for all instances of this class. */
+	virtual void* CreateSparseClassData();
+
+	/** Return the UStruct corresponding to the sidecar data structure that stores data that is constant for all instances of this class. */
+	virtual UScriptStruct* GetSparseClassDataStruct() const;
+
+#if WITH_EDITOR
+	virtual void MoveDataToSparseClassDataStruct() const {}
+#endif
+
 #if WITH_ENGINE
 	/** 
 	 * Returns what UWorld this object is contained within. 
@@ -1243,6 +1253,7 @@ public:
 	DECLARE_FUNCTION(execDefaultVariable);
 	DECLARE_FUNCTION(execLocalOutVariable);
 	DECLARE_FUNCTION(execInterfaceVariable);
+	DECLARE_FUNCTION(execClassSparseDataVariable);
 	DECLARE_FUNCTION(execInterfaceContext);
 	DECLARE_FUNCTION(execArrayElement);
 	DECLARE_FUNCTION(execBoolVariable);
diff --git a/Engine/Source/Runtime/CoreUObject/Public/UObject/ObjectMacros.h b/Engine/Source/Runtime/CoreUObject/Public/UObject/ObjectMacros.h
--- a/Engine/Source/Runtime/CoreUObject/Public/UObject/ObjectMacros.h
+++ b/Engine/Source/Runtime/CoreUObject/Public/UObject/ObjectMacros.h
@@ -725,6 +725,9 @@ namespace UC
 
 		/// Marks this class as an 'early access' preview (while not considered production-ready, it's a step beyond 'experimental' and is being provided as a preview of things to come)
 		EarlyAccessPreview,
+
+		// Some properties are stored once per class in a sidecar structure and not on instances of the class
+		SparseClassDataType,
 	};
 }
 
diff --git a/Engine/Source/Runtime/CoreUObject/Public/UObject/Script.h b/Engine/Source/Runtime/CoreUObject/Public/UObject/Script.h
--- a/Engine/Source/Runtime/CoreUObject/Public/UObject/Script.h
+++ b/Engine/Source/Runtime/CoreUObject/Public/UObject/Script.h
@@ -270,6 +270,7 @@ enum EExprToken
 	EX_SwitchValue			= 0x69,
 	EX_InstrumentationEvent	= 0x6A, // Instrumentation event
 	EX_ArrayGetByRef		= 0x6B,
+	EX_ClassSparseDataVariable = 0x6C, // Sparse data variable
 	EX_Max					= 0x100,
 };
 
diff --git a/Engine/Source/Runtime/CoreUObject/Public/UObject/ScriptSerialization.h b/Engine/Source/Runtime/CoreUObject/Public/UObject/ScriptSerialization.h
--- a/Engine/Source/Runtime/CoreUObject/Public/UObject/ScriptSerialization.h
+++ b/Engine/Source/Runtime/CoreUObject/Public/UObject/ScriptSerialization.h
@@ -215,6 +215,7 @@
 		case EX_InstanceVariable:
 		case EX_DefaultVariable:
 		case EX_LocalOutVariable:
+		case EX_ClassSparseDataVariable:
 		{
 			XFER_PROP_POINTER;
 			break;
diff --git a/Engine/Source/Runtime/Engine/Classes/Engine/BlueprintGeneratedClass.h b/Engine/Source/Runtime/Engine/Classes/Engine/BlueprintGeneratedClass.h
--- a/Engine/Source/Runtime/Engine/Classes/Engine/BlueprintGeneratedClass.h
+++ b/Engine/Source/Runtime/Engine/Classes/Engine/BlueprintGeneratedClass.h
@@ -640,6 +640,12 @@ public:
 	UPROPERTY()
 	uint8 bHasCookedComponentInstancingData:1;
 
+#if WITH_EDITORONLY_DATA
+	/** Used to check if this class has sparse data that can be serialized. This will be false when loading the data if it hasn't already been saved out. */
+	UPROPERTY()
+	uint32 bIsSparseClassDataSerializable : 1;
+#endif // WITH_EDITORONLY_DATA
+
 private:
 	/** Flag to make sure the custom property list has been initialized */
 	uint8 bCustomPropertyListForPostConstructionInitialized:1;
diff --git a/Engine/Source/Runtime/Engine/Classes/Engine/MemberReference.h b/Engine/Source/Runtime/Engine/Classes/Engine/MemberReference.h
--- a/Engine/Source/Runtime/Engine/Classes/Engine/MemberReference.h
+++ b/Engine/Source/Runtime/Engine/Classes/Engine/MemberReference.h
@@ -245,6 +245,9 @@ public:
 		return !MemberScope.IsEmpty();
 	}
 
+	/** Returns true if this is in the sparse data struct for OwningClass */
+	ENGINE_API bool IsSparseClassData(const UClass* OwningClass) const;
+
 #if WITH_EDITOR
 	/**
 	 * Returns a search string to submit to Find-in-Blueprints to find references to this reference
@@ -383,8 +386,16 @@ public:
 #if WITH_EDITOR
 				TargetScope = GetClassToUse(TargetScope, bUseUpToDateClass);
 #endif
-				// Find in target scope
-				ReturnField = FindField<TFieldType>(TargetScope, MemberName);
+				// Find in target scope or in the sparse class data
+				UScriptStruct* SparseClassDataStruct = TargetScope->GetSparseClassDataStruct();
+				if (SparseClassDataStruct)
+				{
+					ReturnField = FindField<TFieldType>(SparseClassDataStruct, MemberName);
+				}
+				if (ReturnField == nullptr)
+				{
+					ReturnField = FindField<TFieldType>(TargetScope, MemberName);
+				}
 
 #if WITH_EDITOR
 
diff --git a/Engine/Source/Runtime/Engine/Private/BlueprintGeneratedClass.cpp b/Engine/Source/Runtime/Engine/Private/BlueprintGeneratedClass.cpp
--- a/Engine/Source/Runtime/Engine/Private/BlueprintGeneratedClass.cpp
+++ b/Engine/Source/Runtime/Engine/Private/BlueprintGeneratedClass.cpp
@@ -57,6 +57,9 @@ UBlueprintGeneratedClass::UBlueprintGeneratedClass(const FObjectInitializer& Obj
 	bHasNativizedParent = false;
 	bHasCookedComponentInstancingData = false;
 	bCustomPropertyListForPostConstructionInitialized = false;
+#if WITH_EDITORONLY_DATA
+	bIsSparseClassDataSerializable = false;
+#endif
 }
 
 void UBlueprintGeneratedClass::PostInitProperties()
@@ -404,6 +407,16 @@ void UBlueprintGeneratedClass::SerializeDefaultObject(UObject* Object, FStructur
 			CheckAndApplyComponentTemplateOverrides(ClassDefaultObject);
 		}
 	}
+
+#if WITH_EDITORONLY_DATA
+	if (bIsSparseClassDataSerializable)
+#endif
+	{
+		if (Object->GetSparseClassDataStruct())
+		{
+			SerializeSparseClassData(Object, FStructuredArchiveFromArchive(UnderlyingArchive).GetSlot());
+		}
+	}
 }
 
 void UBlueprintGeneratedClass::PostLoadDefaultObject(UObject* Object)
@@ -423,6 +436,18 @@ void UBlueprintGeneratedClass::PostLoadDefaultObject(UObject* Object)
 			ClassDefaultObject->LoadConfig();
 		}
 	}
+
+#if WITH_EDITOR
+#if WITH_EDITORONLY_DATA
+	Object->MoveDataToSparseClassDataStruct();
+
+	if (Object->GetSparseClassDataStruct())
+	{
+		// now that any data has been moved into the sparse data structure we can safely serialize it
+		bIsSparseClassDataSerializable = true;
+	}
+#endif
+#endif
 }
 
 bool UBlueprintGeneratedClass::BuildCustomPropertyListForPostConstruction(FCustomPropertyListNode*& InPropertyList, UStruct* InStruct, const uint8* DataPtr, const uint8* DefaultDataPtr)
diff --git a/Engine/Source/Runtime/Engine/Private/MemberReference.cpp b/Engine/Source/Runtime/Engine/Private/MemberReference.cpp
--- a/Engine/Source/Runtime/Engine/Private/MemberReference.cpp
+++ b/Engine/Source/Runtime/Engine/Private/MemberReference.cpp
@@ -118,6 +118,19 @@ void FMemberReference::InvalidateScope()
 	}
 }
 
+bool FMemberReference::IsSparseClassData(const UClass* OwningClass) const
+{
+	bool bIsSparseClassData = false;
+	UScriptStruct* SparseClassDataStruct = OwningClass ? OwningClass->GetSparseClassDataStruct() : nullptr;
+	if (SparseClassDataStruct)
+	{
+		UProperty* VariableProperty = FindField<UProperty>(SparseClassDataStruct, GetMemberName());
+		bIsSparseClassData = VariableProperty != nullptr;
+	}
+
+	return bIsSparseClassData;
+}
+
 #if WITH_EDITOR
 
 FString FMemberReference::GetReferenceSearchString(UClass* InFieldOwner) const
