Fixing the build.

diff --git a/Engine/Source/Editor/PropertyEditor/Private/ObjectPropertyNode.cpp b/Engine/Source/Editor/PropertyEditor/Private/ObjectPropertyNode.cpp
--- a/Engine/Source/Editor/PropertyEditor/Private/ObjectPropertyNode.cpp
+++ b/Engine/Source/Editor/PropertyEditor/Private/ObjectPropertyNode.cpp
@@ -480,7 +480,7 @@ void FObjectPropertyNode::GetCategoryProperties(const TSet<UClass*>& ClassesToCo
 
 	if (bMetaDataAllowVisible)
 	{
-		if( PropertyEditorHelpers::ShouldBeVisible(*this, *It) )
+		if (PropertyEditorHelpers::ShouldBeVisible(*this, CurrentProperty))
 		{
 			if (!CategoriesFromBlueprints.Contains(CategoryName) && !CategoriesFromProperties.Contains(CategoryName))
 			{
