Fixing build errors.

diff --git a/Engine/Source/Editor/PropertyEditor/Private/StructurePropertyNode.h b/Engine/Source/Editor/PropertyEditor/Private/StructurePropertyNode.h
--- a/Engine/Source/Editor/PropertyEditor/Private/StructurePropertyNode.h
+++ b/Engine/Source/Editor/PropertyEditor/Private/StructurePropertyNode.h
@@ -58,7 +58,7 @@ public:
 
 		uint8* ReadAddress = StructData->GetStructMemory();
 		check(ReadAddress);
-		OutAddresses.Add(nullptr, InPropertyNode.GetValueBaseAddress(ReadAddress, InPropertyNode.HasNodeFlags(EPropertyNodeFlags::IsSparseClassData)), true);
+		OutAddresses.Add(nullptr, InPropertyNode.GetValueBaseAddress(ReadAddress, InPropertyNode.HasNodeFlags(EPropertyNodeFlags::IsSparseClassData) != 0), true);
 		return true;
 	}
 
