Fix 2019 compile errors as a result of the very badly HasNodeFlags
#rnx

diff --git a/Engine/Source/Editor/PropertyEditor/Private/EditConditionContext.cpp b/Engine/Source/Editor/PropertyEditor/Private/EditConditionContext.cpp
--- a/Engine/Source/Editor/PropertyEditor/Private/EditConditionContext.cpp
+++ b/Engine/Source/Editor/PropertyEditor/Private/EditConditionContext.cpp
@@ -104,7 +104,7 @@ TOptional<bool> FEditConditionContext::GetBoolValue(const FString& PropertyName)
 			return TOptional<bool>();
 		}
 
-		uint8* ParentPtr = ParentNode->GetValueAddress(BasePtr, PinnedNode->HasNodeFlags(EPropertyNodeFlags::IsSparseClassData));
+		uint8* ParentPtr = ParentNode->GetValueAddress(BasePtr, PinnedNode->HasNodeFlags(EPropertyNodeFlags::IsSparseClassData) != 0);
 
 		uint8* ValuePtr = ComplexParentNode->GetValuePtrOfInstance(Index, BoolProperty, ParentNode);
 
@@ -151,7 +151,7 @@ TOptional<double> FEditConditionContext::GetNumericValue(const FString& Property
 			return TOptional<double>();
 		}
 
-		uint8* ParentPtr = ParentNode->GetValueAddress(BasePtr, PinnedNode->HasNodeFlags(EPropertyNodeFlags::IsSparseClassData));
+		uint8* ParentPtr = ParentNode->GetValueAddress(BasePtr, PinnedNode->HasNodeFlags(EPropertyNodeFlags::IsSparseClassData) != 0);
 
 		uint8* ValuePtr = ComplexParentNode->GetValuePtrOfInstance(Index, NumericProperty, ParentNode);
 
@@ -225,7 +225,7 @@ TOptional<FString> FEditConditionContext::GetEnumValue(const FString& PropertyNa
 			return TOptional<FString>();
 		}
 
-		uint8* ParentPtr = ParentNode->GetValueAddress(BasePtr, PinnedNode->HasNodeFlags(EPropertyNodeFlags::IsSparseClassData));
+		uint8* ParentPtr = ParentNode->GetValueAddress(BasePtr, PinnedNode->HasNodeFlags(EPropertyNodeFlags::IsSparseClassData) != 0);
 
 		uint8* ValuePtr = ComplexParentNode->GetValuePtrOfInstance(Index, Property, ParentNode);
 
diff --git a/Engine/Source/Editor/PropertyEditor/Private/Presentation/PropertyEditor/PropertyEditor.cpp b/Engine/Source/Editor/PropertyEditor/Private/Presentation/PropertyEditor/PropertyEditor.cpp
--- a/Engine/Source/Editor/PropertyEditor/Private/Presentation/PropertyEditor/PropertyEditor.cpp
+++ b/Engine/Source/Editor/PropertyEditor/Private/Presentation/PropertyEditor/PropertyEditor.cpp
@@ -469,7 +469,7 @@ void FPropertyEditor::ToggleEditConditionState()
 		// ParentNode can point to a struct inside that object (which is stored as an FItemPropertyNode)
 		// We need all three pointers to get the value pointer
 		uint8* BaseAddress = ComplexParentNode->GetMemoryOfInstance(Index);
-		uint8* ParentOffset = ParentNode->GetValueAddress(BaseAddress, PropertyNode->HasNodeFlags(EPropertyNodeFlags::IsSparseClassData));
+		uint8* ParentOffset = ParentNode->GetValueAddress(BaseAddress, PropertyNode->HasNodeFlags(EPropertyNodeFlags::IsSparseClassData) != 0);
 
 		uint8* ValuePtr = ComplexParentNode->GetValuePtrOfInstance(Index, EditConditionProperty, ParentNode);
 
@@ -496,7 +496,7 @@ void FPropertyEditor::ToggleEditConditionState()
 				for (int32 InstanceIndex = 0; InstanceIndex < ArchetypeInstances.Num(); ++InstanceIndex)
 				{
 					uint8* ArchetypeBaseOffset = ComplexParentNode->GetValueAddressFromObject(ArchetypeInstances[InstanceIndex]);
-					uint8* ArchetypeParentOffset = ParentNode->GetValueAddress(ArchetypeBaseOffset, PropertyNode->HasNodeFlags(EPropertyNodeFlags::IsSparseClassData));
+					uint8* ArchetypeParentOffset = ParentNode->GetValueAddress(ArchetypeBaseOffset, PropertyNode->HasNodeFlags(EPropertyNodeFlags::IsSparseClassData) != 0);
 					uint8* ArchetypeValueAddr = EditConditionProperty->ContainerPtrToValuePtr<uint8>(ArchetypeParentOffset);
 
 					// Only propagate if the current value on the instance matches the previous value on the template.
diff --git a/Engine/Source/Editor/PropertyEditor/Private/PropertyHandleImpl.cpp b/Engine/Source/Editor/PropertyEditor/Private/PropertyHandleImpl.cpp
--- a/Engine/Source/Editor/PropertyEditor/Private/PropertyHandleImpl.cpp
+++ b/Engine/Source/Editor/PropertyEditor/Private/PropertyHandleImpl.cpp
@@ -2683,7 +2683,7 @@ uint8* FPropertyHandleBase::GetValueBaseAddress(uint8* Base)
 	TSharedPtr<FPropertyNode> PropertyNode = Implementation->GetPropertyNode();
 	if (PropertyNode.IsValid())
 	{
-		return PropertyNode->GetValueBaseAddress(Base, PropertyNode->HasNodeFlags(EPropertyNodeFlags::IsSparseClassData));
+		return PropertyNode->GetValueBaseAddress(Base, PropertyNode->HasNodeFlags(EPropertyNodeFlags::IsSparseClassData) != 0);
 	}
 
 	return nullptr;
diff --git a/Engine/Source/Editor/PropertyEditor/Private/PropertyNode.cpp b/Engine/Source/Editor/PropertyEditor/Private/PropertyNode.cpp
--- a/Engine/Source/Editor/PropertyEditor/Private/PropertyNode.cpp
+++ b/Engine/Source/Editor/PropertyEditor/Private/PropertyNode.cpp
@@ -933,12 +933,12 @@ uint8* FPropertyNode::GetStartAddress(const UObject* Obj) const
 
 uint8* FPropertyNode::GetValueBaseAddressFromObject(const UObject* Obj)
 {
-	return GetValueBaseAddress(GetStartAddress(Obj), HasNodeFlags(EPropertyNodeFlags::IsSparseClassData));
+	return GetValueBaseAddress(GetStartAddress(Obj), HasNodeFlags(EPropertyNodeFlags::IsSparseClassData) != 0);
 }
 
 uint8* FPropertyNode::GetValueAddressFromObject(const UObject* Obj)
 {
-	return GetValueAddress(GetStartAddress(Obj), HasNodeFlags(EPropertyNodeFlags::IsSparseClassData));
+	return GetValueAddress(GetStartAddress(Obj), HasNodeFlags(EPropertyNodeFlags::IsSparseClassData) != 0);
 }
 
 
