Fix an obvious copy paste error that resulted in container properties not being handled correctly when multiselecting objects.

diff --git a/Engine/Source/Editor/PropertyEditor/Private/ObjectPropertyNode.cpp b/Engine/Source/Editor/PropertyEditor/Private/ObjectPropertyNode.cpp
--- a/Engine/Source/Editor/PropertyEditor/Private/ObjectPropertyNode.cpp
+++ b/Engine/Source/Editor/PropertyEditor/Private/ObjectPropertyNode.cpp
@@ -281,7 +281,7 @@ bool FObjectPropertyNode::GetReadAddressUncached(FPropertyNode& InNode,
 					for (int32 ObjIndex = 1; ObjIndex < GetNumObjects(); ObjIndex++)
 					{
 						TempObject = GetUObject(ObjIndex);
-						if (TempObject && Num != FScriptArrayHelper::Num(BaseAddr))
+						if (TempObject && Num != FScriptArrayHelper::Num(InNode.GetValueBaseAddressFromObject(TempObject)))
 						{
 							bAllTheSame = false;
 						}
@@ -293,7 +293,7 @@ bool FObjectPropertyNode::GetReadAddressUncached(FPropertyNode& InNode,
 					for (int32 ObjIndex = 1; ObjIndex < GetNumObjects(); ++ObjIndex)
 					{
 						TempObject = GetUObject(ObjIndex);
-						if (TempObject && Num != FScriptSetHelper::Num(BaseAddr))
+						if (TempObject && Num != FScriptSetHelper::Num(InNode.GetValueBaseAddressFromObject(TempObject)))
 						{
 							bAllTheSame = false;
 						}
@@ -305,7 +305,7 @@ bool FObjectPropertyNode::GetReadAddressUncached(FPropertyNode& InNode,
 					for (int32 ObjIndex = 1; ObjIndex < GetNumObjects(); ++ObjIndex)
 					{
 						TempObject = GetUObject(ObjIndex);
-						if (TempObject && Num != FScriptMapHelper::Num(BaseAddr))
+						if (TempObject && Num != FScriptMapHelper::Num(InNode.GetValueBaseAddressFromObject(TempObject)))
 						{
 							bAllTheSame = false;
 						}
