Fixed .IsEnabled for single column widgets (ie buttons).
-The enabled state was being stomped by IsPropertyEditingEnabled.

diff --git a/Engine/Source/Editor/PropertyEditor/Private/SDetailSingleItemRow.cpp b/Engine/Source/Editor/PropertyEditor/Private/SDetailSingleItemRow.cpp
--- a/Engine/Source/Editor/PropertyEditor/Private/SDetailSingleItemRow.cpp
+++ b/Engine/Source/Editor/PropertyEditor/Private/SDetailSingleItemRow.cpp
@@ -438,7 +438,7 @@ void SDetailSingleItemRow::Construct( const FArguments& InArgs, FDetailLayoutCus
 			}
 			else
 			{
-				Row.WholeRowWidget.Widget->SetEnabled(IsPropertyEditingEnabled);
+				InternalLeftColumnRowBox->SetEnabled(IsPropertyEditingEnabled);
 				InternalLeftColumnRowBox->AddSlot()
 					.HAlign(Row.WholeRowWidget.HorizontalAlignment)
 					.VAlign(Row.WholeRowWidget.VerticalAlignment)
