Porting fix from Dev-Editor that broke nested field editing.

diff --git a/Engine/Source/Editor/PropertyEditor/Private/EditConditionContext.cpp b/Engine/Source/Editor/PropertyEditor/Private/EditConditionContext.cpp
--- a/Engine/Source/Editor/PropertyEditor/Private/EditConditionContext.cpp
+++ b/Engine/Source/Editor/PropertyEditor/Private/EditConditionContext.cpp
@@ -12,32 +12,13 @@ FEditConditionContext::FEditConditionContext(FPropertyNode& InPropertyNode)
 {
 	PropertyNode = InPropertyNode.AsShared();
 
-	FComplexPropertyNode* ComplexParentNode = FindComplexParent();
+	FComplexPropertyNode* ComplexParentNode = InPropertyNode.FindComplexParent();
 	check(ComplexParentNode);
 
 	const UProperty* Property = InPropertyNode.GetProperty();
 	check(Property);
 }
 
-FComplexPropertyNode* FEditConditionContext::FindComplexParent() const
-{
-	if (!PropertyNode.IsValid())
-	{
-		return nullptr;
-	}
-
-	TSharedPtr<FPropertyNode> PinnedNode = PropertyNode.Pin();
-
-	FPropertyNode* SearchStart = PinnedNode.Get();
-	if (PropertyEditorHelpers::IsStaticArray(*SearchStart))
-	{
-		//in the case of conditional static arrays, we have to go up one more level to get the proper parent struct.
-		SearchStart = SearchStart->GetParentNode();
-	}
-
-	return SearchStart->FindComplexParent();
-}
-
 const UBoolProperty* FEditConditionContext::GetSingleBoolProperty(const TSharedPtr<FEditConditionExpression>& Expression) const
 {
 	if (!PropertyNode.IsValid())
@@ -109,14 +90,22 @@ TOptional<bool> FEditConditionContext::GetBoolValue(const FString& PropertyName)
 	}
 
 	TSharedPtr<FPropertyNode> PinnedNode = PropertyNode.Pin();
+	FPropertyNode* ParentNode = PinnedNode->GetParentNode();
 	
 	TOptional<bool> Result;
 
-	FComplexPropertyNode* ComplexParentNode = FindComplexParent();
+	FComplexPropertyNode* ComplexParentNode = PinnedNode->FindComplexParent();
 	for (int32 Index = 0; Index < ComplexParentNode->GetInstancesNum(); ++Index)
 	{
 		uint8* BasePtr = ComplexParentNode->GetMemoryOfInstance(Index);
-		uint8* ValuePtr = BoolProperty->ContainerPtrToValuePtr<uint8>(BasePtr);
+		if (BasePtr == nullptr)
+		{
+			// deleted or not fully loaded yet
+			return TOptional<bool>();
+		}
+
+		uint8* ParentPtr = ParentNode->GetValueAddress(BasePtr);
+		uint8* ValuePtr = BoolProperty->ContainerPtrToValuePtr<uint8>(ParentPtr);
 
 		bool bValue = BoolProperty->GetPropertyValue(ValuePtr);
 		if (!Result.IsSet())
@@ -142,14 +131,22 @@ TOptional<double> FEditConditionContext::GetNumericValue(const FString& Property
 	}
 
 	TSharedPtr<FPropertyNode> PinnedNode = PropertyNode.Pin();
+	FPropertyNode* ParentNode = PinnedNode->GetParentNode();
 
 	TOptional<double> Result;
 
-	FComplexPropertyNode* ComplexParentNode = FindComplexParent();
+	FComplexPropertyNode* ComplexParentNode = PinnedNode->FindComplexParent();
 	for (int32 Index = 0; Index < ComplexParentNode->GetInstancesNum(); ++Index)
 	{
 		uint8* BasePtr = ComplexParentNode->GetMemoryOfInstance(Index);
-		uint8* ValuePtr = NumericProperty->ContainerPtrToValuePtr<uint8>(BasePtr);
+		if (BasePtr == nullptr)
+		{
+			// deleted or not fully loaded yet
+			return TOptional<double>();
+		}
+
+		uint8* ParentPtr = ParentNode->GetValueAddress(BasePtr);
+		uint8* ValuePtr = NumericProperty->ContainerPtrToValuePtr<uint8>(ParentPtr);
 
 		double Value = 0;
 
@@ -202,14 +199,22 @@ TOptional<FString> FEditConditionContext::GetEnumValue(const FString& PropertyNa
 	}
 	
 	TSharedPtr<FPropertyNode> PinnedNode = PropertyNode.Pin();
-
+	FPropertyNode* ParentNode = PinnedNode->GetParentNode();
+	
 	TOptional<int64> Result;
 
-	FComplexPropertyNode* ComplexParentNode = FindComplexParent();
+	FComplexPropertyNode* ComplexParentNode = PinnedNode->FindComplexParent();
 	for (int32 Index = 0; Index < ComplexParentNode->GetInstancesNum(); ++Index)
 	{
 		uint8* BasePtr = ComplexParentNode->GetMemoryOfInstance(Index);
-		uint8* ValuePtr = Property->ContainerPtrToValuePtr<uint8>(BasePtr);
+		if (BasePtr == nullptr)
+		{
+			// deleted or not fully loaded yet
+			return TOptional<FString>();
+		}
+
+		uint8* ParentPtr = ParentNode->GetValueAddress(BasePtr);
+		uint8* ValuePtr = Property->ContainerPtrToValuePtr<uint8>(ParentPtr);
 
 		int64 Value = NumericProperty->GetSignedIntPropertyValue(ValuePtr);
 		if (!Result.IsSet())
diff --git a/Engine/Source/Editor/PropertyEditor/Private/EditConditionContext.h b/Engine/Source/Editor/PropertyEditor/Private/EditConditionContext.h
--- a/Engine/Source/Editor/PropertyEditor/Private/EditConditionContext.h
+++ b/Engine/Source/Editor/PropertyEditor/Private/EditConditionContext.h
@@ -37,6 +37,4 @@ public:
 
 private:
 	TWeakPtr<FPropertyNode> PropertyNode;
-
-	FComplexPropertyNode* FindComplexParent() const;
 };
\ No newline at end of file
diff --git a/Engine/Source/Editor/PropertyEditor/Private/EditConditionParserTests.cpp b/Engine/Source/Editor/PropertyEditor/Private/EditConditionParserTests.cpp
--- a/Engine/Source/Editor/PropertyEditor/Private/EditConditionParserTests.cpp
+++ b/Engine/Source/Editor/PropertyEditor/Private/EditConditionParserTests.cpp
@@ -404,4 +404,37 @@ bool FEditConditionParser_EvaluateUObject::RunTest(const FString& Parameters)
 	return bAllResults;
 }
 
+IMPLEMENT_SIMPLE_AUTOMATION_TEST(FEditConditionParser_SingleBool, "EditConditionParser.SingleBool", EAutomationTestFlags::EditorContext | EAutomationTestFlags::EngineFilter)
+bool FEditConditionParser_SingleBool::RunTest(const FString& Parameters)
+{
+	UEditConditionTestObject* TestObject = NewObject<UEditConditionTestObject>();
+	TestObject->AddToRoot();
+
+	TSharedPtr<FObjectPropertyNode> ObjectNode(new FObjectPropertyNode);
+	ObjectNode->AddObject(TestObject);
+
+	FPropertyNodeInitParams InitParams;
+	ObjectNode->InitNode(InitParams);
+
+	static const FName BoolPropertyName(TEXT("BoolProperty"));
+	TSharedPtr<FPropertyNode> PropertyNode = ObjectNode->FindChildPropertyNode(BoolPropertyName, true);
+	FEditConditionContext Context(*PropertyNode.Get());
+
+	FEditConditionParser Parser;
+
+	{
+		TSharedPtr<FEditConditionExpression> Expression = Parser.Parse(FString(TEXT("BoolProperty")));
+		const UBoolProperty* Property = Context.GetSingleBoolProperty(Expression);
+		TestNotNull(TEXT("Bool"), Property);
+	}
+
+	{
+		TSharedPtr<FEditConditionExpression> Expression = Parser.Parse(FString(TEXT("UintBitfieldProperty")));
+		const UBoolProperty* Property = Context.GetSingleBoolProperty(Expression);
+		TestNotNull(TEXT("Uint Bitfield"), Property);
+	}
+
+	return true;
+}
+
 #endif // WITH_DEV_AUTOMATION_TESTS
\ No newline at end of file
diff --git a/Engine/Source/Editor/PropertyEditor/Private/EditConditionParserTests.h b/Engine/Source/Editor/PropertyEditor/Private/EditConditionParserTests.h
--- a/Engine/Source/Editor/PropertyEditor/Private/EditConditionParserTests.h
+++ b/Engine/Source/Editor/PropertyEditor/Private/EditConditionParserTests.h
@@ -38,4 +38,7 @@ class UEditConditionTestObject : public UObject
 
 	UPROPERTY(EditAnywhere, Category=Test)
 	int32 IntegerProperty;
+
+	UPROPERTY(EditAnywhere, Category=Test)
+	uint8 UintBitfieldProperty : 1;
 };
