Merging using //UE4/Dev-Framework_to_//UE4/Release-4.24

Fixes bugs with data inside container classes that are inside of sparse class data. When we're inside of a container class the pointer offsets for the sparse data structure have already been taken into account and trying to account for it again will result in bad pointers.

diff --git a/Engine/Source/Editor/PropertyEditor/Private/PropertyHandleImpl.cpp b/Engine/Source/Editor/PropertyEditor/Private/PropertyHandleImpl.cpp
--- a/Engine/Source/Editor/PropertyEditor/Private/PropertyHandleImpl.cpp
+++ b/Engine/Source/Editor/PropertyEditor/Private/PropertyHandleImpl.cpp
@@ -42,16 +42,39 @@ void PropertyToTextHelper(FString& OutString, FPropertyNode* InPropertyNode, UPr
 
 void PropertyToTextHelper(FString& OutString, FPropertyNode* InPropertyNode, UProperty* Property, const FObjectBaseAddress& ObjectAddress, EPropertyPortFlags PortFlags)
 {
-	if (!InPropertyNode->HasNodeFlags(EPropertyNodeFlags::IsSparseClassData))
+	bool bIsSparseProperty = !!InPropertyNode->HasNodeFlags(EPropertyNodeFlags::IsSparseClassData);
+	bool bIsInContainer = false;
+	UProperty* Outer = Cast<UProperty>(Property->GetOuter());
+	if (bIsSparseProperty)
+	{
+		while (Outer)
+		{
+			UArrayProperty* ArrayOuter = Cast<UArrayProperty>(Property->GetOuter());
+			USetProperty* SetOuter = Cast<USetProperty>(Property->GetOuter());
+			UMapProperty* MapOuter = Cast<UMapProperty>(Property->GetOuter());
+			if (ArrayOuter || SetOuter || MapOuter)
+			{
+				bIsInContainer = true;
+				break;
+			}
+
+			Outer = Cast<UProperty>(Outer->GetOuter());
+		}
+	}
+	if (!bIsSparseProperty || bIsInContainer)
 	{
 		PropertyToTextHelper(OutString, InPropertyNode, Property, ObjectAddress.BaseAddress, PortFlags);
 	}
 	else
 	{
-// FRED_TODO: change this to call PropertyToTextHelper with the new base address and verify it is doing the right thing
+// TODO: once we're sure that these don't differ we should always use the call to PropertyToTextHelper
 		void* BaseAddress = ObjectAddress.Object->GetClass()->GetOrCreateSparseClassData();
 		void* ValueAddress = Property->ContainerPtrToValuePtr<void>(BaseAddress);
 		Property->ExportText_Direct(OutString, ValueAddress, ValueAddress, nullptr, PortFlags);
+
+		FString Test;
+		PropertyToTextHelper(Test, InPropertyNode, Property, (uint8*)ValueAddress, PortFlags);
+		check(Test.Compare(OutString) == 0);
 	}
 }
 
