EditConditions no longer crash when their parent object gets deleted from under them, eg. when editing a blueprint CDO and then compiling.

Source: 9581009 (//UE4/Dev-Editor)

diff --git a/Engine/Source/Editor/PropertyEditor/Private/EditConditionContext.cpp b/Engine/Source/Editor/PropertyEditor/Private/EditConditionContext.cpp
--- a/Engine/Source/Editor/PropertyEditor/Private/EditConditionContext.cpp
+++ b/Engine/Source/Editor/PropertyEditor/Private/EditConditionContext.cpp
@@ -81,7 +81,14 @@ T* FindTypedField(const TWeakPtr<FPropertyNode>& PropertyNode, const FString& Pr
 	return nullptr;
 }
 
-FPropertyNode* GetEditConditionParentNode(const TSharedPtr<FPropertyNode>& PropertyNode)
+/** 
+ * Get the parent to use as the context when evaluating the edit condition.
+ * For normal properties inside a UObject, this is the UObject. 
+ * For children of containers, this is the UObject the container is in. 
+ * Note: We do not support nested containers.
+ * The result can be nullptr in exceptional cases, eg. if the UI is getting rebuilt.
+ */
+static FPropertyNode* GetEditConditionParentNode(const TSharedPtr<FPropertyNode>& PropertyNode)
 {
 	FPropertyNode* ParentNode = PropertyNode->GetParentNode();
 	const UObject* PropertyOuter = PropertyNode->GetProperty()->GetOuter();
@@ -97,6 +104,30 @@ FPropertyNode* GetEditConditionParentNode(const TSharedPtr<FPropertyNode>& Prope
 	return ParentNode;
 }
 
+static uint8* GetPropertyValuePtr(const UProperty* Property, const TSharedPtr<FPropertyNode>& PropertyNode, FPropertyNode* ParentNode, FComplexPropertyNode* ComplexParentNode, int32 Index)
+{
+	uint8* BasePtr = ComplexParentNode->GetMemoryOfInstance(Index);
+	if (BasePtr == nullptr)
+	{
+		return nullptr;
+	}
+
+	uint8* ParentPtr = ParentNode->GetValueAddress(BasePtr, PropertyNode->HasNodeFlags(EPropertyNodeFlags::IsSparseClassData) != 0);
+	if (ParentPtr == nullptr)
+	{
+		return nullptr;
+	}
+
+	uint8* ValuePtr = ComplexParentNode->GetValuePtrOfInstance(Index, Property, ParentNode);
+
+	// SPARSEDATA_TODO: remove the next three lines once we're sure the pointer math is correct
+	uint8* OldValuePtr = Property->ContainerPtrToValuePtr<uint8>(ParentPtr);
+	check(OldValuePtr == ValuePtr || PropertyNode->HasNodeFlags(EPropertyNodeFlags::IsSparseClassData));
+	ensure(ValuePtr != nullptr);
+
+	return ValuePtr;
+}
+
 TOptional<bool> FEditConditionContext::GetBoolValue(const FString& PropertyName) const
 {
 	const UBoolProperty* BoolProperty = FindTypedField<UBoolProperty>(PropertyNode, PropertyName);
@@ -106,28 +137,28 @@ TOptional<bool> FEditConditionContext::GetBoolValue(const FString& PropertyName)
 	}
 
 	TSharedPtr<FPropertyNode> PinnedNode = PropertyNode.Pin();
+	if (!PinnedNode.IsValid())
+	{
+		return TOptional<bool>();
+	}
+
 	FPropertyNode* ParentNode = GetEditConditionParentNode(PinnedNode);
+	if (ParentNode == nullptr)
+	{
+		return TOptional<bool>();
+	}
+
 	FComplexPropertyNode* ComplexParentNode = PinnedNode->FindComplexParent();
 
 	TOptional<bool> Result;
 	for (int32 Index = 0; Index < ComplexParentNode->GetInstancesNum(); ++Index)
 	{
-		uint8* BasePtr = ComplexParentNode->GetMemoryOfInstance(Index);
-		if (BasePtr == nullptr)
+		uint8* ValuePtr = GetPropertyValuePtr(BoolProperty, PinnedNode, ParentNode, ComplexParentNode, Index);
+		if (ValuePtr == nullptr)
 		{
-			// deleted or not fully loaded yet
 			return TOptional<bool>();
 		}
 
-		uint8* ParentPtr = ParentNode->GetValueAddress(BasePtr, PinnedNode->HasNodeFlags(EPropertyNodeFlags::IsSparseClassData) != 0);
-
-		uint8* ValuePtr = ComplexParentNode->GetValuePtrOfInstance(Index, BoolProperty, ParentNode);
-
-		// SPARSEDATA_TODO: remove the next three lines once we're sure the pointer math is correct
-		uint8* OldValuePtr = BoolProperty->ContainerPtrToValuePtr<uint8>(ParentPtr);
-		check(OldValuePtr == ValuePtr || PinnedNode->HasNodeFlags(EPropertyNodeFlags::IsSparseClassData));
-		ensure(ValuePtr != nullptr);
-
 		bool bValue = BoolProperty->GetPropertyValue(ValuePtr);
 		if (!Result.IsSet())
 		{
@@ -152,28 +183,28 @@ TOptional<double> FEditConditionContext::GetNumericValue(const FString& Property
 	}
 
 	TSharedPtr<FPropertyNode> PinnedNode = PropertyNode.Pin();
+	if (!PinnedNode.IsValid())
+	{
+		return TOptional<double>();
+	}
+
 	FPropertyNode* ParentNode = GetEditConditionParentNode(PinnedNode);
+	if (ParentNode == nullptr)
+	{
+		return TOptional<double>();
+	}
+
 	FComplexPropertyNode* ComplexParentNode = PinnedNode->FindComplexParent();
 
 	TOptional<double> Result;
 	for (int32 Index = 0; Index < ComplexParentNode->GetInstancesNum(); ++Index)
 	{
-		uint8* BasePtr = ComplexParentNode->GetMemoryOfInstance(Index);
-		if (BasePtr == nullptr)
+		uint8* ValuePtr = GetPropertyValuePtr(NumericProperty, PinnedNode, ParentNode, ComplexParentNode, Index);
+		if (ValuePtr == nullptr)
 		{
-			// deleted or not fully loaded yet
 			return TOptional<double>();
 		}
 
-		uint8* ParentPtr = ParentNode->GetValueAddress(BasePtr, PinnedNode->HasNodeFlags(EPropertyNodeFlags::IsSparseClassData) != 0);
-
-		uint8* ValuePtr = ComplexParentNode->GetValuePtrOfInstance(Index, NumericProperty, ParentNode);
-
-		// SPARSEDATA_TODO: remove the next three lines once we're sure the pointer math is correct
-		uint8* OldValuePtr = NumericProperty->ContainerPtrToValuePtr<uint8>(ParentPtr);
-		check(OldValuePtr == ValuePtr || PinnedNode->HasNodeFlags(EPropertyNodeFlags::IsSparseClassData));
-		ensure(ValuePtr != nullptr);
-
 		double Value = 0;
 
 		if (NumericProperty->IsInteger())
@@ -223,30 +254,35 @@ TOptional<FString> FEditConditionContext::GetEnumValue(const FString& PropertyNa
 	{
 		return TOptional<FString>();
 	}
-	
+
+	if (!NumericProperty->IsInteger())
+	{
+		return TOptional<FString>();
+	}
+
 	TSharedPtr<FPropertyNode> PinnedNode = PropertyNode.Pin();
+	if (!PinnedNode.IsValid())
+	{
+		return TOptional<FString>();
+	}
+
 	FPropertyNode* ParentNode = GetEditConditionParentNode(PinnedNode);
+	if (ParentNode == nullptr)
+	{
+		return TOptional<FString>();
+	}
+
 	FComplexPropertyNode* ComplexParentNode = PinnedNode->FindComplexParent();
 
 	TOptional<int64> Result;
 	for (int32 Index = 0; Index < ComplexParentNode->GetInstancesNum(); ++Index)
 	{
-		uint8* BasePtr = ComplexParentNode->GetMemoryOfInstance(Index);
-		if (BasePtr == nullptr)
+		uint8* ValuePtr = GetPropertyValuePtr(Property, PinnedNode, ParentNode, ComplexParentNode, Index);
+		if (ValuePtr == nullptr)
 		{
-			// deleted or not fully loaded yet
 			return TOptional<FString>();
 		}
 
-		uint8* ParentPtr = ParentNode->GetValueAddress(BasePtr, PinnedNode->HasNodeFlags(EPropertyNodeFlags::IsSparseClassData) != 0);
-
-		uint8* ValuePtr = ComplexParentNode->GetValuePtrOfInstance(Index, Property, ParentNode);
-
-		// SPARSEDATA_TODO: remove the next three lines once we're sure the pointer math is correct
-		uint8* OldValuePtr = Property->ContainerPtrToValuePtr<uint8>(ParentPtr);
-		check(OldValuePtr == ValuePtr || PinnedNode->HasNodeFlags(EPropertyNodeFlags::IsSparseClassData));
-		ensure(ValuePtr != nullptr);
-
 		int64 Value = NumericProperty->GetSignedIntPropertyValue(ValuePtr);
 		if (!Result.IsSet())
 		{
@@ -259,12 +295,12 @@ TOptional<FString> FEditConditionContext::GetEnumValue(const FString& PropertyNa
 		}
 	}
 
-	if (Result.IsSet())
+	if (!Result.IsSet())
 	{
-		return EnumType->GetNameStringByValue(Result.GetValue());
+		return TOptional<FString>();
 	}
 
-	return TOptional<FString>();
+	return EnumType->GetNameStringByValue(Result.GetValue());
 }
 
 TOptional<FString> FEditConditionContext::GetTypeName(const FString& PropertyName) const
