Fixed categories in the details panel being unlocalized

diff --git a/Engine/Source/Editor/PropertyEditor/Private/DetailCategoryBuilderImpl.cpp b/Engine/Source/Editor/PropertyEditor/Private/DetailCategoryBuilderImpl.cpp
--- a/Engine/Source/Editor/PropertyEditor/Private/DetailCategoryBuilderImpl.cpp
+++ b/Engine/Source/Editor/PropertyEditor/Private/DetailCategoryBuilderImpl.cpp
@@ -611,33 +611,21 @@ void FDetailCategoryImpl::SetDisplayName(FName InCategoryName, const FText& Loca
 	{
 		DisplayName = LocalizedNameOverride;
 	}
+	else if (InCategoryName != NAME_None)
+	{
+		DisplayName = FText::AsCultureInvariant(FName::NameToDisplayString(InCategoryName.ToString(), false));
+	}
 	else
 	{
-		const UStruct* BaseStruct = GetParentLayoutImpl().GetRootNode()->GetBaseStructure();
 		// Use the base class name if there is one otherwise this is a generic category not specific to a class
-		FName BaseStructName = BaseStruct ? BaseStruct->GetFName() : FName("Generic");
-
-		FString CategoryStr = InCategoryName != NAME_None ? InCategoryName.ToString() : BaseStructName.ToString();
-		FString SourceCategoryStr = FName::NameToDisplayString(CategoryStr, false);
-
-		bool FoundText = false;
-		FText DisplayNameText;
-		if (InCategoryName != NAME_None)
-		{
-			FoundText = FText::FindText(TEXT("DetailCategory.CategoryName"), InCategoryName.ToString(), /*OUT*/DisplayNameText);
-		}
-		else
-		{
-			FoundText = FText::FindText(TEXT("DetailCategory.ClassName"), BaseStructName.ToString(), /*OUT*/DisplayNameText);
-		}
-
-		if (FoundText)
+		const UStruct* BaseStruct = GetParentLayoutImpl().GetRootNode()->GetBaseStructure();
+		if (BaseStruct)
 		{
-			DisplayName = DisplayNameText;
+			DisplayName = BaseStruct->GetDisplayNameText();
 		}
 		else
 		{
-			DisplayName = FText::FromString(SourceCategoryStr);
+			DisplayName = NSLOCTEXT("DetailCategory", "GenericCategory", "Generic");
 		}
 	}
 }
diff --git a/Engine/Source/Editor/PropertyEditor/Private/DetailLayoutHelpers.cpp b/Engine/Source/Editor/PropertyEditor/Private/DetailLayoutHelpers.cpp
--- a/Engine/Source/Editor/PropertyEditor/Private/DetailLayoutHelpers.cpp
+++ b/Engine/Source/Editor/PropertyEditor/Private/DetailLayoutHelpers.cpp
@@ -181,6 +181,24 @@ namespace DetailLayoutHelpers
 
 							// Add a property to the default category
 							FDetailCategoryImpl& CategoryImpl = DetailLayout.DefaultCategory(CategoryName);
+							{
+								FText CategoryDisplayName = CategoryImpl.GetDisplayName();
+								if (CategoryDisplayName.IsEmpty() || CategoryDisplayName.IsCultureInvariant())
+								{
+									const FString CategorySourceText = CategoryName.ToString();
+									const FString CategoryKey = Property->GetFullGroupName(false);
+
+									if (FText::FindText(TEXT("UObjectCategory"), CategoryKey, /*OUT*/CategoryDisplayName, &CategorySourceText))
+									{
+										// Category names in English are typically gathered in their non-pretty form (eg "UserInterface" rather than "User Interface"), so skip 
+										// applying the localized variant if the text matches the raw category name, as in this case the pretty printer will do a better job
+										if (!CategorySourceText.Equals(CategoryDisplayName.ToString(), ESearchCase::CaseSensitive))
+										{
+											CategoryImpl.SetDisplayName(CategoryName, CategoryDisplayName);
+										}
+									}
+								}
+							}
 							CategoryImpl.AddPropertyNode(ChildNodePtr.ToSharedRef(), InstanceName);
 						}
 
