Merge from Dev-Editor: Add temporary UObject to root to ensure it isn't being GC

diff --git a/Engine/Source/Editor/PropertyEditor/Private/UserInterface/PropertyEditor/SPropertyMenuAssetPicker.cpp b/Engine/Source/Editor/PropertyEditor/Private/UserInterface/PropertyEditor/SPropertyMenuAssetPicker.cpp
--- a/Engine/Source/Editor/PropertyEditor/Private/UserInterface/PropertyEditor/SPropertyMenuAssetPicker.cpp
+++ b/Engine/Source/Editor/PropertyEditor/Private/UserInterface/PropertyEditor/SPropertyMenuAssetPicker.cpp
@@ -291,12 +291,15 @@ void SPropertyMenuAssetPicker::OnCreateNewAssetSelected(TWeakObjectPtr<UFactory>
 	if (FactoryPtr.IsValid())
 	{
 		UFactory* FactoryInstance = DuplicateObject<UFactory>(FactoryPtr.Get(), GetTransientPackage());
+		// Ensure this object is not GC for the duration of CreateAssetWithDialog
+		FactoryInstance->AddToRoot();
 		FAssetToolsModule& AssetToolsModule = FAssetToolsModule::GetModule();
 		UObject* NewAsset = AssetToolsModule.Get().CreateAssetWithDialog(FactoryInstance->GetSupportedClass(), FactoryInstance);
 		if (NewAsset != nullptr)
 		{
 			SetValue(NewAsset);
 		}
+		FactoryInstance->RemoveFromRoot();
 	}
 }
 
