Exposed FAddPropertyParams to IDetailCategoryBuilder::AddExternalObjects API

This is done primarily to allow unique IDs to be specified to persist expansion state when displaying multiple external objects of the same type.
This forces AllowChildren(true) to mimic previous behavior

diff --git a/Engine/Source/Editor/PropertyEditor/Private/DetailCategoryBuilderImpl.cpp b/Engine/Source/Editor/PropertyEditor/Private/DetailCategoryBuilderImpl.cpp
--- a/Engine/Source/Editor/PropertyEditor/Private/DetailCategoryBuilderImpl.cpp
+++ b/Engine/Source/Editor/PropertyEditor/Private/DetailCategoryBuilderImpl.cpp
@@ -325,11 +325,14 @@ IDetailPropertyRow& FDetailCategoryImpl::AddProperty(TSharedPtr<IPropertyHandle>
 	return *NewCustomization.PropertyRow;
 }
 
-IDetailPropertyRow* FDetailCategoryImpl::AddExternalObjects(const TArray<UObject*>& Objects, EPropertyLocation::Type Location /*= EPropertyLocation::Default*/)
+IDetailPropertyRow* FDetailCategoryImpl::AddExternalObjects(const TArray<UObject*>& Objects, EPropertyLocation::Type Location /*= EPropertyLocation::Default*/, const FAddPropertyParams& Params /*= FAddPropertyParams()*/)
 {
 	FDetailLayoutCustomization NewCustomization;
 
-	FDetailPropertyRow::MakeExternalPropertyRowCustomization(Objects, NAME_None, AsShared(), NewCustomization, FAddPropertyParams().AllowChildren(true));
+	FAddPropertyParams AddPropertyParams = Params;
+	AddPropertyParams.AllowChildren(true);
+
+	FDetailPropertyRow::MakeExternalPropertyRowCustomization(Objects, NAME_None, AsShared(), NewCustomization, AddPropertyParams);
 
 	TSharedPtr<FDetailPropertyRow> NewRow = NewCustomization.PropertyRow;
 
diff --git a/Engine/Source/Editor/PropertyEditor/Private/DetailCategoryBuilderImpl.h b/Engine/Source/Editor/PropertyEditor/Private/DetailCategoryBuilderImpl.h
--- a/Engine/Source/Editor/PropertyEditor/Private/DetailCategoryBuilderImpl.h
+++ b/Engine/Source/Editor/PropertyEditor/Private/DetailCategoryBuilderImpl.h
@@ -162,7 +162,7 @@ public:
 	virtual IDetailCategoryBuilder& HeaderContent(TSharedRef<SWidget> InHeaderContent) override;
 	virtual IDetailPropertyRow& AddProperty(FName PropertyPath, UClass* ClassOuter = nullptr, FName InstanceName = NAME_None, EPropertyLocation::Type Location = EPropertyLocation::Default) override;
 	virtual IDetailPropertyRow& AddProperty(TSharedPtr<IPropertyHandle> PropertyHandle, EPropertyLocation::Type Location = EPropertyLocation::Default) override;
-	virtual IDetailPropertyRow* AddExternalObjects(const TArray<UObject*>& Objects, EPropertyLocation::Type Location = EPropertyLocation::Default) override;
+	virtual IDetailPropertyRow* AddExternalObjects(const TArray<UObject*>& Objects, EPropertyLocation::Type Location = EPropertyLocation::Default, const FAddPropertyParams& Params = FAddPropertyParams()) override;
 	virtual IDetailPropertyRow* AddExternalObjectProperty(const TArray<UObject*>& Objects, FName PropertyName, EPropertyLocation::Type Location = EPropertyLocation::Default, const FAddPropertyParams& Params = FAddPropertyParams()) override;
 	virtual IDetailPropertyRow* AddExternalStructure(TSharedPtr<FStructOnScope> StructData, EPropertyLocation::Type Location = EPropertyLocation::Default) override;
 	virtual IDetailPropertyRow* AddExternalStructureProperty(TSharedPtr<FStructOnScope> StructData, FName PropertyName, EPropertyLocation::Type Location = EPropertyLocation::Default, const FAddPropertyParams& Params = FAddPropertyParams()) override;
diff --git a/Engine/Source/Editor/PropertyEditor/Private/DetailPropertyRow.cpp b/Engine/Source/Editor/PropertyEditor/Private/DetailPropertyRow.cpp
--- a/Engine/Source/Editor/PropertyEditor/Private/DetailPropertyRow.cpp
+++ b/Engine/Source/Editor/PropertyEditor/Private/DetailPropertyRow.cpp
@@ -466,6 +466,7 @@ void FDetailPropertyRow::MakeExternalPropertyRowCustomization(TSharedPtr<FStruct
 				if (Property->GetFName() == PropertyName)
 				{
 					OutCustomization.PropertyRow = MakeShareable(new FDetailPropertyRow(PropertyNode, ParentCategory, RootPropertyNode));
+					OutCustomization.PropertyRow->SetCustomExpansionId(Parameters.GetUniqueId());
 					break;
 				}
 			}
@@ -496,6 +497,7 @@ void FDetailPropertyRow::MakeExternalPropertyRowCustomization(TSharedPtr<FStruct
 		RootPropertyNode->AddChildNode(ItemNode);
 
 		OutCustomization.PropertyRow = MakeShareable(new FDetailPropertyRow(ItemNode, ParentCategory, RootPropertyNode));
+		OutCustomization.PropertyRow->SetCustomExpansionId(Parameters.GetUniqueId());
 	}
 }
 
@@ -534,11 +536,13 @@ void FDetailPropertyRow::MakeExternalPropertyRowCustomization(const TArray<UObje
 			PropertyNode->RebuildChildren();
 
 			OutCustomization.PropertyRow = MakeShared<FDetailPropertyRow>(PropertyNode, ParentCategory, RootPropertyNode);
+			OutCustomization.PropertyRow->SetCustomExpansionId(Parameters.GetUniqueId());
 		}
 	}
 	else
 	{
 		OutCustomization.PropertyRow = MakeShared<FDetailPropertyRow>(RootPropertyNode, ParentCategory, RootPropertyNode);
+		OutCustomization.PropertyRow->SetCustomExpansionId(Parameters.GetUniqueId());
 	}
 }
 
diff --git a/Engine/Source/Editor/PropertyEditor/Public/DetailCategoryBuilder.h b/Engine/Source/Editor/PropertyEditor/Public/DetailCategoryBuilder.h
--- a/Engine/Source/Editor/PropertyEditor/Public/DetailCategoryBuilder.h
+++ b/Engine/Source/Editor/PropertyEditor/Public/DetailCategoryBuilder.h
@@ -83,7 +83,7 @@ public:
 	 * @param Location		The location within the category where the property is shown
 	 * @return A property row for customizing the property or NULL if the property could not be found
 	 */
-	virtual IDetailPropertyRow* AddExternalObjects(const TArray<UObject*>& Objects, EPropertyLocation::Type Location = EPropertyLocation::Default) = 0;
+	virtual IDetailPropertyRow* AddExternalObjects(const TArray<UObject*>& Objects, EPropertyLocation::Type Location = EPropertyLocation::Default, const FAddPropertyParams& Params = FAddPropertyParams()) = 0;
 
 	/**
 	 * Adds an external object property that is not a property on the object(s) being customized
