FPropertyHandleBase no longer accesses invalid properties without checking for null.

diff --git a/Engine/Source/Editor/PropertyEditor/Private/PropertyHandleImpl.cpp b/Engine/Source/Editor/PropertyEditor/Private/PropertyHandleImpl.cpp
--- a/Engine/Source/Editor/PropertyEditor/Private/PropertyHandleImpl.cpp
+++ b/Engine/Source/Editor/PropertyEditor/Private/PropertyHandleImpl.cpp
@@ -2851,6 +2851,10 @@ FPropertyAccess::Result FPropertyHandleBase::GetPerObjectValue( const int32 Obje
 bool FPropertyHandleBase::GeneratePossibleValues(TArray< TSharedPtr<FString> >& OutOptionStrings, TArray< FText >& OutToolTips, TArray<bool>& OutRestrictedItems)
 {
 	UProperty* Property = GetProperty();
+	if (Property == nullptr)
+	{
+		return false;
+	}
 
 	bool bUsesAlternateDisplayValues = false;
 
@@ -3148,12 +3152,16 @@ TArray<TSharedPtr<IPropertyHandle>> FPropertyHandleBase::AddChildStructure( TSha
 bool FPropertyHandleBase::CanResetToDefault() const
 {
 	UProperty* Property = GetProperty();
+	if (Property == nullptr)
+	{
+		return false;
+	}
 
 	// Should not be able to reset fixed size arrays
-	const bool bFixedSized = Property && Property->PropertyFlags & CPF_EditFixedSize;
-	const bool bCanResetToDefault = !(Property && Property->PropertyFlags & CPF_Config);
+	const bool bFixedSized = Property->PropertyFlags & CPF_EditFixedSize;
+	const bool bCanResetToDefault = !(Property->PropertyFlags & CPF_Config);
 
-	return Property && bCanResetToDefault && !bFixedSized && DiffersFromDefault();
+	return bCanResetToDefault && !bFixedSized && DiffersFromDefault();
 }
 
 void FPropertyHandleBase::ExecuteCustomResetToDefault(const FResetToDefaultOverride& InOnCustomResetToDefault)
