Fixed compilation error due to conversion of flag comparison to bool.

diff --git a/Engine/Source/Editor/PropertyEditor/Private/PropertyHandleImpl.cpp b/Engine/Source/Editor/PropertyEditor/Private/PropertyHandleImpl.cpp
--- a/Engine/Source/Editor/PropertyEditor/Private/PropertyHandleImpl.cpp
+++ b/Engine/Source/Editor/PropertyEditor/Private/PropertyHandleImpl.cpp
@@ -3158,8 +3158,8 @@ bool FPropertyHandleBase::CanResetToDefault() const
 	}
 
 	// Should not be able to reset fixed size arrays
-	const bool bFixedSized = Property->PropertyFlags & CPF_EditFixedSize;
-	const bool bCanResetToDefault = !(Property->PropertyFlags & CPF_Config);
+	const bool bFixedSized = (Property->PropertyFlags & CPF_EditFixedSize) != 0;
+	const bool bCanResetToDefault = (Property->PropertyFlags & CPF_Config) == 0;
 
 	return bCanResetToDefault && !bFixedSized && DiffersFromDefault();
 }
