PropertyEditor - Add an api to generate default property buttons to the property handle, for use when generating whole row content, but you want the default buttons.

diff --git a/Engine/Source/Editor/PropertyEditor/Private/PropertyHandleImpl.cpp b/Engine/Source/Editor/PropertyEditor/Private/PropertyHandleImpl.cpp
--- a/Engine/Source/Editor/PropertyEditor/Private/PropertyHandleImpl.cpp
+++ b/Engine/Source/Editor/PropertyEditor/Private/PropertyHandleImpl.cpp
@@ -25,6 +25,7 @@
 #include "IDetailPropertyRow.h"
 #include "ObjectEditorUtils.h"
 #include "PropertyEditorHelpers.h"
+#include "SResetToDefaultPropertyEditor.h"
 
 #define LOCTEXT_NAMESPACE "PropertyHandleImplementation"
 
@@ -2331,6 +2332,60 @@ TSharedRef<SWidget> FPropertyHandleBase::CreatePropertyValueWidget( bool bDispla
 	return SNullWidget::NullWidget;
 }
 
+class SDefaultPropertyButtonWidgets : public SCompoundWidget
+{
+	SLATE_BEGIN_ARGS(SDefaultPropertyButtonWidgets)	{}
+	SLATE_END_ARGS()
+
+	void Construct(const FArguments& InArgs, TSharedRef<FPropertyEditor> InPropertyEditor)
+	{
+		PropertyEditor = InPropertyEditor;
+		TSharedRef<SHorizontalBox> ButtonBox = SNew(SHorizontalBox);
+
+		TArray<TSharedRef<SWidget>> RequiredButtons;
+		PropertyEditorHelpers::MakeRequiredPropertyButtons(PropertyEditor.ToSharedRef(), RequiredButtons);
+		for (TSharedRef<SWidget> RequiredButton : RequiredButtons)
+		{
+			ButtonBox->AddSlot()
+			.AutoWidth()
+			.HAlign(HAlign_Center)
+			.VAlign(VAlign_Center)
+			.Padding(2.0f, 1.0f)
+			[
+				RequiredButton
+			];
+		}
+
+		ButtonBox->AddSlot()
+		.AutoWidth()
+		.HAlign(HAlign_Left)
+		.VAlign(VAlign_Center)
+		.Padding(4.0f, 0.0f)
+		[
+			SNew(SResetToDefaultPropertyEditor, PropertyEditor->GetPropertyHandle())
+		];
+
+		ChildSlot
+		[
+			ButtonBox
+		];
+	}
+
+private:
+	TSharedPtr<FPropertyEditor> PropertyEditor;
+};
+
+TSharedRef<SWidget> FPropertyHandleBase::CreateDefaultPropertyButtonWidgets() const
+{
+	TArray<TSharedRef<SWidget>> DefaultButtons;
+	if (Implementation.IsValid() && Implementation->GetPropertyNode().IsValid())
+	{
+		TSharedRef<FPropertyEditor> PropertyEditor = FPropertyEditor::Create(Implementation->GetPropertyNode().ToSharedRef(), Implementation->GetPropertyUtilities().ToSharedRef());
+		return SNew(SDefaultPropertyButtonWidgets, PropertyEditor);
+	}
+	return SNullWidget::NullWidget;
+}
+
 bool FPropertyHandleBase::IsEditConst() const
 {
 	return Implementation->IsEditConst();
diff --git a/Engine/Source/Editor/PropertyEditor/Private/PropertyHandleImpl.h b/Engine/Source/Editor/PropertyEditor/Private/PropertyHandleImpl.h
--- a/Engine/Source/Editor/PropertyEditor/Private/PropertyHandleImpl.h
+++ b/Engine/Source/Editor/PropertyEditor/Private/PropertyHandleImpl.h
@@ -433,6 +433,7 @@ public:
 	virtual FString GeneratePathToProperty() const override;
 	virtual TSharedRef<SWidget> CreatePropertyNameWidget( const FText& NameOverride = FText::GetEmpty(), const FText& ToolTipOverride = FText::GetEmpty(), bool bDisplayResetToDefault = false, bool bDisplayText = true, bool bDisplayThumbnail = true ) const override;
 	virtual TSharedRef<SWidget> CreatePropertyValueWidget( bool bDisplayDefaultPropertyButtons = true ) const override;
+	virtual TSharedRef<SWidget> CreateDefaultPropertyButtonWidgets() const override;
 	virtual bool IsEditConst() const override;
 	virtual bool IsEditable() const override;
 	virtual void SetOnPropertyValueChanged( const FSimpleDelegate& InOnPropertyValueChanged ) override;
diff --git a/Engine/Source/Editor/PropertyEditor/Public/PropertyHandle.h b/Engine/Source/Editor/PropertyEditor/Public/PropertyHandle.h
--- a/Engine/Source/Editor/PropertyEditor/Public/PropertyHandle.h
+++ b/Engine/Source/Editor/PropertyEditor/Public/PropertyHandle.h
@@ -534,6 +534,12 @@ public:
 	 */
 	virtual TSharedRef<SWidget> CreatePropertyValueWidget( bool bDisplayDefaultPropertyButtons = true ) const = 0;
 
+	/**
+	 * Creates the default buttons which appear next to value widgets.  This is useful when creating customizations
+	 * which don't use CreatePropertyValueWidget but you still want array item behaviors and reset to default capabilities.
+	 */
+	virtual TSharedRef<SWidget> CreateDefaultPropertyButtonWidgets() const = 0;
+
 	/**
 	 * Adds a restriction to the possible values for this property.
 	 * @param Restriction	The restriction being added to this property.
