Fix crash on creating local variable in a function in animBP

diff --git a/Engine/Source/Editor/PropertyEditor/Private/SDetailSingleItemRow.cpp b/Engine/Source/Editor/PropertyEditor/Private/SDetailSingleItemRow.cpp
--- a/Engine/Source/Editor/PropertyEditor/Private/SDetailSingleItemRow.cpp
+++ b/Engine/Source/Editor/PropertyEditor/Private/SDetailSingleItemRow.cpp
@@ -740,8 +740,9 @@ TSharedRef<SWidget> SDetailSingleItemRow::CreateExtensionWidget(TSharedRef<SWidg
 		{
 			TSharedPtr<IPropertyHandle> Handle = PropertyEditorHelpers::GetPropertyHandle(InCustomization.GetPropertyNode().ToSharedRef(), nullptr, nullptr);
 
-			UClass* ObjectClass = InCustomization.GetPropertyNode()->FindObjectItemParent()->GetObjectBaseClass();
-			if(Handle->IsValidHandle() && ExtensionHandler->IsPropertyExtendable(ObjectClass, *Handle))
+			FObjectPropertyNode* ObjectItemParent = InCustomization.GetPropertyNode()->FindObjectItemParent();
+			UClass* ObjectClass = (ObjectItemParent) ? ObjectItemParent->GetObjectBaseClass() : nullptr;
+			if(Handle->IsValidHandle() && ObjectClass && ExtensionHandler->IsPropertyExtendable(ObjectClass, *Handle))
 			{
 				FDetailLayoutBuilderImpl& DetailLayout = OwnerTreeNode.Pin()->GetParentCategory()->GetParentLayoutImpl();
 				ExtensionWidget = ExtensionHandler->GenerateExtensionWidget(DetailLayout, ObjectClass, Handle);
