Cherry-pick 10494848 :
HideCategories on components in an actor now correctly hides the categories when the actor is selected.

diff --git a/Engine/Source/Editor/PropertyEditor/Private/ObjectPropertyNode.cpp b/Engine/Source/Editor/PropertyEditor/Private/ObjectPropertyNode.cpp
--- a/Engine/Source/Editor/PropertyEditor/Private/ObjectPropertyNode.cpp
+++ b/Engine/Source/Editor/PropertyEditor/Private/ObjectPropertyNode.cpp
@@ -457,8 +457,6 @@ void FObjectPropertyNode::GetCategoryProperties(const TSet<UClass*>& ClassesToCo
 {
 	static const FName Name_bShowOnlyWhenTrue("bShowOnlyWhenTrue");
 
-	bool bHidden = false;
-
 	FName CategoryName = FObjectEditorUtils::GetCategoryFName(CurrentProperty);
 
 	for (UClass* Class : ClassesToConsider)
@@ -480,7 +478,7 @@ void FObjectPropertyNode::GetCategoryProperties(const TSet<UClass*>& ClassesToCo
 
 	if (bMetaDataAllowVisible)
 	{
-		if (PropertyEditorHelpers::ShouldBeVisible(*this, CurrentProperty))
+		if (PropertyEditorHelpers::ShouldBeVisible(*this, CurrentProperty) && !HiddenCategories.Contains(CategoryName))
 		{
 			if (!CategoriesFromBlueprints.Contains(CategoryName) && !CategoriesFromProperties.Contains(CategoryName))
 			{
