Cherry-pick 10506158 :

- Added possibility to disallow assets with certain asset metadata values in the asset picker
- Forbit usage of virtual textures in procedural tools
- Shield against bad texture data in procedural tools

diff --git a/Engine/Source/Editor/PropertyEditor/Private/UserInterface/PropertyEditor/SPropertyEditorAsset.cpp b/Engine/Source/Editor/PropertyEditor/Private/UserInterface/PropertyEditor/SPropertyEditorAsset.cpp
--- a/Engine/Source/Editor/PropertyEditor/Private/UserInterface/PropertyEditor/SPropertyEditorAsset.cpp
+++ b/Engine/Source/Editor/PropertyEditor/Private/UserInterface/PropertyEditor/SPropertyEditorAsset.cpp
@@ -205,6 +205,51 @@ void SPropertyEditorAsset::InitializeClassFilters(const UProperty* Property)
 	}
 }
 
+void SPropertyEditorAsset::InitializeAssetDataTags(const UProperty* Property)
+{
+	if (Property == nullptr)
+	{
+		return;
+	}
+
+	const UProperty* MetadataProperty = GetActualMetadataProperty(Property);
+	const FString DisallowedAssetDataTagsFilterString = MetadataProperty->GetMetaData(TEXT("DisallowedAssetDataTags"));
+	if (!DisallowedAssetDataTagsFilterString.IsEmpty())
+	{
+		TArray<FString> DisallowedAssetDataTagsAndValues;
+		DisallowedAssetDataTagsFilterString.ParseIntoArray(DisallowedAssetDataTagsAndValues, TEXT(","), true);
+
+		for (const FString& TagAndOptionalValueString : DisallowedAssetDataTagsAndValues)
+		{
+			TArray<FString> TagAndOptionalValue;
+			TagAndOptionalValueString.ParseIntoArray(TagAndOptionalValue, TEXT("="), true);
+			size_t NumStrings = TagAndOptionalValue.Num();
+			check((NumStrings == 1) || (NumStrings == 2)); // there should be a single '=' within a tag/value pair
+
+			if (!DisallowedAssetDataTags.IsValid())
+			{
+				DisallowedAssetDataTags = MakeShared<FAssetDataTagMap>();
+			}
+			DisallowedAssetDataTags->Add(FName(*TagAndOptionalValue[0]), (NumStrings > 1) ? TagAndOptionalValue[1] : FString());
+		}
+	}
+}
+
+bool SPropertyEditorAsset::IsAssetAllowed(const FAssetData& InAssetData)
+{
+	if (DisallowedAssetDataTags.IsValid())
+	{
+		for (const auto& DisallowedTagAndValue : *DisallowedAssetDataTags.Get())
+		{
+			if (InAssetData.TagsAndValues.ContainsKeyValue(DisallowedTagAndValue.Key, DisallowedTagAndValue.Value))
+			{
+				return false;
+			}
+		}
+	}
+	return true;
+}
+
 // Awful hack to deal with UClass::FindCommonBase taking an array of non-const classes...
 static TArray<UClass*> ConstCastClassArray(TArray<const UClass*>& Classes)
 {
@@ -239,6 +284,20 @@ void SPropertyEditorAsset::Construct(const FArguments& InArgs, const TSharedPtr<
 	ObjectClass = InArgs._Class != nullptr ? InArgs._Class : GetObjectPropertyClass(Property);
 	bAllowClear = InArgs._AllowClear.IsSet() ? InArgs._AllowClear.GetValue() : (Property ? !(Property->PropertyFlags & CPF_NoClear) : true);
 
+	InitializeAssetDataTags(Property);
+	if (DisallowedAssetDataTags.IsValid())
+	{
+		// re-route the filter delegate to our own if we have our own asset data tags filter :
+		OnShouldFilterAsset.BindLambda([this, AssetFilter = InArgs._OnShouldFilterAsset](const FAssetData& InAssetData)
+		{
+			if (IsAssetAllowed(InAssetData))
+			{
+				return AssetFilter.IsBound() ? AssetFilter.Execute(InAssetData) : false;
+			}
+			return true;
+		});
+	}
+
 	InitializeClassFilters(Property);
 
 	// Make the ObjectClass more specific if we only have one class filter
diff --git a/Engine/Source/Editor/PropertyEditor/Private/UserInterface/PropertyEditor/SPropertyEditorAsset.h b/Engine/Source/Editor/PropertyEditor/Private/UserInterface/PropertyEditor/SPropertyEditorAsset.h
--- a/Engine/Source/Editor/PropertyEditor/Private/UserInterface/PropertyEditor/SPropertyEditorAsset.h
+++ b/Engine/Source/Editor/PropertyEditor/Private/UserInterface/PropertyEditor/SPropertyEditorAsset.h
@@ -327,6 +327,14 @@ private:
 	/** Check if the given class is allowed by the AllowedClasses and DisallowedClasses as defined in the property metadata. */
 	bool IsClassAllowed(const UClass* InClass) const;
 
+	/**
+	 * Initialize DisallowedAssetDataTags based on the property metadata.
+	 */
+	void InitializeAssetDataTags(const UProperty* Property);
+
+	/** @return Returns true if the asset is relevant for this property*/
+	bool IsAssetAllowed(const FAssetData& InAssetData);
+
 private:
 
 	/** Main combobutton */
@@ -356,6 +364,9 @@ private:
 	/** A list of the factories we can use to create new assets */
 	TArray<UFactory*> NewAssetFactories;
 
+	/** Tags (and eventually values) that can NOT be used with this property */
+	TSharedPtr<FAssetDataTagMap> DisallowedAssetDataTags;
+
 	/** Whether the asset can be 'None' in this case */
 	bool bAllowClear;
 
