Array drag and drop to add elements

diff --git a/Engine/Source/Editor/PropertyEditor/Private/Presentation/PropertyEditor/PropertyEditor.cpp b/Engine/Source/Editor/PropertyEditor/Private/Presentation/PropertyEditor/PropertyEditor.cpp
--- a/Engine/Source/Editor/PropertyEditor/Private/Presentation/PropertyEditor/PropertyEditor.cpp
+++ b/Engine/Source/Editor/PropertyEditor/Private/Presentation/PropertyEditor/PropertyEditor.cpp
@@ -195,6 +195,12 @@ void FPropertyEditor::AddItem()
 	PropertyUtilities->EnqueueDeferredAction( FSimpleDelegate::CreateSP( this, &FPropertyEditor::OnAddItem ) );
 }
 
+void FPropertyEditor::AddGivenItem(const FString& InGivenItem)
+{
+	// This action must be deferred until next tick so that we avoid accessing invalid data before we have a chance to tick
+	PropertyUtilities->EnqueueDeferredAction(FSimpleDelegate::CreateSP(this, &FPropertyEditor::OnAddGivenItem, InGivenItem));
+}
+
 void FPropertyEditor::OnAddItem()
 {
 	// Check to make sure that the property is a valid container
@@ -227,6 +233,30 @@ void FPropertyEditor::OnAddItem()
 	}
 }
 
+void FPropertyEditor::OnAddGivenItem(const FString InGivenItem)
+{
+	OnAddItem();
+
+	// Check to make sure that the property is a valid container
+	TSharedPtr<IPropertyHandleArray> ArrayHandle = PropertyHandle->AsArray();
+
+	check(ArrayHandle.IsValid());
+
+	TSharedPtr<IPropertyHandle> ElementHandle;
+
+	if (ArrayHandle.IsValid())
+	{
+		uint32 Last;
+		ArrayHandle->GetNumElements(Last);
+		ElementHandle = ArrayHandle->GetElement(Last - 1);
+	}
+
+	if (ElementHandle.IsValid())
+	{
+		ElementHandle->SetValueFromFormattedString(InGivenItem);
+	}
+}
+
 void FPropertyEditor::ClearItem()
 {
 	OnClearItem();
diff --git a/Engine/Source/Editor/PropertyEditor/Private/Presentation/PropertyEditor/PropertyEditor.h b/Engine/Source/Editor/PropertyEditor/Private/Presentation/PropertyEditor/PropertyEditor.h
--- a/Engine/Source/Editor/PropertyEditor/Private/Presentation/PropertyEditor/PropertyEditor.h
+++ b/Engine/Source/Editor/PropertyEditor/Private/Presentation/PropertyEditor/PropertyEditor.h
@@ -72,6 +72,7 @@ public:
 
 	void UseSelected();
 	void AddItem();
+	void AddGivenItem(const FString& InGivenItem);
 	void ClearItem();
 	void InsertItem();
 	void DeleteItem();
@@ -109,6 +110,7 @@ private:
 
 	void OnUseSelected();
 	void OnAddItem();
+	void OnAddGivenItem(const FString InGivenItem);
 	void OnClearItem();
 	void OnInsertItem();
 	void OnDeleteItem();
diff --git a/Engine/Source/Editor/PropertyEditor/Private/SDetailSingleItemRow.cpp b/Engine/Source/Editor/PropertyEditor/Private/SDetailSingleItemRow.cpp
--- a/Engine/Source/Editor/PropertyEditor/Private/SDetailSingleItemRow.cpp
+++ b/Engine/Source/Editor/PropertyEditor/Private/SDetailSingleItemRow.cpp
@@ -221,6 +221,12 @@ FReply SDetailSingleItemRow::OnArrayDrop(const FDragDropEvent& DragDropEvent)
 	return FReply::Handled();
 }
 
+FReply SDetailSingleItemRow::OnArrayHeaderDrop(const FDragDropEvent& DragDropEvent)
+{
+	OnArrayDragLeave(DragDropEvent);
+	return FReply::Handled();
+}
+
 TSharedPtr<FPropertyNode> SDetailSingleItemRow::GetCopyPastePropertyNode() const
 {
 	TSharedPtr<FPropertyNode> PropertyNode = Customization->GetPropertyNode();
@@ -361,7 +367,15 @@ void SDetailSingleItemRow::Construct( const FArguments& InArgs, FDetailLayoutCus
 				ArrayDropDelegate = FOnTableRowDrop::CreateSP(this, &SDetailSingleItemRow::OnArrayDrop);
 				SwappablePropertyNode = PropertyNode;
 			}
-
+			else if (PropertyNode.IsValid() 
+				&& Cast<UArrayProperty>(PropertyNode->GetProperty()) != nullptr // Is an array
+				&& Cast<UObjectProperty>(Cast<UArrayProperty>(PropertyNode->GetProperty())->Inner) != nullptr) // Is an object array
+			{
+				ArrayDragDelegate = FOnTableRowDragEnter::CreateSP(this, &SDetailSingleItemRow::OnArrayDragEnter);
+				ArrayDragLeaveDelegate = FOnTableRowDragLeave::CreateSP(this, &SDetailSingleItemRow::OnArrayDragLeave);
+				ArrayDropDelegate = FOnTableRowDrop::CreateSP(this, &SDetailSingleItemRow::OnArrayHeaderDrop);
+			}
+			
 			InternalLeftColumnRowBox->AddSlot()
 				.Padding(0.0f, 0.0f)
 				.HAlign(HAlign_Left)
diff --git a/Engine/Source/Editor/PropertyEditor/Private/SDetailSingleItemRow.h b/Engine/Source/Editor/PropertyEditor/Private/SDetailSingleItemRow.h
--- a/Engine/Source/Editor/PropertyEditor/Private/SDetailSingleItemRow.h
+++ b/Engine/Source/Editor/PropertyEditor/Private/SDetailSingleItemRow.h
@@ -102,6 +102,7 @@ private:
 	void OnArrayDragEnter(const FDragDropEvent& DragDropEvent);
 	void OnArrayDragLeave(const FDragDropEvent& DragDropEvent);
 	FReply OnArrayDrop(const FDragDropEvent& DragDropEvent);
+	FReply OnArrayHeaderDrop(const FDragDropEvent& DragDropEvent);
 
 	TSharedPtr<FPropertyNode> GetCopyPastePropertyNode() const;
 private:
diff --git a/Engine/Source/Editor/PropertyEditor/Private/UserInterface/PropertyEditor/SPropertyEditorArray.h b/Engine/Source/Editor/PropertyEditor/Private/UserInterface/PropertyEditor/SPropertyEditorArray.h
--- a/Engine/Source/Editor/PropertyEditor/Private/UserInterface/PropertyEditor/SPropertyEditorArray.h
+++ b/Engine/Source/Editor/PropertyEditor/Private/UserInterface/PropertyEditor/SPropertyEditorArray.h
@@ -11,6 +11,8 @@
 #include "PropertyEditorHelpers.h"
 #include "UserInterface/PropertyEditor/PropertyEditorConstants.h"
 #include "Widgets/Text/STextBlock.h"
+#include "DragAndDrop/AssetDragDropOp.h"
+#include "SDropTarget.h"
 
 
 #define LOCTEXT_NAMESPACE "PropertyEditor"
@@ -43,9 +45,17 @@ public:
 		ChildSlot
 		.Padding( 0.0f, 0.0f, 2.0f, 0.0f )
 		[
-			SNew( STextBlock )
-			.Text( TextAttr )
-			.Font( InArgs._Font )
+			SNew(SDropTarget)
+			.OnDrop(this, &SPropertyEditorArray::OnDragDropTarget)
+			.OnAllowDrop(this, &SPropertyEditorArray::WillAddValidElements)
+			.OnIsRecognized(this, &SPropertyEditorArray::IsValidAssetDropOp)
+			.Content()
+			[
+				SNew(STextBlock)
+				.Margin(FMargin(4.0f, 2.0f))
+				.Text(TextAttr)
+				.Font(InArgs._Font)
+			]
 		];
 
 		SetEnabled( TAttribute<bool>( this, &SPropertyEditorArray::CanEdit ) );
@@ -82,6 +92,85 @@ private:
 	{
 		return PropertyEditor.IsValid() ? !PropertyEditor->IsEditConst() : true;
 	}
+
+	FReply OnDragDropTarget(TSharedPtr<FDragDropOperation> InOperation)
+	{
+		UObjectProperty* ObjectProperty = nullptr;
+		if (const UArrayProperty* ArrayProperty = Cast<UArrayProperty>(PropertyEditor->GetProperty()))
+		{
+			ObjectProperty = Cast<UObjectProperty>(ArrayProperty->Inner);
+		}
+
+		// Only try to add entries if we are dropping on an asset array
+		if (ObjectProperty && InOperation->IsOfType<FAssetDragDropOp>())
+		{
+			TSharedPtr<FAssetDragDropOp> DragDropOp = StaticCastSharedPtr<FAssetDragDropOp>(InOperation);
+			if (DragDropOp.IsValid())
+			{
+				for (FAssetData AssetData : DragDropOp->GetAssets())
+				{
+					// if the type matches
+					if (ObjectProperty->PropertyClass == AssetData.GetClass())
+					{
+						PropertyEditor->AddGivenItem(AssetData.ObjectPath.ToString());
+					}
+				}
+				// Let this bubble up to the rest of the row
+				return FReply::Unhandled();
+
+			}
+		}
+		return FReply::Unhandled();
+	}
+
+	bool IsValidAssetDropOp(TSharedPtr<FDragDropOperation> InOperation)
+	{
+		UObjectProperty* ObjectProperty = nullptr;
+		if (const UArrayProperty* ArrayProperty = Cast<UArrayProperty>(PropertyEditor->GetProperty()))
+		{
+			ObjectProperty = Cast<UObjectProperty>(ArrayProperty->Inner);
+		}
+
+		// Only try to add entries if we are dropping on an asset array
+		if (ObjectProperty && InOperation->IsOfType<FAssetDragDropOp>())
+		{
+			return true;
+		}
+		return false;
+	}
+
+	bool WillAddValidElements(TSharedPtr<FDragDropOperation> InOperation)
+	{
+		UObjectProperty* ObjectProperty = nullptr;
+
+		if (const UArrayProperty* ArrayProperty = Cast<UArrayProperty>(PropertyEditor->GetProperty()))
+		{
+			ObjectProperty = Cast<UObjectProperty>(ArrayProperty->Inner);
+		}
+
+		// Only try to add entries if we are dropping on an asset array
+		if (ObjectProperty && InOperation->IsOfType<FAssetDragDropOp>())
+		{
+			bool bHasOnlyValidElements = true;
+			TSharedPtr<FAssetDragDropOp> DragDropOp = StaticCastSharedPtr<FAssetDragDropOp>(InOperation);
+			if (DragDropOp.IsValid())
+			{
+				for (FAssetData AssetData : DragDropOp->GetAssets())
+				{
+					// if the type matches
+					if (ObjectProperty->PropertyClass != AssetData.GetClass())
+					{
+						bHasOnlyValidElements = false;
+						break;
+					}
+				}
+			}
+			return bHasOnlyValidElements;
+		}
+		
+		return false;
+	}
+
 private:
 	TSharedPtr<FPropertyEditor> PropertyEditor;
 };
