Passing in an FString for the SPropertyEditorAsset's ObjectPath caused the path not to update when the underlying UMaterialInterface was changed. Replacing it with a delegate fixes the most visible symptom of this bug.

diff --git a/Engine/Source/Editor/PropertyEditor/Private/PropertyCustomizationHelpers.cpp b/Engine/Source/Editor/PropertyEditor/Private/PropertyCustomizationHelpers.cpp
--- a/Engine/Source/Editor/PropertyEditor/Private/PropertyCustomizationHelpers.cpp
+++ b/Engine/Source/Editor/PropertyEditor/Private/PropertyCustomizationHelpers.cpp
@@ -838,7 +838,7 @@ public:
 					.FillWidth(1.0f)
 					[
 						SNew( SObjectPropertyEntryBox )
-						.ObjectPath(MaterialItem.Material->GetPathName())
+						.ObjectPath(this, &FMaterialItemView::OnGetObjectPath)
 						.AllowedClass(UMaterialInterface::StaticClass())
 						.OnObjectChanged(this, &FMaterialItemView::OnSetObject)
 						.ThumbnailPool(ThumbnailPool)
@@ -929,6 +929,11 @@ private:
 		}
 	}
 
+	FString OnGetObjectPath() const
+	{
+		return MaterialItem.Material->GetPathName();
+	}
+
 	void OnSetObject( const FAssetData& AssetData )
 	{
 		const bool bReplaceAll = false;
