Missing initialization of ShowCustomFilterOption in constructor causing an empty menu entry to show up sometimes in the detail view.

diff --git a/Engine/Source/Editor/PropertyEditor/Public/IDetailsView.h b/Engine/Source/Editor/PropertyEditor/Public/IDetailsView.h
--- a/Engine/Source/Editor/PropertyEditor/Public/IDetailsView.h
+++ b/Engine/Source/Editor/PropertyEditor/Public/IDetailsView.h
@@ -131,6 +131,7 @@ public:
 		, bForceHiddenPropertyVisibility(false)
 		, bShowKeyablePropertiesOption(true)
 		, bShowAnimatedPropertiesOption(true)
+		, bShowCustomFilterOption(false)
 		, ColumnWidth(.65f)
 	{
 	}
