Fixed issue where properties in sub-categories wouldn't appear in the details panel for an EditInlineNew object

diff --git a/Engine/Source/Editor/PropertyEditor/Private/DetailPropertyRow.cpp b/Engine/Source/Editor/PropertyEditor/Private/DetailPropertyRow.cpp
--- a/Engine/Source/Editor/PropertyEditor/Private/DetailPropertyRow.cpp
+++ b/Engine/Source/Editor/PropertyEditor/Private/DetailPropertyRow.cpp
@@ -261,6 +261,11 @@ void FDetailPropertyRow::OnGenerateChildren( FDetailNodeList& OutChildren )
 			const bool bIgnoreAdvancedDropdown = true;
 			MyCategory->GetGeneratedChildren(OutChildren, bIgnoreVisibility, bIgnoreAdvancedDropdown);
 		}
+		else
+		{
+			// Fall back to the default if we can't find the category implementation
+			GenerateChildrenForPropertyNode(PropertyNode, OutChildren);
+		}
 	}
 	else if (PropertyNode->AsCategoryNode() || PropertyNode->GetProperty() || ExternalObjectLayout.IsValid())
 	{
