Change to MakeNewBlueprint in the editor to automatically set the right interface to implement when the property as the MetaTag MustImplement.

diff --git a/Engine/Source/Editor/PropertyEditor/Private/Presentation/PropertyEditor/PropertyEditor.cpp b/Engine/Source/Editor/PropertyEditor/Private/Presentation/PropertyEditor/PropertyEditor.cpp
--- a/Engine/Source/Editor/PropertyEditor/Private/Presentation/PropertyEditor/PropertyEditor.cpp
+++ b/Engine/Source/Editor/PropertyEditor/Private/Presentation/PropertyEditor/PropertyEditor.cpp
@@ -16,6 +16,7 @@
 #include "Kismet2/KismetEditorUtilities.h"
 #include "IConfigEditorModule.h"
 #include "PropertyNode.h"
+#include "Kismet2/BlueprintEditorUtils.h"
 
 #define LOCTEXT_NAMESPACE "PropertyEditor"
 
@@ -274,12 +275,19 @@ void FPropertyEditor::MakeNewBlueprint()
 	UClassProperty* ClassProp = Cast<UClassProperty>(NodeProperty);
 	UClass* Class = (ClassProp ? ClassProp->MetaClass : FEditorClassUtils::GetClassFromString(NodeProperty->GetMetaData("MetaClass")));
 
+	UClass* RequiredInterface = FEditorClassUtils::GetClassFromString(NodeProperty->GetMetaData("MustImplement"));
+
 	if (Class)
 	{
 		UBlueprint* Blueprint = FKismetEditorUtilities::CreateBlueprintFromClass(LOCTEXT("CreateNewBlueprint", "Create New Blueprint"), Class, FString::Printf(TEXT("New%s"),*Class->GetName()));
 
 		if(Blueprint != NULL && Blueprint->GeneratedClass)
 		{
+			if (RequiredInterface != nullptr && FKismetEditorUtilities::CanBlueprintImplementInterface(Blueprint, RequiredInterface))
+			{
+				FBlueprintEditorUtils::ImplementNewInterface(Blueprint, RequiredInterface->GetFName());
+			}
+			
 			PropertyHandle->SetValueFromFormattedString(Blueprint->GeneratedClass->GetPathName());
 
 			FAssetEditorManager::Get().OpenEditorForAsset(Blueprint);
diff --git a/Engine/Source/Runtime/UMG/Public/Components/ListViewBase.h b/Engine/Source/Runtime/UMG/Public/Components/ListViewBase.h
--- a/Engine/Source/Runtime/UMG/Public/Components/ListViewBase.h
+++ b/Engine/Source/Runtime/UMG/Public/Components/ListViewBase.h
@@ -562,7 +562,7 @@ protected:
 
 	// Note: Options for this property can be configured via class and property metadata. See class declaration comment above.
 	/** The type of widget to create for each entry displayed in the list. */
-	UPROPERTY(EditAnywhere, BlueprintReadOnly, Category = ListEntries, meta = (DesignerRebuild, AllowPrivateAccess = true))
+	UPROPERTY(EditAnywhere, BlueprintReadOnly, Category = ListEntries, meta = (DesignerRebuild, AllowPrivateAccess = true, MustImplement = UserListEntry))
 	TSubclassOf<UUserWidget> EntryWidgetClass;
 
 private:
