Cheap fix for oscillating SPropertyValueWidget width while dragging a UMG widget in the designer.
The way drag&drop is implemented (SDesignerView::ProcessDropAndAddWidget), a new UCanvasPanelSlot gets recreated every frame, so the details panel gets refreshed every frame.
Since new property rows are created before old ones are destroyed in the details panel, the HasCustomResetToDefault flag on the property node toggles from frame to frame, so we alternate between having a ResetToDefaultPropertyEditor widget and not having one.
By having a spacer fill the blank, the property row layout doesn't change while dragging, but we still see a flashing yellow reset arrow (when visible).

diff --git a/Engine/Source/Editor/PropertyEditor/Private/DetailPropertyRow.cpp b/Engine/Source/Editor/PropertyEditor/Private/DetailPropertyRow.cpp
--- a/Engine/Source/Editor/PropertyEditor/Private/DetailPropertyRow.cpp
+++ b/Engine/Source/Editor/PropertyEditor/Private/DetailPropertyRow.cpp
@@ -13,6 +13,7 @@
 #include "ObjectPropertyNode.h"
 #include "DetailWidgetRow.h"
 #include "CategoryPropertyNode.h"
+#include "Widgets/Layout/SSpacer.h"
 
 const float FDetailWidgetRow::DefaultValueMinWidth = 125.0f;
 const float FDetailWidgetRow::DefaultValueMaxWidth = 125.0f;
@@ -695,15 +696,30 @@ void FDetailPropertyRow::MakeValueWidget( FDetailWidgetRow& Row, const TSharedPt
 		.IsEnabled( IsEnabledAttrib );
 
 	TSharedPtr<SResetToDefaultPropertyEditor> ResetButton = nullptr;
-	if (!PropertyHandle->HasMetaData(TEXT("NoResetToDefault")) && !PropertyHandle->IsResetToDefaultCustomized())
+	TSharedPtr<SWidget> ResetWidget = nullptr;
+	if (!PropertyHandle->HasMetaData(TEXT("NoResetToDefault")))
 	{
-		SAssignNew(ResetButton, SResetToDefaultPropertyEditor, PropertyEditor->GetPropertyHandle())
-			.IsEnabled(IsEnabledAttrib)
-			.CustomResetToDefault(CustomResetToDefault);
+		if (PropertyHandle->IsResetToDefaultCustomized())
+		{
+			// FIXME: Workaround for JIRA UE-73210.
+			// We had an oscillating SPropertyValueWidget width while dragging a UMG widget in the designer.
+			// The way drag&drop is implemented (SDesignerView::ProcessDropAndAddWidget), a new UCanvasPanelSlot gets 
+			// recreated every frame, so the details panel gets refreshed every frame. Since new property rows are created 
+			// before old ones are destroyed in the details panel, the HasCustomResetToDefault flag on the property node 
+			// toggles from frame to frame, so we alternate between having a ResetToDefaultPropertyEditor and not having one.
+			// By having a spacer fill the blank, the property row layout doesn't change while dragging, but we still see 
+			// a flashing yellow reset arrow (when visible).
+			const FSlateBrush* DiffersFromDefaultBrush = FEditorStyle::GetBrush("PropertyWindow.DiffersFromDefault");
+			ResetWidget = SNew(SSpacer).Size(DiffersFromDefaultBrush != nullptr ? DiffersFromDefaultBrush->ImageSize : FVector2D(8.0f, 8.0f));
+		}
+		else
+		{
+			SAssignNew(ResetButton, SResetToDefaultPropertyEditor, PropertyEditor->GetPropertyHandle())
+				.IsEnabled(IsEnabledAttrib)
+				.CustomResetToDefault(CustomResetToDefault);
+			ResetWidget = ResetButton;
+		}
 	};
-	
-
-	TSharedRef<SWidget> ResetWidget = ResetButton.IsValid() ? ResetButton.ToSharedRef() : SNullWidget::NullWidget;
 
 	TSharedPtr<SPropertyValueWidget> PropertyValue;
 
@@ -723,7 +739,7 @@ void FDetailPropertyRow::MakeValueWidget( FDetailWidgetRow& Row, const TSharedPt
 		[
 			SAssignNew( PropertyValue, SPropertyValueWidget, PropertyEditor, GetPropertyUtilities() )
 			.ShowPropertyButtons( false ) // We handle this ourselves
-			.OptionalResetWidget(ResetWidget)
+			.OptionalResetWidget(ResetButton.IsValid() ? ResetButton.ToSharedRef() : SNullWidget::NullWidget)
 		];
 		MinWidth = PropertyValue->GetMinDesiredWidth();
 		MaxWidth = PropertyValue->GetMaxDesiredWidth();
@@ -762,7 +778,7 @@ void FDetailPropertyRow::MakeValueWidget( FDetailWidgetRow& Row, const TSharedPt
 		}
 
 		if ((!PropertyValue.IsValid() || (PropertyValue.IsValid() && !PropertyValue->CreatedResetButton()))
-			&& ResetButton.IsValid())
+			&& ResetWidget.IsValid())
 		{
 			ValueWidget->AddSlot()
 			.Padding(4.0f, 0.0f)
@@ -770,7 +786,7 @@ void FDetailPropertyRow::MakeValueWidget( FDetailWidgetRow& Row, const TSharedPt
 			.VAlign(VAlign_Center)
 			.HAlign(HAlign_Left)
 			[
-				ResetWidget
+				ResetWidget.ToSharedRef()
 			];
 		}
 	}
