PR #5858: Fixed Actor Picker not filtering by Allowed or Disallowed classes (Contributed by [REDACTED_PERSON])

diff --git a/Engine/Source/Editor/PropertyEditor/Private/UserInterface/PropertyEditor/SPropertyEditorAsset.cpp b/Engine/Source/Editor/PropertyEditor/Private/UserInterface/PropertyEditor/SPropertyEditorAsset.cpp
--- a/Engine/Source/Editor/PropertyEditor/Private/UserInterface/PropertyEditor/SPropertyEditorAsset.cpp
+++ b/Engine/Source/Editor/PropertyEditor/Private/UserInterface/PropertyEditor/SPropertyEditorAsset.cpp
@@ -680,7 +680,33 @@ void SPropertyEditorAsset::OnMenuOpenChanged(bool bOpen)
 
 bool SPropertyEditorAsset::IsFilteredActor( const AActor* const Actor ) const
 {
-	return Actor->IsA( ObjectClass ) && !Actor->IsChildActor();
+	bool IsAllowed = Actor->IsA(ObjectClass) && !Actor->IsChildActor();
+
+	if (IsAllowed)
+	{
+		bool ClassFilterAllowed = false;
+		for (const UClass* Class : AllowedClassFilters)
+		{
+			if (Actor->GetClass()->IsChildOf(Class))
+			{
+				ClassFilterAllowed = true;
+				break;
+			}
+		}
+
+		for (const UClass* Class : DisallowedClassFilters)
+		{
+			if (Actor->GetClass()->IsChildOf(Class))
+			{
+				ClassFilterAllowed = false;
+				break;
+			}
+		}
+
+		IsAllowed = ClassFilterAllowed;
+	}
+
+	return IsAllowed;
 }
 
 void SPropertyEditorAsset::CloseComboButton()
