Moved Material property customization to MaterialList.h/cpp .

diff --git a/Engine/Source/Editor/DetailCustomizations/Private/ComponentMaterialCategory.cpp b/Engine/Source/Editor/DetailCustomizations/Private/ComponentMaterialCategory.cpp
--- a/Engine/Source/Editor/DetailCustomizations/Private/ComponentMaterialCategory.cpp
+++ b/Engine/Source/Editor/DetailCustomizations/Private/ComponentMaterialCategory.cpp
@@ -15,7 +15,7 @@
 #include "Editor.h"
 #include "DetailLayoutBuilder.h"
 #include "DetailCategoryBuilder.h"
-#include "PropertyCustomizationHelpers.h"
+#include "MaterialList.h"
 #include "IPropertyUtilities.h"
 
 #include "LandscapeProxy.h"
diff --git a/Engine/Source/Editor/LevelEditor/Private/LevelEditorGenericDetails.cpp b/Engine/Source/Editor/LevelEditor/Private/LevelEditorGenericDetails.cpp
--- a/Engine/Source/Editor/LevelEditor/Private/LevelEditorGenericDetails.cpp
+++ b/Engine/Source/Editor/LevelEditor/Private/LevelEditorGenericDetails.cpp
@@ -22,8 +22,7 @@
 #include "SurfaceIterators.h"
 #include "ScopedTransaction.h"
 #include "FunctionalTest.h"
-
-#include "PropertyCustomizationHelpers.h"
+#include "MaterialList.h"
 
 #define LOCTEXT_NAMESPACE "FLevelEditorGenericDetails"
 
diff --git a/Engine/Source/Editor/Persona/Private/PersonaMeshDetails.cpp b/Engine/Source/Editor/Persona/Private/PersonaMeshDetails.cpp
--- a/Engine/Source/Editor/Persona/Private/PersonaMeshDetails.cpp
+++ b/Engine/Source/Editor/Persona/Private/PersonaMeshDetails.cpp
@@ -32,6 +32,7 @@
 #include "EditorDirectories.h"
 #include "UnrealEdGlobals.h"
 #include "IDetailsView.h"
+#include "MaterialList.h"
 #include "PropertyCustomizationHelpers.h"
 #include "DesktopPlatformModule.h"
 #include "Interfaces/IMainFrameModule.h"
diff --git a/Engine/Source/Editor/PropertyEditor/Private/MaterialList.cpp b/Engine/Source/Editor/PropertyEditor/Private/MaterialList.cpp
new file mode 100644
--- /dev/null
+++ b/Engine/Source/Editor/PropertyEditor/Private/MaterialList.cpp
@@ -0,0 +1,633 @@
+// Copyright 1998-2019 Epic Games, Inc. All Rights Reserved.
+
+#include "MaterialList.h"
+#include "DetailLayoutBuilder.h"
+#include "IDetailChildrenBuilder.h"
+#include "PropertyCustomizationHelpers.h"
+#include "PropertyHandle.h"
+#include "Widgets/Input/SHyperlink.h"
+
+#define LOCTEXT_NAMESPACE "PropertyCustomizationHelpers"
+
+/**
+ * Builds up a list of unique materials while creating some information about the materials
+ */
+class FMaterialListBuilder : public IMaterialListBuilder
+{
+	friend class FMaterialList;
+public:
+
+	/** 
+	 * Adds a new material to the list
+	 * 
+	 * @param SlotIndex		The slot (usually mesh element index) where the material is located on the component
+	 * @param Material		The material being used
+	 * @param bCanBeReplced	Whether or not the material can be replaced by a user
+	 */
+	virtual void AddMaterial( uint32 SlotIndex, UMaterialInterface* Material, bool bCanBeReplaced ) override
+	{
+		int32 NumMaterials = MaterialSlots.Num();
+
+		FMaterialListItem MaterialItem( Material, SlotIndex, bCanBeReplaced ); 
+		if( !UniqueMaterials.Contains( MaterialItem ) ) 
+		{
+			MaterialSlots.Add( MaterialItem );
+			UniqueMaterials.Add( MaterialItem );
+		}
+
+		// Did we actually add material?  If we did then we need to increment the number of materials in the element
+		if( MaterialSlots.Num() > NumMaterials )
+		{
+			// Resize the array to support the slot if needed
+			if( !MaterialCount.IsValidIndex(SlotIndex) )
+			{
+				int32 NumToAdd = (SlotIndex - MaterialCount.Num()) + 1;
+				if( NumToAdd > 0 )
+				{
+					MaterialCount.AddZeroed( NumToAdd );
+				}
+			}
+
+			++MaterialCount[SlotIndex];
+		}
+	}
+
+	/** Empties the list */
+	void Empty()
+	{
+		UniqueMaterials.Empty();
+		MaterialSlots.Reset();
+		MaterialCount.Reset();
+	}
+
+	/** Sorts the list by slot index */
+	void Sort()
+	{
+		struct FSortByIndex
+		{
+			bool operator()( const FMaterialListItem& A, const FMaterialListItem& B ) const
+			{
+				return A.SlotIndex < B.SlotIndex;
+			}
+		};
+
+		MaterialSlots.Sort( FSortByIndex() );
+	}
+
+	/** @return The number of materials in the list */
+	uint32 GetNumMaterials() const { return MaterialSlots.Num(); }
+
+	/** @return The number of materials in the list at a given slot */
+	uint32 GetNumMaterialsInSlot( uint32 Index ) const { return MaterialCount[Index]; }
+private:
+	/** All unique materials */
+	TSet<FMaterialListItem> UniqueMaterials;
+	/** All material items in the list */
+	TArray<FMaterialListItem> MaterialSlots;
+	/** Material counts for each slot.  The slot is the index and the value at that index is the count */
+	TArray<uint32> MaterialCount;
+};
+
+/**
+ * A view of a single item in an FMaterialList
+ */
+class FMaterialItemView : public TSharedFromThis<FMaterialItemView>
+{
+public:
+	/**
+	 * Creates a new instance of this class
+	 *
+	 * @param Material				The material to view
+	 * @param InOnMaterialChanged	Delegate for when the material changes
+	 */
+	static TSharedRef<FMaterialItemView> Create(
+		const FMaterialListItem& Material, 
+		FOnMaterialChanged InOnMaterialChanged,
+		FOnGenerateWidgetsForMaterial InOnGenerateNameWidgetsForMaterial, 
+		FOnGenerateWidgetsForMaterial InOnGenerateWidgetsForMaterial, 
+		FOnResetMaterialToDefaultClicked InOnResetToDefaultClicked,
+		int32 InMultipleMaterialCount,
+		bool bShowUsedTextures,
+		bool bDisplayCompactSize)
+	{
+		return MakeShareable( new FMaterialItemView( Material, InOnMaterialChanged, InOnGenerateNameWidgetsForMaterial, InOnGenerateWidgetsForMaterial, InOnResetToDefaultClicked, InMultipleMaterialCount, bShowUsedTextures, bDisplayCompactSize) );
+	}
+
+	TSharedRef<SWidget> CreateNameContent()
+	{
+		FFormatNamedArguments Arguments;
+		Arguments.Add(TEXT("ElementIndex"), MaterialItem.SlotIndex);
+
+		return 
+			SNew(SVerticalBox)
+			+SVerticalBox::Slot()
+			.VAlign(VAlign_Center)
+			[
+				SNew( STextBlock )
+				.Font( IDetailLayoutBuilder::GetDetailFont() )
+				.Text( FText::Format(LOCTEXT("ElementIndex", "Element {ElementIndex}"), Arguments ) )
+			]
+			+SVerticalBox::Slot()
+			.Padding(0.0f,4.0f)
+			.AutoHeight()
+			[
+				OnGenerateCustomNameWidgets.IsBound() ? OnGenerateCustomNameWidgets.Execute( MaterialItem.Material.Get(), MaterialItem.SlotIndex ) : StaticCastSharedRef<SWidget>( SNullWidget::NullWidget )
+			];
+	}
+
+	TSharedRef<SWidget> CreateValueContent( const TSharedPtr<FAssetThumbnailPool>& ThumbnailPool )
+	{
+		FIntPoint ThumbnailSize(64, 64);
+
+		FResetToDefaultOverride ResetToDefaultOverride = FResetToDefaultOverride::Create(
+			FIsResetToDefaultVisible::CreateSP(this, &FMaterialItemView::GetReplaceVisibility),
+			FResetToDefaultHandler::CreateSP(this, &FMaterialItemView::OnResetToBaseClicked)
+		);
+
+		return
+			SNew(SHorizontalBox)
+			+ SHorizontalBox::Slot()
+			[
+				SNew(SVerticalBox)
+				+SVerticalBox::Slot()
+				.AutoHeight()
+				.Padding( 0.0f )
+				.VAlign(VAlign_Center)
+				.HAlign(HAlign_Fill)
+				[
+					SNew(SHorizontalBox)
+					+SHorizontalBox::Slot()
+					.FillWidth(1.0f)
+					[
+						SNew( SObjectPropertyEntryBox )
+						.ObjectPath(this, &FMaterialItemView::OnGetObjectPath)
+						.AllowedClass(UMaterialInterface::StaticClass())
+						.OnObjectChanged(this, &FMaterialItemView::OnSetObject)
+						.ThumbnailPool(ThumbnailPool)
+						.DisplayCompactSize(bDisplayCompactSize)
+						.CustomResetToDefault(ResetToDefaultOverride)
+						.CustomContentSlot()
+						[
+							SNew( SBox )
+							.HAlign(HAlign_Left)
+							.VAlign(VAlign_Center)
+							[
+								SNew(SHorizontalBox)
+								+SHorizontalBox::Slot()
+								.VAlign(VAlign_Center)
+								.Padding(0.0f, 0.0f, 3.0f, 0.0f)
+								.AutoWidth()
+								[
+									// Add a menu for displaying all textures 
+									SNew( SComboButton )
+									.OnGetMenuContent( this, &FMaterialItemView::OnGetTexturesMenuForMaterial )
+									.VAlign(VAlign_Center)
+									.ContentPadding(2)
+									.IsEnabled( this, &FMaterialItemView::IsTexturesMenuEnabled )
+									.Visibility( bShowUsedTextures ? EVisibility::Visible : EVisibility::Hidden )
+									.ButtonContent()
+									[
+										SNew( STextBlock )
+										.Font( IDetailLayoutBuilder::GetDetailFont() )
+										.ToolTipText( LOCTEXT("ViewTexturesToolTip", "View the textures used by this material" ) )
+										.Text( LOCTEXT("ViewTextures","Textures") )
+									]
+								]
+								+SHorizontalBox::Slot()
+								.Padding(3.0f, 0.0f)
+								.FillWidth(1.0f)
+								[
+									OnGenerateCustomMaterialWidgets.IsBound() && bDisplayCompactSize ? OnGenerateCustomMaterialWidgets.Execute(MaterialItem.Material.Get(), MaterialItem.SlotIndex) : StaticCastSharedRef<SWidget>(SNullWidget::NullWidget)
+								]
+							]
+						]
+					]
+				]
+				+SVerticalBox::Slot()
+				.AutoHeight()
+				.Padding(2)
+				.VAlign( VAlign_Center )
+				[
+					OnGenerateCustomMaterialWidgets.IsBound() && !bDisplayCompactSize ? OnGenerateCustomMaterialWidgets.Execute( MaterialItem.Material.Get(), MaterialItem.SlotIndex ) : StaticCastSharedRef<SWidget>( SNullWidget::NullWidget )
+				]
+			];
+	}
+
+private:
+
+	FMaterialItemView(	const FMaterialListItem& InMaterial, 
+						FOnMaterialChanged& InOnMaterialChanged, 
+						FOnGenerateWidgetsForMaterial& InOnGenerateNameWidgets, 
+						FOnGenerateWidgetsForMaterial& InOnGenerateMaterialWidgets, 
+						FOnResetMaterialToDefaultClicked& InOnResetToDefaultClicked,
+						int32 InMultipleMaterialCount,
+						bool bInShowUsedTextures,
+						bool bInDisplayCompactSize)
+						
+		: MaterialItem( InMaterial )
+		, OnMaterialChanged( InOnMaterialChanged )
+		, OnGenerateCustomNameWidgets( InOnGenerateNameWidgets )
+		, OnGenerateCustomMaterialWidgets( InOnGenerateMaterialWidgets )
+		, OnResetToDefaultClicked( InOnResetToDefaultClicked )
+		, MultipleMaterialCount( InMultipleMaterialCount )
+		, bShowUsedTextures( bInShowUsedTextures )
+		, bDisplayCompactSize(bInDisplayCompactSize)
+	{
+
+	}
+
+	void ReplaceMaterial( UMaterialInterface* NewMaterial, bool bReplaceAll = false )
+	{
+		UMaterialInterface* PrevMaterial = NULL;
+		if( MaterialItem.Material.IsValid() )
+		{
+			PrevMaterial = MaterialItem.Material.Get();
+		}
+
+		if( NewMaterial != PrevMaterial )
+		{
+			// Replace the material
+			OnMaterialChanged.ExecuteIfBound( NewMaterial, PrevMaterial, MaterialItem.SlotIndex, bReplaceAll );
+		}
+	}
+
+	void OnSetObject( const FAssetData& AssetData )
+	{
+		const bool bReplaceAll = false;
+
+		UMaterialInterface* NewMaterial = Cast<UMaterialInterface>(AssetData.GetAsset());
+		ReplaceMaterial( NewMaterial, bReplaceAll );
+	}
+
+	FString OnGetObjectPath() const
+	{
+		return MaterialItem.Material->GetPathName();
+	}
+
+	/**
+	 * @return Whether or not the textures menu is enabled
+	 */
+	bool IsTexturesMenuEnabled() const
+	{
+		return MaterialItem.Material.Get() != NULL;
+	}
+
+	TSharedRef<SWidget> OnGetTexturesMenuForMaterial()
+	{
+		FMenuBuilder MenuBuilder( true, NULL );
+
+		if( MaterialItem.Material.IsValid() )
+		{
+			UMaterialInterface* Material = MaterialItem.Material.Get();
+
+			TArray< UTexture* > Textures;
+			Material->GetUsedTextures(Textures, EMaterialQualityLevel::Num, false, ERHIFeatureLevel::Num, true);
+
+			// Add a menu item for each texture.  Clicking on the texture will display it in the content browser
+			// UObject for delegate compatibility
+			for( UObject* Texture : Textures )
+			{
+				FUIAction Action( FExecuteAction::CreateSP( this, &FMaterialItemView::GoToAssetInContentBrowser, MakeWeakObjectPtr(Texture) ) );
+
+				MenuBuilder.AddMenuEntry( FText::FromString( Texture->GetName() ), LOCTEXT( "BrowseTexture_ToolTip", "Find this texture in the content browser" ), FSlateIcon(), Action );
+			}
+		}
+
+		return MenuBuilder.MakeWidget();
+	}
+
+	/**
+	 * Finds the asset in the content browser
+	 */
+	void GoToAssetInContentBrowser( TWeakObjectPtr<UObject> Object )
+	{
+		if( Object.IsValid() )
+		{
+			TArray< UObject* > Objects;
+			Objects.Add( Object.Get() );
+			GEditor->SyncBrowserToObjects( Objects );
+		}
+	}
+
+	/**
+	 * Called to get the visibility of the replace button
+	 */
+	bool GetReplaceVisibility(TSharedPtr<IPropertyHandle> PropertyHandle) const
+	{
+		// Only show the replace button if the current material can be replaced
+		if (OnMaterialChanged.IsBound() && MaterialItem.bCanBeReplaced)
+		{
+			return true;
+		}
+
+		return false;
+	}
+
+	/**
+	 * Called when reset to base is clicked
+	 */
+	void OnResetToBaseClicked(TSharedPtr<IPropertyHandle> PropertyHandle)
+	{
+		// Only allow reset to base if the current material can be replaced
+		if( MaterialItem.Material.IsValid() && MaterialItem.bCanBeReplaced )
+		{
+			bool bReplaceAll = false;
+			ReplaceMaterial( NULL, bReplaceAll );
+			OnResetToDefaultClicked.ExecuteIfBound( MaterialItem.Material.Get(), MaterialItem.SlotIndex );
+		}
+	}
+
+private:
+	FMaterialListItem MaterialItem;
+	FOnMaterialChanged OnMaterialChanged;
+	FOnGenerateWidgetsForMaterial OnGenerateCustomNameWidgets;
+	FOnGenerateWidgetsForMaterial OnGenerateCustomMaterialWidgets;
+	FOnResetMaterialToDefaultClicked OnResetToDefaultClicked;
+	int32 MultipleMaterialCount;
+	bool bShowUsedTextures;
+	bool bDisplayCompactSize;
+};
+
+
+FMaterialList::FMaterialList(IDetailLayoutBuilder& InDetailLayoutBuilder, FMaterialListDelegates& InMaterialListDelegates, bool bInAllowCollapse, bool bInShowUsedTextures, bool bInDisplayCompactSize)
+	: MaterialListDelegates( InMaterialListDelegates )
+	, DetailLayoutBuilder( InDetailLayoutBuilder )
+	, MaterialListBuilder( new FMaterialListBuilder )
+	, bAllowCollpase(bInAllowCollapse)
+	, bShowUsedTextures(bInShowUsedTextures)
+	, bDisplayCompactSize(bInDisplayCompactSize)
+{
+}
+
+void FMaterialList::OnDisplayMaterialsForElement( int32 SlotIndex )
+{
+	// We now want to display all the materials in the element
+	ExpandedSlots.Add( SlotIndex );
+
+	MaterialListBuilder->Empty();
+	MaterialListDelegates.OnGetMaterials.ExecuteIfBound( *MaterialListBuilder );
+
+	OnRebuildChildren.ExecuteIfBound();
+}
+
+void FMaterialList::OnHideMaterialsForElement( int32 SlotIndex )
+{
+	// No longer want to expand the element
+	ExpandedSlots.Remove( SlotIndex );
+
+	// regenerate the materials
+	MaterialListBuilder->Empty();
+	MaterialListDelegates.OnGetMaterials.ExecuteIfBound( *MaterialListBuilder );
+	
+	OnRebuildChildren.ExecuteIfBound();
+}
+
+
+void FMaterialList::Tick( float DeltaTime )
+{
+	// Check each material to see if its still valid.  This allows the material list to stay up to date when materials are changed out from under us
+	if( MaterialListDelegates.OnGetMaterials.IsBound() )
+	{
+		// Whether or not to refresh the material list
+		bool bRefreshMaterialList = false;
+
+		// Get the current list of materials from the user
+		MaterialListBuilder->Empty();
+		MaterialListDelegates.OnGetMaterials.ExecuteIfBound( *MaterialListBuilder );
+
+		if( MaterialListBuilder->GetNumMaterials() != DisplayedMaterials.Num() )
+		{
+			// The array sizes differ so we need to refresh the list
+			bRefreshMaterialList = true;
+		}
+		else
+		{
+			// Compare the new list against the currently displayed list
+			for( int32 MaterialIndex = 0; MaterialIndex < MaterialListBuilder->MaterialSlots.Num(); ++MaterialIndex )
+			{
+				const FMaterialListItem& Item = MaterialListBuilder->MaterialSlots[MaterialIndex];
+
+				// The displayed materials is out of date if there isn't a 1:1 mapping between the material sets
+				if( !DisplayedMaterials.IsValidIndex( MaterialIndex ) || DisplayedMaterials[ MaterialIndex ] != Item )
+				{
+					bRefreshMaterialList = true;
+					break;
+				}
+			}
+		}
+
+		if (!bRefreshMaterialList && MaterialListDelegates.OnMaterialListDirty.IsBound())
+		{
+			bRefreshMaterialList = MaterialListDelegates.OnMaterialListDirty.Execute();
+		}
+
+		if( bRefreshMaterialList )
+		{
+			OnRebuildChildren.ExecuteIfBound();
+		}
+	}
+}
+
+void FMaterialList::GenerateHeaderRowContent( FDetailWidgetRow& NodeRow )
+{
+	NodeRow.CopyAction(FUIAction(FExecuteAction::CreateSP(this, &FMaterialList::OnCopyMaterialList), FCanExecuteAction::CreateSP(this, &FMaterialList::OnCanCopyMaterialList)));
+	NodeRow.PasteAction(FUIAction(FExecuteAction::CreateSP(this, &FMaterialList::OnPasteMaterialList)));
+
+	if (bAllowCollpase)
+	{
+		NodeRow.NameContent()
+			[
+				SNew(STextBlock)
+				.Text(LOCTEXT("MaterialHeaderTitle", "Materials"))
+				.Font(IDetailLayoutBuilder::GetDetailFont())
+			];
+	}
+}
+
+void FMaterialList::GenerateChildContent( IDetailChildrenBuilder& ChildrenBuilder )
+{
+	ViewedMaterials.Empty();
+	DisplayedMaterials.Empty();
+	if( MaterialListBuilder->GetNumMaterials() > 0 )
+	{
+		DisplayedMaterials = MaterialListBuilder->MaterialSlots;
+
+		MaterialListBuilder->Sort();
+		TArray<FMaterialListItem>& MaterialSlots = MaterialListBuilder->MaterialSlots;
+
+		int32 CurrentSlot = INDEX_NONE;
+		bool bDisplayAllMaterialsInSlot = true;
+		for( auto It = MaterialSlots.CreateConstIterator(); It; ++It )
+		{
+			const FMaterialListItem& Material = *It;
+
+			if( CurrentSlot != Material.SlotIndex )
+			{
+				// We've encountered a new slot.  Make a widget to display that
+				CurrentSlot = Material.SlotIndex;
+
+				uint32 NumMaterials = MaterialListBuilder->GetNumMaterialsInSlot(CurrentSlot);
+
+				// If an element is expanded we want to display all its materials
+				bool bWantToDisplayAllMaterials = NumMaterials > 1 && ExpandedSlots.Contains(CurrentSlot);
+
+				// If we are currently displaying an expanded set of materials for an element add a link to collapse all of them
+				if( bWantToDisplayAllMaterials )
+				{
+					FDetailWidgetRow& ChildRow = ChildrenBuilder.AddCustomRow( LOCTEXT( "HideAllMaterialSearchString", "Hide All Materials") );
+
+					FFormatNamedArguments Arguments;
+					Arguments.Add(TEXT("ElementSlot"), CurrentSlot);
+					ChildRow
+					.ValueContent()
+					.MaxDesiredWidth(0.0f)// No Max Width
+					[
+						SNew( SBox )
+							.HAlign( HAlign_Center )
+							[
+								SNew( SHyperlink )
+									.TextStyle( FEditorStyle::Get(), "MaterialList.HyperlinkStyle" )
+									.Text( FText::Format(LOCTEXT("HideAllMaterialLinkText", "Hide All Materials on Element {ElementSlot}"), Arguments ) )
+									.OnNavigate( this, &FMaterialList::OnHideMaterialsForElement, CurrentSlot )
+							]
+					];
+				}	
+
+				if( NumMaterials > 1 && !bWantToDisplayAllMaterials )
+				{
+					// The current slot has multiple elements to view
+					bDisplayAllMaterialsInSlot = false;
+
+					FDetailWidgetRow& ChildRow = ChildrenBuilder.AddCustomRow( FText::GetEmpty() );
+
+					AddMaterialItem( ChildRow, CurrentSlot, FMaterialListItem( NULL, CurrentSlot, true ), !bDisplayAllMaterialsInSlot );
+				}
+				else
+				{
+					bDisplayAllMaterialsInSlot = true;
+				}
+
+			}
+
+			// Display each thumbnail element unless we shouldn't display multiple materials for one slot
+			if( bDisplayAllMaterialsInSlot )
+			{
+				FDetailWidgetRow& ChildRow = ChildrenBuilder.AddCustomRow( Material.Material.IsValid()? FText::FromString(Material.Material->GetName()) : FText::GetEmpty() );
+
+				AddMaterialItem( ChildRow, CurrentSlot, Material, !bDisplayAllMaterialsInSlot );
+			}
+		}
+	}
+	else
+	{
+		FDetailWidgetRow& ChildRow = ChildrenBuilder.AddCustomRow( LOCTEXT("NoMaterials", "No Materials") );
+
+		ChildRow
+		[
+			SNew( SBox )
+			.HAlign( HAlign_Center )
+			[
+				SNew( STextBlock )
+				.Text( LOCTEXT("NoMaterials", "No Materials") ) 
+				.Font( IDetailLayoutBuilder::GetDetailFont() )
+			]
+		];
+	}		
+}
+
+bool FMaterialList::OnCanCopyMaterialList() const
+{
+	if (MaterialListDelegates.OnCanCopyMaterialList.IsBound())
+	{
+		return MaterialListDelegates.OnCanCopyMaterialList.Execute();
+	}
+
+	return false;
+}
+
+void FMaterialList::OnCopyMaterialList()
+{
+	if (MaterialListDelegates.OnCopyMaterialList.IsBound())
+	{
+		MaterialListDelegates.OnCopyMaterialList.Execute();
+	}
+}
+
+void FMaterialList::OnPasteMaterialList()
+{
+	if (MaterialListDelegates.OnPasteMaterialList.IsBound())
+	{
+		MaterialListDelegates.OnPasteMaterialList.Execute();
+	}
+}
+
+bool FMaterialList::OnCanCopyMaterialItem(int32 CurrentSlot) const
+{
+	if (MaterialListDelegates.OnCanCopyMaterialItem.IsBound())
+	{
+		return MaterialListDelegates.OnCanCopyMaterialItem.Execute(CurrentSlot);
+	}
+
+	return false;
+}
+
+void FMaterialList::OnCopyMaterialItem(int32 CurrentSlot)
+{
+	if (MaterialListDelegates.OnCopyMaterialItem.IsBound())
+	{
+		MaterialListDelegates.OnCopyMaterialItem.Execute(CurrentSlot);
+	}
+}
+
+void FMaterialList::OnPasteMaterialItem(int32 CurrentSlot)
+{
+	if (MaterialListDelegates.OnPasteMaterialItem.IsBound())
+	{
+		MaterialListDelegates.OnPasteMaterialItem.Execute(CurrentSlot);
+	}
+}
+			
+void FMaterialList::AddMaterialItem( FDetailWidgetRow& Row, int32 CurrentSlot, const FMaterialListItem& Item, bool bDisplayLink )
+{
+	uint32 NumMaterials = MaterialListBuilder->GetNumMaterialsInSlot(CurrentSlot);
+
+	TSharedRef<FMaterialItemView> NewView = FMaterialItemView::Create( Item, MaterialListDelegates.OnMaterialChanged, MaterialListDelegates.OnGenerateCustomNameWidgets, MaterialListDelegates.OnGenerateCustomMaterialWidgets, MaterialListDelegates.OnResetMaterialToDefaultClicked, NumMaterials, bShowUsedTextures, bDisplayCompactSize);
+
+	TSharedPtr<SWidget> RightSideContent;
+	if( bDisplayLink )
+	{
+		FFormatNamedArguments Arguments;
+		Arguments.Add(TEXT("NumMaterials"), NumMaterials);
+
+		RightSideContent = 
+			SNew( SBox )
+				.HAlign(HAlign_Left)
+				.VAlign(VAlign_Top)
+				[
+					SNew( SHyperlink )
+					.TextStyle( FEditorStyle::Get(), "MaterialList.HyperlinkStyle" )
+					.Text( FText::Format(LOCTEXT("DisplayAllMaterialLinkText", "Display {NumMaterials} materials"), Arguments) )
+					.ToolTipText( LOCTEXT("DisplayAllMaterialLink_ToolTip","Display all materials. Drag and drop a material here to replace all materials.") )
+					.OnNavigate( this, &FMaterialList::OnDisplayMaterialsForElement, CurrentSlot )
+				];
+	}
+	else
+	{
+		RightSideContent = NewView->CreateValueContent( DetailLayoutBuilder.GetThumbnailPool() );
+		ViewedMaterials.Add( NewView );
+	}
+
+	Row.CopyAction(FUIAction(FExecuteAction::CreateSP(this, &FMaterialList::OnCopyMaterialItem, Item.SlotIndex), FCanExecuteAction::CreateSP(this, &FMaterialList::OnCanCopyMaterialItem, Item.SlotIndex)));
+	Row.PasteAction(FUIAction(FExecuteAction::CreateSP(this, &FMaterialList::OnPasteMaterialItem, Item.SlotIndex)));
+
+	Row.NameContent()
+	[
+		NewView->CreateNameContent()
+	]
+	.ValueContent()
+	.MinDesiredWidth(250.f)
+	.MaxDesiredWidth(0.0f) // no maximum
+	[
+		RightSideContent.ToSharedRef()
+	];
+}
diff --git a/Engine/Source/Editor/PropertyEditor/Private/PropertyCustomizationHelpers.cpp b/Engine/Source/Editor/PropertyEditor/Private/PropertyCustomizationHelpers.cpp
--- a/Engine/Source/Editor/PropertyEditor/Private/PropertyCustomizationHelpers.cpp
+++ b/Engine/Source/Editor/PropertyEditor/Private/PropertyCustomizationHelpers.cpp
@@ -687,629 +687,6 @@ bool SProperty::IsValidProperty() const
 	return PropertyHandle.IsValid() && PropertyHandle->IsValidHandle();
 }
 
-/**
- * Builds up a list of unique materials while creating some information about the materials
- */
-class FMaterialListBuilder : public IMaterialListBuilder
-{
-	friend class FMaterialList;
-public:
-
-	/** 
-	 * Adds a new material to the list
-	 * 
-	 * @param SlotIndex		The slot (usually mesh element index) where the material is located on the component
-	 * @param Material		The material being used
-	 * @param bCanBeReplced	Whether or not the material can be replaced by a user
-	 */
-	virtual void AddMaterial( uint32 SlotIndex, UMaterialInterface* Material, bool bCanBeReplaced ) override
-	{
-		int32 NumMaterials = MaterialSlots.Num();
-
-		FMaterialListItem MaterialItem( Material, SlotIndex, bCanBeReplaced ); 
-		if( !UniqueMaterials.Contains( MaterialItem ) ) 
-		{
-			MaterialSlots.Add( MaterialItem );
-			UniqueMaterials.Add( MaterialItem );
-		}
-
-		// Did we actually add material?  If we did then we need to increment the number of materials in the element
-		if( MaterialSlots.Num() > NumMaterials )
-		{
-			// Resize the array to support the slot if needed
-			if( !MaterialCount.IsValidIndex(SlotIndex) )
-			{
-				int32 NumToAdd = (SlotIndex - MaterialCount.Num()) + 1;
-				if( NumToAdd > 0 )
-				{
-					MaterialCount.AddZeroed( NumToAdd );
-				}
-			}
-
-			++MaterialCount[SlotIndex];
-		}
-	}
-
-	/** Empties the list */
-	void Empty()
-	{
-		UniqueMaterials.Empty();
-		MaterialSlots.Reset();
-		MaterialCount.Reset();
-	}
-
-	/** Sorts the list by slot index */
-	void Sort()
-	{
-		struct FSortByIndex
-		{
-			bool operator()( const FMaterialListItem& A, const FMaterialListItem& B ) const
-			{
-				return A.SlotIndex < B.SlotIndex;
-			}
-		};
-
-		MaterialSlots.Sort( FSortByIndex() );
-	}
-
-	/** @return The number of materials in the list */
-	uint32 GetNumMaterials() const { return MaterialSlots.Num(); }
-
-	/** @return The number of materials in the list at a given slot */
-	uint32 GetNumMaterialsInSlot( uint32 Index ) const { return MaterialCount[Index]; }
-private:
-	/** All unique materials */
-	TSet<FMaterialListItem> UniqueMaterials;
-	/** All material items in the list */
-	TArray<FMaterialListItem> MaterialSlots;
-	/** Material counts for each slot.  The slot is the index and the value at that index is the count */
-	TArray<uint32> MaterialCount;
-};
-
-/**
- * A view of a single item in an FMaterialList
- */
-class FMaterialItemView : public TSharedFromThis<FMaterialItemView>
-{
-public:
-	/**
-	 * Creates a new instance of this class
-	 *
-	 * @param Material				The material to view
-	 * @param InOnMaterialChanged	Delegate for when the material changes
-	 */
-	static TSharedRef<FMaterialItemView> Create(
-		const FMaterialListItem& Material, 
-		FOnMaterialChanged InOnMaterialChanged,
-		FOnGenerateWidgetsForMaterial InOnGenerateNameWidgetsForMaterial, 
-		FOnGenerateWidgetsForMaterial InOnGenerateWidgetsForMaterial, 
-		FOnResetMaterialToDefaultClicked InOnResetToDefaultClicked,
-		int32 InMultipleMaterialCount,
-		bool bShowUsedTextures,
-		bool bDisplayCompactSize)
-	{
-		return MakeShareable( new FMaterialItemView( Material, InOnMaterialChanged, InOnGenerateNameWidgetsForMaterial, InOnGenerateWidgetsForMaterial, InOnResetToDefaultClicked, InMultipleMaterialCount, bShowUsedTextures, bDisplayCompactSize) );
-	}
-
-	TSharedRef<SWidget> CreateNameContent()
-	{
-		FFormatNamedArguments Arguments;
-		Arguments.Add(TEXT("ElementIndex"), MaterialItem.SlotIndex);
-
-		return 
-			SNew(SVerticalBox)
-			+SVerticalBox::Slot()
-			.VAlign(VAlign_Center)
-			[
-				SNew( STextBlock )
-				.Font( IDetailLayoutBuilder::GetDetailFont() )
-				.Text( FText::Format(LOCTEXT("ElementIndex", "Element {ElementIndex}"), Arguments ) )
-			]
-			+SVerticalBox::Slot()
-			.Padding(0.0f,4.0f)
-			.AutoHeight()
-			[
-				OnGenerateCustomNameWidgets.IsBound() ? OnGenerateCustomNameWidgets.Execute( MaterialItem.Material.Get(), MaterialItem.SlotIndex ) : StaticCastSharedRef<SWidget>( SNullWidget::NullWidget )
-			];
-	}
-
-	TSharedRef<SWidget> CreateValueContent( const TSharedPtr<FAssetThumbnailPool>& ThumbnailPool )
-	{
-		FIntPoint ThumbnailSize(64, 64);
-
-		FResetToDefaultOverride ResetToDefaultOverride = FResetToDefaultOverride::Create(
-			FIsResetToDefaultVisible::CreateSP(this, &FMaterialItemView::GetReplaceVisibility),
-			FResetToDefaultHandler::CreateSP(this, &FMaterialItemView::OnResetToBaseClicked)
-		);
-
-		return
-			SNew(SHorizontalBox)
-			+ SHorizontalBox::Slot()
-			[
-				SNew(SVerticalBox)
-				+SVerticalBox::Slot()
-				.AutoHeight()
-				.Padding( 0.0f )
-				.VAlign(VAlign_Center)
-				.HAlign(HAlign_Fill)
-				[
-					SNew(SHorizontalBox)
-					+SHorizontalBox::Slot()
-					.FillWidth(1.0f)
-					[
-						SNew( SObjectPropertyEntryBox )
-						.ObjectPath(this, &FMaterialItemView::OnGetObjectPath)
-						.AllowedClass(UMaterialInterface::StaticClass())
-						.OnObjectChanged(this, &FMaterialItemView::OnSetObject)
-						.ThumbnailPool(ThumbnailPool)
-						.DisplayCompactSize(bDisplayCompactSize)
-						.CustomResetToDefault(ResetToDefaultOverride)
-						.CustomContentSlot()
-						[
-							SNew( SBox )
-							.HAlign(HAlign_Left)
-							.VAlign(VAlign_Center)
-							[
-								SNew(SHorizontalBox)
-								+SHorizontalBox::Slot()
-								.VAlign(VAlign_Center)
-								.Padding(0.0f, 0.0f, 3.0f, 0.0f)
-								.AutoWidth()
-								[
-									// Add a menu for displaying all textures 
-									SNew( SComboButton )
-									.OnGetMenuContent( this, &FMaterialItemView::OnGetTexturesMenuForMaterial )
-									.VAlign(VAlign_Center)
-									.ContentPadding(2)
-									.IsEnabled( this, &FMaterialItemView::IsTexturesMenuEnabled )
-									.Visibility( bShowUsedTextures ? EVisibility::Visible : EVisibility::Hidden )
-									.ButtonContent()
-									[
-										SNew( STextBlock )
-										.Font( IDetailLayoutBuilder::GetDetailFont() )
-										.ToolTipText( LOCTEXT("ViewTexturesToolTip", "View the textures used by this material" ) )
-										.Text( LOCTEXT("ViewTextures","Textures") )
-									]
-								]
-								+SHorizontalBox::Slot()
-								.Padding(3.0f, 0.0f)
-								.FillWidth(1.0f)
-								[
-									OnGenerateCustomMaterialWidgets.IsBound() && bDisplayCompactSize ? OnGenerateCustomMaterialWidgets.Execute(MaterialItem.Material.Get(), MaterialItem.SlotIndex) : StaticCastSharedRef<SWidget>(SNullWidget::NullWidget)
-								]
-							]
-						]
-					]
-				]
-				+SVerticalBox::Slot()
-				.AutoHeight()
-				.Padding(2)
-				.VAlign( VAlign_Center )
-				[
-					OnGenerateCustomMaterialWidgets.IsBound() && !bDisplayCompactSize ? OnGenerateCustomMaterialWidgets.Execute( MaterialItem.Material.Get(), MaterialItem.SlotIndex ) : StaticCastSharedRef<SWidget>( SNullWidget::NullWidget )
-				]
-			];
-	}
-
-private:
-
-	FMaterialItemView(	const FMaterialListItem& InMaterial, 
-						FOnMaterialChanged& InOnMaterialChanged, 
-						FOnGenerateWidgetsForMaterial& InOnGenerateNameWidgets, 
-						FOnGenerateWidgetsForMaterial& InOnGenerateMaterialWidgets, 
-						FOnResetMaterialToDefaultClicked& InOnResetToDefaultClicked,
-						int32 InMultipleMaterialCount,
-						bool bInShowUsedTextures,
-						bool bInDisplayCompactSize)
-						
-		: MaterialItem( InMaterial )
-		, OnMaterialChanged( InOnMaterialChanged )
-		, OnGenerateCustomNameWidgets( InOnGenerateNameWidgets )
-		, OnGenerateCustomMaterialWidgets( InOnGenerateMaterialWidgets )
-		, OnResetToDefaultClicked( InOnResetToDefaultClicked )
-		, MultipleMaterialCount( InMultipleMaterialCount )
-		, bShowUsedTextures( bInShowUsedTextures )
-		, bDisplayCompactSize(bInDisplayCompactSize)
-	{
-
-	}
-
-	void ReplaceMaterial( UMaterialInterface* NewMaterial, bool bReplaceAll = false )
-	{
-		UMaterialInterface* PrevMaterial = NULL;
-		if( MaterialItem.Material.IsValid() )
-		{
-			PrevMaterial = MaterialItem.Material.Get();
-		}
-
-		if( NewMaterial != PrevMaterial )
-		{
-			// Replace the material
-			OnMaterialChanged.ExecuteIfBound( NewMaterial, PrevMaterial, MaterialItem.SlotIndex, bReplaceAll );
-		}
-	}
-
-	FString OnGetObjectPath() const
-	{
-		return MaterialItem.Material->GetPathName();
-	}
-
-	void OnSetObject( const FAssetData& AssetData )
-	{
-		const bool bReplaceAll = false;
-
-		UMaterialInterface* NewMaterial = Cast<UMaterialInterface>(AssetData.GetAsset());
-		ReplaceMaterial( NewMaterial, bReplaceAll );
-	}
-
-	/**
-	 * @return Whether or not the textures menu is enabled
-	 */
-	bool IsTexturesMenuEnabled() const
-	{
-		return MaterialItem.Material.Get() != NULL;
-	}
-
-	TSharedRef<SWidget> OnGetTexturesMenuForMaterial()
-	{
-		FMenuBuilder MenuBuilder( true, NULL );
-
-		if( MaterialItem.Material.IsValid() )
-		{
-			UMaterialInterface* Material = MaterialItem.Material.Get();
-
-			TArray< UTexture* > Textures;
-			Material->GetUsedTextures(Textures, EMaterialQualityLevel::Num, false, ERHIFeatureLevel::Num, true);
-
-			// Add a menu item for each texture.  Clicking on the texture will display it in the content browser
-			// UObject for delegate compatibility
-			for( UObject* Texture : Textures )
-			{
-				FUIAction Action( FExecuteAction::CreateSP( this, &FMaterialItemView::GoToAssetInContentBrowser, MakeWeakObjectPtr(Texture) ) );
-
-				MenuBuilder.AddMenuEntry( FText::FromString( Texture->GetName() ), LOCTEXT( "BrowseTexture_ToolTip", "Find this texture in the content browser" ), FSlateIcon(), Action );
-			}
-		}
-
-		return MenuBuilder.MakeWidget();
-	}
-
-	/**
-	 * Finds the asset in the content browser
-	 */
-	void GoToAssetInContentBrowser( TWeakObjectPtr<UObject> Object )
-	{
-		if( Object.IsValid() )
-		{
-			TArray< UObject* > Objects;
-			Objects.Add( Object.Get() );
-			GEditor->SyncBrowserToObjects( Objects );
-		}
-	}
-
-	/**
-	 * Called to get the visibility of the replace button
-	 */
-	bool GetReplaceVisibility(TSharedPtr<IPropertyHandle> PropertyHandle) const
-	{
-		// Only show the replace button if the current material can be replaced
-		if (OnMaterialChanged.IsBound() && MaterialItem.bCanBeReplaced)
-		{
-			return true;
-		}
-
-		return false;
-	}
-
-	/**
-	 * Called when reset to base is clicked
-	 */
-	void OnResetToBaseClicked(TSharedPtr<IPropertyHandle> PropertyHandle)
-	{
-		// Only allow reset to base if the current material can be replaced
-		if( MaterialItem.Material.IsValid() && MaterialItem.bCanBeReplaced )
-		{
-			bool bReplaceAll = false;
-			ReplaceMaterial( NULL, bReplaceAll );
-			OnResetToDefaultClicked.ExecuteIfBound( MaterialItem.Material.Get(), MaterialItem.SlotIndex );
-		}
-	}
-
-private:
-	FMaterialListItem MaterialItem;
-	FOnMaterialChanged OnMaterialChanged;
-	FOnGenerateWidgetsForMaterial OnGenerateCustomNameWidgets;
-	FOnGenerateWidgetsForMaterial OnGenerateCustomMaterialWidgets;
-	FOnResetMaterialToDefaultClicked OnResetToDefaultClicked;
-	int32 MultipleMaterialCount;
-	bool bShowUsedTextures;
-	bool bDisplayCompactSize;
-};
-
-
-FMaterialList::FMaterialList(IDetailLayoutBuilder& InDetailLayoutBuilder, FMaterialListDelegates& InMaterialListDelegates, bool bInAllowCollapse, bool bInShowUsedTextures, bool bInDisplayCompactSize)
-	: MaterialListDelegates( InMaterialListDelegates )
-	, DetailLayoutBuilder( InDetailLayoutBuilder )
-	, MaterialListBuilder( new FMaterialListBuilder )
-	, bAllowCollpase(bInAllowCollapse)
-	, bShowUsedTextures(bInShowUsedTextures)
-	, bDisplayCompactSize(bInDisplayCompactSize)
-{
-}
-
-void FMaterialList::OnDisplayMaterialsForElement( int32 SlotIndex )
-{
-	// We now want to display all the materials in the element
-	ExpandedSlots.Add( SlotIndex );
-
-	MaterialListBuilder->Empty();
-	MaterialListDelegates.OnGetMaterials.ExecuteIfBound( *MaterialListBuilder );
-
-	OnRebuildChildren.ExecuteIfBound();
-}
-
-void FMaterialList::OnHideMaterialsForElement( int32 SlotIndex )
-{
-	// No longer want to expand the element
-	ExpandedSlots.Remove( SlotIndex );
-
-	// regenerate the materials
-	MaterialListBuilder->Empty();
-	MaterialListDelegates.OnGetMaterials.ExecuteIfBound( *MaterialListBuilder );
-	
-	OnRebuildChildren.ExecuteIfBound();
-}
-
-
-void FMaterialList::Tick( float DeltaTime )
-{
-	// Check each material to see if its still valid.  This allows the material list to stay up to date when materials are changed out from under us
-	if( MaterialListDelegates.OnGetMaterials.IsBound() )
-	{
-		// Whether or not to refresh the material list
-		bool bRefrestMaterialList = false;
-
-		// Get the current list of materials from the user
-		MaterialListBuilder->Empty();
-		MaterialListDelegates.OnGetMaterials.ExecuteIfBound( *MaterialListBuilder );
-
-		if( MaterialListBuilder->GetNumMaterials() != DisplayedMaterials.Num() )
-		{
-			// The array sizes differ so we need to refresh the list
-			bRefrestMaterialList = true;
-		}
-		else
-		{
-			// Compare the new list against the currently displayed list
-			for( int32 MaterialIndex = 0; MaterialIndex < MaterialListBuilder->MaterialSlots.Num(); ++MaterialIndex )
-			{
-				const FMaterialListItem& Item = MaterialListBuilder->MaterialSlots[MaterialIndex];
-
-				// The displayed materials is out of date if there isn't a 1:1 mapping between the material sets
-				if( !DisplayedMaterials.IsValidIndex( MaterialIndex ) || DisplayedMaterials[ MaterialIndex ] != Item )
-				{
-					bRefrestMaterialList = true;
-					break;
-				}
-			}
-		}
-
-		if (!bRefrestMaterialList && MaterialListDelegates.OnMaterialListDirty.IsBound())
-		{
-			bRefrestMaterialList = MaterialListDelegates.OnMaterialListDirty.Execute();
-		}
-
-		if( bRefrestMaterialList )
-		{
-			OnRebuildChildren.ExecuteIfBound();
-		}
-	}
-}
-
-void FMaterialList::GenerateHeaderRowContent( FDetailWidgetRow& NodeRow )
-{
-	NodeRow.CopyAction(FUIAction(FExecuteAction::CreateSP(this, &FMaterialList::OnCopyMaterialList), FCanExecuteAction::CreateSP(this, &FMaterialList::OnCanCopyMaterialList)));
-	NodeRow.PasteAction(FUIAction(FExecuteAction::CreateSP(this, &FMaterialList::OnPasteMaterialList)));
-
-	if (bAllowCollpase)
-	{
-		NodeRow.NameContent()
-			[
-				SNew(STextBlock)
-				.Text(LOCTEXT("MaterialHeaderTitle", "Materials"))
-				.Font(IDetailLayoutBuilder::GetDetailFont())
-			];
-	}
-}
-
-void FMaterialList::GenerateChildContent( IDetailChildrenBuilder& ChildrenBuilder )
-{
-	ViewedMaterials.Empty();
-	DisplayedMaterials.Empty();
-	if( MaterialListBuilder->GetNumMaterials() > 0 )
-	{
-		DisplayedMaterials = MaterialListBuilder->MaterialSlots;
-
-		MaterialListBuilder->Sort();
-		TArray<FMaterialListItem>& MaterialSlots = MaterialListBuilder->MaterialSlots;
-
-		int32 CurrentSlot = INDEX_NONE;
-		bool bDisplayAllMaterialsInSlot = true;
-		for( auto It = MaterialSlots.CreateConstIterator(); It; ++It )
-		{
-			const FMaterialListItem& Material = *It;
-
-			if( CurrentSlot != Material.SlotIndex )
-			{
-				// We've encountered a new slot.  Make a widget to display that
-				CurrentSlot = Material.SlotIndex;
-
-				uint32 NumMaterials = MaterialListBuilder->GetNumMaterialsInSlot(CurrentSlot);
-
-				// If an element is expanded we want to display all its materials
-				bool bWantToDisplayAllMaterials = NumMaterials > 1 && ExpandedSlots.Contains(CurrentSlot);
-
-				// If we are currently displaying an expanded set of materials for an element add a link to collapse all of them
-				if( bWantToDisplayAllMaterials )
-				{
-					FDetailWidgetRow& ChildRow = ChildrenBuilder.AddCustomRow( LOCTEXT( "HideAllMaterialSearchString", "Hide All Materials") );
-
-					FFormatNamedArguments Arguments;
-					Arguments.Add(TEXT("ElementSlot"), CurrentSlot);
-					ChildRow
-					.ValueContent()
-					.MaxDesiredWidth(0.0f)// No Max Width
-					[
-						SNew( SBox )
-							.HAlign( HAlign_Center )
-							[
-								SNew( SHyperlink )
-									.TextStyle( FEditorStyle::Get(), "MaterialList.HyperlinkStyle" )
-									.Text( FText::Format(LOCTEXT("HideAllMaterialLinkText", "Hide All Materials on Element {ElementSlot}"), Arguments ) )
-									.OnNavigate( this, &FMaterialList::OnHideMaterialsForElement, CurrentSlot )
-							]
-					];
-				}	
-
-				if( NumMaterials > 1 && !bWantToDisplayAllMaterials )
-				{
-					// The current slot has multiple elements to view
-					bDisplayAllMaterialsInSlot = false;
-
-					FDetailWidgetRow& ChildRow = ChildrenBuilder.AddCustomRow( FText::GetEmpty() );
-
-					AddMaterialItem( ChildRow, CurrentSlot, FMaterialListItem( NULL, CurrentSlot, true ), !bDisplayAllMaterialsInSlot );
-				}
-				else
-				{
-					bDisplayAllMaterialsInSlot = true;
-				}
-
-			}
-
-			// Display each thumbnail element unless we shouldn't display multiple materials for one slot
-			if( bDisplayAllMaterialsInSlot )
-			{
-				FDetailWidgetRow& ChildRow = ChildrenBuilder.AddCustomRow( Material.Material.IsValid()? FText::FromString(Material.Material->GetName()) : FText::GetEmpty() );
-
-				AddMaterialItem( ChildRow, CurrentSlot, Material, !bDisplayAllMaterialsInSlot );
-			}
-		}
-	}
-	else
-	{
-		FDetailWidgetRow& ChildRow = ChildrenBuilder.AddCustomRow( LOCTEXT("NoMaterials", "No Materials") );
-
-		ChildRow
-		[
-			SNew( SBox )
-			.HAlign( HAlign_Center )
-			[
-				SNew( STextBlock )
-				.Text( LOCTEXT("NoMaterials", "No Materials") ) 
-				.Font( IDetailLayoutBuilder::GetDetailFont() )
-			]
-		];
-	}		
-}
-
-bool FMaterialList::OnCanCopyMaterialList() const
-{
-	if (MaterialListDelegates.OnCanCopyMaterialList.IsBound())
-	{
-		return MaterialListDelegates.OnCanCopyMaterialList.Execute();
-	}
-
-	return false;
-}
-
-void FMaterialList::OnCopyMaterialList()
-{
-	if (MaterialListDelegates.OnCopyMaterialList.IsBound())
-	{
-		MaterialListDelegates.OnCopyMaterialList.Execute();
-	}
-}
-
-void FMaterialList::OnPasteMaterialList()
-{
-	if (MaterialListDelegates.OnPasteMaterialList.IsBound())
-	{
-		MaterialListDelegates.OnPasteMaterialList.Execute();
-	}
-}
-
-bool FMaterialList::OnCanCopyMaterialItem(int32 CurrentSlot) const
-{
-	if (MaterialListDelegates.OnCanCopyMaterialItem.IsBound())
-	{
-		return MaterialListDelegates.OnCanCopyMaterialItem.Execute(CurrentSlot);
-	}
-
-	return false;
-}
-
-void FMaterialList::OnCopyMaterialItem(int32 CurrentSlot)
-{
-	if (MaterialListDelegates.OnCopyMaterialItem.IsBound())
-	{
-		MaterialListDelegates.OnCopyMaterialItem.Execute(CurrentSlot);
-	}
-}
-
-void FMaterialList::OnPasteMaterialItem(int32 CurrentSlot)
-{
-	if (MaterialListDelegates.OnPasteMaterialItem.IsBound())
-	{
-		MaterialListDelegates.OnPasteMaterialItem.Execute(CurrentSlot);
-	}
-}
-			
-void FMaterialList::AddMaterialItem( FDetailWidgetRow& Row, int32 CurrentSlot, const FMaterialListItem& Item, bool bDisplayLink )
-{
-	uint32 NumMaterials = MaterialListBuilder->GetNumMaterialsInSlot(CurrentSlot);
-
-	TSharedRef<FMaterialItemView> NewView = FMaterialItemView::Create( Item, MaterialListDelegates.OnMaterialChanged, MaterialListDelegates.OnGenerateCustomNameWidgets, MaterialListDelegates.OnGenerateCustomMaterialWidgets, MaterialListDelegates.OnResetMaterialToDefaultClicked, NumMaterials, bShowUsedTextures, bDisplayCompactSize);
-
-	TSharedPtr<SWidget> RightSideContent;
-	if( bDisplayLink )
-	{
-		FFormatNamedArguments Arguments;
-		Arguments.Add(TEXT("NumMaterials"), NumMaterials);
-
-		RightSideContent = 
-			SNew( SBox )
-				.HAlign(HAlign_Left)
-				.VAlign(VAlign_Top)
-				[
-					SNew( SHyperlink )
-					.TextStyle( FEditorStyle::Get(), "MaterialList.HyperlinkStyle" )
-					.Text( FText::Format(LOCTEXT("DisplayAllMaterialLinkText", "Display {NumMaterials} materials"), Arguments) )
-					.ToolTipText( LOCTEXT("DisplayAllMaterialLink_ToolTip","Display all materials. Drag and drop a material here to replace all materials.") )
-					.OnNavigate( this, &FMaterialList::OnDisplayMaterialsForElement, CurrentSlot )
-				];
-	}
-	else
-	{
-		RightSideContent = NewView->CreateValueContent( DetailLayoutBuilder.GetThumbnailPool() );
-		ViewedMaterials.Add( NewView );
-	}
-
-	Row.CopyAction(FUIAction(FExecuteAction::CreateSP(this, &FMaterialList::OnCopyMaterialItem, Item.SlotIndex), FCanExecuteAction::CreateSP(this, &FMaterialList::OnCanCopyMaterialItem, Item.SlotIndex)));
-	Row.PasteAction(FUIAction(FExecuteAction::CreateSP(this, &FMaterialList::OnPasteMaterialItem, Item.SlotIndex)));
-
-	Row.NameContent()
-	[
-		NewView->CreateNameContent()
-	]
-	.ValueContent()
-	.MinDesiredWidth(250.f)
-	.MaxDesiredWidth(0.0f) // no maximum
-	[
-		RightSideContent.ToSharedRef()
-	];
-}
-
 TSharedRef<SWidget> PropertyCustomizationHelpers::MakePropertyComboBox(const TSharedPtr<IPropertyHandle>& InPropertyHandle, FOnGetPropertyComboBoxStrings OnGetStrings, FOnGetPropertyComboBoxValue OnGetValue, FOnPropertyComboBoxValueSelected OnValueSelected)
 {
 	FSlateFontInfo FontStyle = FEditorStyle::GetFontStyle(PropertyEditorConstants::PropertyFontStyle);
diff --git a/Engine/Source/Editor/PropertyEditor/Public/MaterialList.h b/Engine/Source/Editor/PropertyEditor/Public/MaterialList.h
new file mode 100644
--- /dev/null
+++ b/Engine/Source/Editor/PropertyEditor/Public/MaterialList.h
@@ -0,0 +1,219 @@
+// Copyright 1998-2019 Epic Games, Inc. All Rights Reserved.
+
+#pragma once
+
+#include "IDetailCustomNodeBuilder.h"
+#include "DetailWidgetRow.h"
+#include "Materials/MaterialInterface.h"
+
+class FMaterialItemView;
+class FMaterialListBuilder;
+class IDetailChildrenBuilder;
+class IDetailLayoutBuilder;
+class IMaterialListBuilder;
+
+/**
+ * Delegate called when we need to get new materials for the list
+ */
+DECLARE_DELEGATE_OneParam(FOnGetMaterials, IMaterialListBuilder&);
+
+/**
+ * Delegate called when a user changes the material
+ */
+DECLARE_DELEGATE_FourParams( FOnMaterialChanged, UMaterialInterface*, UMaterialInterface*, int32, bool );
+
+DECLARE_DELEGATE_RetVal_TwoParams( TSharedRef<SWidget>, FOnGenerateWidgetsForMaterial, UMaterialInterface*, int32 );
+
+DECLARE_DELEGATE_TwoParams( FOnResetMaterialToDefaultClicked, UMaterialInterface*, int32 );
+
+DECLARE_DELEGATE_RetVal(bool, FOnMaterialListDirty);
+
+DECLARE_DELEGATE_RetVal(bool, FOnCanCopyMaterialList);
+DECLARE_DELEGATE(FOnCopyMaterialList);
+DECLARE_DELEGATE(FOnPasteMaterialList);
+
+DECLARE_DELEGATE_RetVal_OneParam(bool, FOnCanCopyMaterialItem, int32);
+DECLARE_DELEGATE_OneParam(FOnCopyMaterialItem, int32);
+DECLARE_DELEGATE_OneParam(FOnPasteMaterialItem, int32);
+
+struct FMaterialListDelegates
+{
+	FMaterialListDelegates()
+		: OnGetMaterials()
+		, OnMaterialChanged()
+		, OnGenerateCustomNameWidgets()
+		, OnGenerateCustomMaterialWidgets()
+		, OnResetMaterialToDefaultClicked()
+	{}
+
+	/** Delegate called to populate the list with materials */
+	FOnGetMaterials OnGetMaterials;
+	/** Delegate called when a user changes the material */
+	FOnMaterialChanged OnMaterialChanged;
+	/** Delegate called to generate custom widgets under the name of in the left column of a details panel*/
+	FOnGenerateWidgetsForMaterial OnGenerateCustomNameWidgets;
+	/** Delegate called to generate custom widgets under each material */
+	FOnGenerateWidgetsForMaterial OnGenerateCustomMaterialWidgets;
+	/** Delegate called when a material list item should be reset to default */
+	FOnResetMaterialToDefaultClicked OnResetMaterialToDefaultClicked;
+	/** Delegate called when we tick the material list to know if the list is dirty*/
+	FOnMaterialListDirty OnMaterialListDirty;
+
+	/** Delegate called Copying a material list */
+	FOnCopyMaterialList OnCopyMaterialList;
+	/** Delegate called to know if we can copy a material list */
+	FOnCanCopyMaterialList OnCanCopyMaterialList;
+	/** Delegate called Pasting a material list */
+	FOnPasteMaterialList OnPasteMaterialList;
+
+	/** Delegate called Copying a material item */
+	FOnCopyMaterialItem OnCopyMaterialItem;
+	/** Delegate called to know if we can copy a material item */
+	FOnCanCopyMaterialItem OnCanCopyMaterialItem;
+	/** Delegate called Pasting a material item */
+	FOnPasteMaterialItem OnPasteMaterialItem;
+};
+
+/**
+ * Builds up a list of unique materials while creating some information about the materials
+ */
+class IMaterialListBuilder
+{
+public:
+
+	/** Virtual destructor. */
+	virtual ~IMaterialListBuilder(){};
+
+	/** 
+	 * Adds a new material to the list
+	 * 
+	 * @param SlotIndex The slot (usually mesh element index) where the material is located on the component.
+	 * @param Material The material being used.
+	 * @param bCanBeReplced Whether or not the material can be replaced by a user.
+	 */
+	virtual void AddMaterial( uint32 SlotIndex, UMaterialInterface* Material, bool bCanBeReplaced ) = 0;
+};
+
+
+/**
+ * A Material item in a material list slot
+ */
+struct FMaterialListItem
+{
+	/** Material being used */
+	TWeakObjectPtr<UMaterialInterface> Material;
+
+	/** Slot on a component where this material is at (mesh element) */
+	int32 SlotIndex;
+
+	/** Whether or not this material can be replaced by a new material */
+	bool bCanBeReplaced;
+
+	FMaterialListItem( UMaterialInterface* InMaterial = NULL, uint32 InSlotIndex = 0, bool bInCanBeReplaced = false )
+		: Material( InMaterial )
+		, SlotIndex( InSlotIndex )
+		, bCanBeReplaced( bInCanBeReplaced )
+	{}
+
+	friend uint32 GetTypeHash( const FMaterialListItem& InItem )
+	{
+		return GetTypeHash( InItem.Material ) + InItem.SlotIndex ;
+	}
+
+	bool operator==( const FMaterialListItem& Other ) const
+	{
+		return Material == Other.Material && SlotIndex == Other.SlotIndex ;
+	}
+
+	bool operator!=( const FMaterialListItem& Other ) const
+	{
+		return !(*this == Other);
+	}
+};
+
+
+class FMaterialList
+	: public IDetailCustomNodeBuilder
+	, public TSharedFromThis<FMaterialList>
+{
+public:
+	PROPERTYEDITOR_API FMaterialList( IDetailLayoutBuilder& InDetailLayoutBuilder, FMaterialListDelegates& MaterialListDelegates, bool bInAllowCollapse = false, bool bInShowUsedTextures = true, bool bInDisplayCompactSize = false);
+
+	/**
+	 * @return true if materials are being displayed.                                                          
+	 */
+	bool IsDisplayingMaterials() const { return true; }
+
+private:
+
+	/**
+	 * Called when a user expands all materials in a slot.
+	 *
+	 * @param SlotIndex The index of the slot being expanded.
+	 */
+	void OnDisplayMaterialsForElement( int32 SlotIndex );
+
+	/**
+	 * Called when a user hides all materials in a slot.
+	 *
+	 * @param SlotIndex The index of the slot being hidden.
+	 */
+	void OnHideMaterialsForElement( int32 SlotIndex );
+
+	/** IDetailCustomNodeBuilder interface */
+	virtual void SetOnRebuildChildren( FSimpleDelegate InOnRebuildChildren  ) override { OnRebuildChildren = InOnRebuildChildren; } 
+	virtual bool RequiresTick() const override { return true; }
+	virtual void Tick( float DeltaTime ) override;
+	virtual void GenerateHeaderRowContent( FDetailWidgetRow& NodeRow ) override;
+	virtual void GenerateChildContent( IDetailChildrenBuilder& ChildrenBuilder ) override;
+	virtual FName GetName() const override { return NAME_None; }
+	virtual bool InitiallyCollapsed() const override { return bAllowCollpase; }
+
+	/**
+	 * Adds a new material item to the list
+	 *
+	 * @param Row			The row to add the item to
+	 * @param CurrentSlot	The slot id of the material
+	 * @param Item			The material item to add
+	 * @param bDisplayLink	If a link to the material should be displayed instead of the actual item (for multiple materials)
+	 */
+	void AddMaterialItem(FDetailWidgetRow& Row, int32 CurrentSlot, const FMaterialListItem& Item, bool bDisplayLink);
+
+private:
+	bool OnCanCopyMaterialList() const;
+	void OnCopyMaterialList();
+	void OnPasteMaterialList();
+
+	bool OnCanCopyMaterialItem(int32 CurrentSlot) const;
+	void OnCopyMaterialItem(int32 CurrentSlot);
+	void OnPasteMaterialItem(int32 CurrentSlot);
+
+	/** Delegates for the material list */
+	FMaterialListDelegates MaterialListDelegates;
+
+	/** Called to rebuild the children of the detail tree */
+	FSimpleDelegate OnRebuildChildren;
+
+	/** Parent detail layout this list is in */
+	IDetailLayoutBuilder& DetailLayoutBuilder;
+
+	/** Set of all unique displayed materials */
+	TArray< FMaterialListItem > DisplayedMaterials;
+
+	/** Set of all materials currently in view (may be less than DisplayedMaterials) */
+	TArray< TSharedRef<FMaterialItemView> > ViewedMaterials;
+
+	/** Set of all expanded slots */
+	TSet<uint32> ExpandedSlots;
+
+	/** Material list builder used to generate materials */
+	TSharedRef<FMaterialListBuilder> MaterialListBuilder;
+
+	/** Allow Collapse of material header row. Right now if you allow collapse, it will initially collapse. */
+	bool bAllowCollpase;
+	/** Whether or not to use the used textures menu for each material entry */
+	bool bShowUsedTextures;
+	/** Whether or not to display a compact form of material entry*/
+	bool bDisplayCompactSize;
+};
+
diff --git a/Engine/Source/Editor/PropertyEditor/Public/PropertyCustomizationHelpers.h b/Engine/Source/Editor/PropertyEditor/Public/PropertyCustomizationHelpers.h
--- a/Engine/Source/Editor/PropertyEditor/Public/PropertyCustomizationHelpers.h
+++ b/Engine/Source/Editor/PropertyEditor/Public/PropertyCustomizationHelpers.h
@@ -20,16 +20,12 @@
 #include "SceneDepthPickerMode.h"
 #include "IDetailPropertyRow.h"
 
-
 class AActor;
 struct FAssetData;
 class FAssetThumbnailPool;
-class FMaterialItemView;
-class FMaterialListBuilder;
 class FPropertyEditor;
 class IDetailChildrenBuilder;
 class IDetailLayoutBuilder;
-class IMaterialListBuilder;
 class SPropertyEditorAsset;
 class SPropertyEditorClass;
 class UFactory;
@@ -40,7 +36,6 @@ namespace SceneOutliner
 	struct FOutlinerFilters;
 }
 
-
 DECLARE_DELEGATE_OneParam(FOnAssetSelected, const FAssetData& /*AssetData*/);
 DECLARE_DELEGATE_RetVal_OneParam(bool, FOnShouldSetAsset, const FAssetData& /*AssetData*/);
 DECLARE_DELEGATE_RetVal_OneParam(bool, FOnShouldFilterAsset, const FAssetData& /*AssetData*/);
@@ -418,214 +413,6 @@ private:
 	bool bDisplayElementNum;
 };
 
-
-/**
- * Delegate called when we need to get new materials for the list
- */
-DECLARE_DELEGATE_OneParam(FOnGetMaterials, IMaterialListBuilder&);
-
-/**
- * Delegate called when a user changes the material
- */
-DECLARE_DELEGATE_FourParams( FOnMaterialChanged, UMaterialInterface*, UMaterialInterface*, int32, bool );
-
-DECLARE_DELEGATE_RetVal_TwoParams( TSharedRef<SWidget>, FOnGenerateWidgetsForMaterial, UMaterialInterface*, int32 );
-
-DECLARE_DELEGATE_TwoParams( FOnResetMaterialToDefaultClicked, UMaterialInterface*, int32 );
-
-DECLARE_DELEGATE_RetVal(bool, FOnMaterialListDirty);
-
-DECLARE_DELEGATE_RetVal(bool, FOnCanCopyMaterialList);
-DECLARE_DELEGATE(FOnCopyMaterialList);
-DECLARE_DELEGATE(FOnPasteMaterialList);
-
-DECLARE_DELEGATE_RetVal_OneParam(bool, FOnCanCopyMaterialItem, int32);
-DECLARE_DELEGATE_OneParam(FOnCopyMaterialItem, int32);
-DECLARE_DELEGATE_OneParam(FOnPasteMaterialItem, int32);
-
-struct FMaterialListDelegates
-{
-	FMaterialListDelegates()
-		: OnGetMaterials()
-		, OnMaterialChanged()
-		, OnGenerateCustomNameWidgets()
-		, OnGenerateCustomMaterialWidgets()
-		, OnResetMaterialToDefaultClicked()
-	{}
-
-	/** Delegate called to populate the list with materials */
-	FOnGetMaterials OnGetMaterials;
-	/** Delegate called when a user changes the material */
-	FOnMaterialChanged OnMaterialChanged;
-	/** Delegate called to generate custom widgets under the name of in the left column of a details panel*/
-	FOnGenerateWidgetsForMaterial OnGenerateCustomNameWidgets;
-	/** Delegate called to generate custom widgets under each material */
-	FOnGenerateWidgetsForMaterial OnGenerateCustomMaterialWidgets;
-	/** Delegate called when a material list item should be reset to default */
-	FOnResetMaterialToDefaultClicked OnResetMaterialToDefaultClicked;
-	/** Delegate called when we tick the material list to know if the list is dirty*/
-	FOnMaterialListDirty OnMaterialListDirty;
-
-	/** Delegate called Copying a material list */
-	FOnCopyMaterialList OnCopyMaterialList;
-	/** Delegate called to know if we can copy a material list */
-	FOnCanCopyMaterialList OnCanCopyMaterialList;
-	/** Delegate called Pasting a material list */
-	FOnPasteMaterialList OnPasteMaterialList;
-
-	/** Delegate called Copying a material item */
-	FOnCopyMaterialItem OnCopyMaterialItem;
-	/** Delegate called to know if we can copy a material item */
-	FOnCanCopyMaterialItem OnCanCopyMaterialItem;
-	/** Delegate called Pasting a material item */
-	FOnPasteMaterialItem OnPasteMaterialItem;
-};
-
-
-/**
- * Builds up a list of unique materials while creating some information about the materials
- */
-class IMaterialListBuilder
-{
-public:
-
-	/** Virtual destructor. */
-	virtual ~IMaterialListBuilder(){};
-
-	/** 
-	 * Adds a new material to the list
-	 * 
-	 * @param SlotIndex The slot (usually mesh element index) where the material is located on the component.
-	 * @param Material The material being used.
-	 * @param bCanBeReplced Whether or not the material can be replaced by a user.
-	 */
-	virtual void AddMaterial( uint32 SlotIndex, UMaterialInterface* Material, bool bCanBeReplaced ) = 0;
-};
-
-
-/**
- * A Material item in a material list slot
- */
-struct FMaterialListItem
-{
-	/** Material being used */
-	TWeakObjectPtr<UMaterialInterface> Material;
-
-	/** Slot on a component where this material is at (mesh element) */
-	int32 SlotIndex;
-
-	/** Whether or not this material can be replaced by a new material */
-	bool bCanBeReplaced;
-
-	FMaterialListItem( UMaterialInterface* InMaterial = NULL, uint32 InSlotIndex = 0, bool bInCanBeReplaced = false )
-		: Material( InMaterial )
-		, SlotIndex( InSlotIndex )
-		, bCanBeReplaced( bInCanBeReplaced )
-	{}
-
-	friend uint32 GetTypeHash( const FMaterialListItem& InItem )
-	{
-		return GetTypeHash( InItem.Material ) + InItem.SlotIndex ;
-	}
-
-	bool operator==( const FMaterialListItem& Other ) const
-	{
-		return Material == Other.Material && SlotIndex == Other.SlotIndex ;
-	}
-
-	bool operator!=( const FMaterialListItem& Other ) const
-	{
-		return !(*this == Other);
-	}
-};
-
-
-class FMaterialList
-	: public IDetailCustomNodeBuilder
-	, public TSharedFromThis<FMaterialList>
-{
-public:
-	PROPERTYEDITOR_API FMaterialList( IDetailLayoutBuilder& InDetailLayoutBuilder, FMaterialListDelegates& MaterialListDelegates, bool bInAllowCollapse = false, bool bInShowUsedTextures = true, bool bInDisplayCompactSize = false);
-
-	/**
-	 * @return true if materials are being displayed.                                                          
-	 */
-	bool IsDisplayingMaterials() const { return true; }
-
-private:
-
-	/**
-	 * Called when a user expands all materials in a slot.
-	 *
-	 * @param SlotIndex The index of the slot being expanded.
-	 */
-	void OnDisplayMaterialsForElement( int32 SlotIndex );
-
-	/**
-	 * Called when a user hides all materials in a slot.
-	 *
-	 * @param SlotIndex The index of the slot being hidden.
-	 */
-	void OnHideMaterialsForElement( int32 SlotIndex );
-
-	/** IDetailCustomNodeBuilder interface */
-	virtual void SetOnRebuildChildren( FSimpleDelegate InOnRebuildChildren  ) override { OnRebuildChildren = InOnRebuildChildren; } 
-	virtual bool RequiresTick() const override { return true; }
-	virtual void Tick( float DeltaTime ) override;
-	virtual void GenerateHeaderRowContent( FDetailWidgetRow& NodeRow ) override;
-	virtual void GenerateChildContent( IDetailChildrenBuilder& ChildrenBuilder ) override;
-	virtual FName GetName() const override { return NAME_None; }
-	virtual bool InitiallyCollapsed() const override { return bAllowCollpase; }
-
-	/**
-	 * Adds a new material item to the list
-	 *
-	 * @param Row			The row to add the item to
-	 * @param CurrentSlot	The slot id of the material
-	 * @param Item			The material item to add
-	 * @param bDisplayLink	If a link to the material should be displayed instead of the actual item (for multiple materials)
-	 */
-	void AddMaterialItem(FDetailWidgetRow& Row, int32 CurrentSlot, const FMaterialListItem& Item, bool bDisplayLink);
-
-private:
-	bool OnCanCopyMaterialList() const;
-	void OnCopyMaterialList();
-	void OnPasteMaterialList();
-
-	bool OnCanCopyMaterialItem(int32 CurrentSlot) const;
-	void OnCopyMaterialItem(int32 CurrentSlot);
-	void OnPasteMaterialItem(int32 CurrentSlot);
-
-	/** Delegates for the material list */
-	FMaterialListDelegates MaterialListDelegates;
-
-	/** Called to rebuild the children of the detail tree */
-	FSimpleDelegate OnRebuildChildren;
-
-	/** Parent detail layout this list is in */
-	IDetailLayoutBuilder& DetailLayoutBuilder;
-
-	/** Set of all unique displayed materials */
-	TArray< FMaterialListItem > DisplayedMaterials;
-
-	/** Set of all materials currently in view (may be less than DisplayedMaterials) */
-	TArray< TSharedRef<FMaterialItemView> > ViewedMaterials;
-
-	/** Set of all expanded slots */
-	TSet<uint32> ExpandedSlots;
-
-	/** Material list builder used to generate materials */
-	TSharedRef<FMaterialListBuilder> MaterialListBuilder;
-
-	/** Allow Collapse of material header row. Right now if you allow collapse, it will initially collapse. */
-	bool bAllowCollpase;
-	/** Whether or not to use the used textures menu for each material entry */
-	bool bShowUsedTextures;
-	/** Whether or not to display a compact form of material entry*/
-	bool bDisplayCompactSize;
-};
-
-
 /**
  * Helper class to create a material slot name widget for material lists
  */
diff --git a/Engine/Source/Editor/StaticMeshEditor/Private/StaticMeshEditorTools.cpp b/Engine/Source/Editor/StaticMeshEditor/Private/StaticMeshEditorTools.cpp
--- a/Engine/Source/Editor/StaticMeshEditor/Private/StaticMeshEditorTools.cpp
+++ b/Engine/Source/Editor/StaticMeshEditor/Private/StaticMeshEditorTools.cpp
@@ -23,6 +23,7 @@
 #include "StaticMeshResources.h"
 #include "StaticMeshEditor.h"
 #include "PropertyCustomizationHelpers.h"
+#include "MaterialList.h"
 #include "PhysicsEngine/BodySetup.h"
 #include "FbxMeshUtils.h"
 #include "Widgets/Input/SVectorInputBox.h"
