Missed includes for non-unity build.

diff --git a/Engine/Source/Editor/PropertyEditor/Private/MaterialList.cpp b/Engine/Source/Editor/PropertyEditor/Private/MaterialList.cpp
--- a/Engine/Source/Editor/PropertyEditor/Private/MaterialList.cpp
+++ b/Engine/Source/Editor/PropertyEditor/Private/MaterialList.cpp
@@ -5,7 +5,10 @@
 #include "IDetailChildrenBuilder.h"
 #include "PropertyCustomizationHelpers.h"
 #include "PropertyHandle.h"
+#include "Widgets/Input/SComboButton.h"
 #include "Widgets/Input/SHyperlink.h"
+#include "Widgets/Layout/SBox.h"
+#include "Widgets/Text/STextBlock.h"
 
 #define LOCTEXT_NAMESPACE "PropertyCustomizationHelpers"
 
