Added missing includes for non-unity build.

diff --git a/Engine/Source/Editor/PropertyEditor/Private/MaterialList.cpp b/Engine/Source/Editor/PropertyEditor/Private/MaterialList.cpp
--- a/Engine/Source/Editor/PropertyEditor/Private/MaterialList.cpp
+++ b/Engine/Source/Editor/PropertyEditor/Private/MaterialList.cpp
@@ -2,9 +2,11 @@
 
 #include "MaterialList.h"
 #include "DetailLayoutBuilder.h"
+#include "Editor.h"
 #include "IDetailChildrenBuilder.h"
 #include "PropertyCustomizationHelpers.h"
 #include "PropertyHandle.h"
+#include "Framework/MultiBox/MultiBoxBuilder.h"
 #include "Widgets/Input/SComboButton.h"
 #include "Widgets/Input/SHyperlink.h"
 #include "Widgets/Layout/SBox.h"
