CIS warning: LOCTEXT_NAMESPACE wasn't being undefined in MaterialList.cpp.

diff --git a/Engine/Source/Editor/PropertyEditor/Private/MaterialList.cpp b/Engine/Source/Editor/PropertyEditor/Private/MaterialList.cpp
--- a/Engine/Source/Editor/PropertyEditor/Private/MaterialList.cpp
+++ b/Engine/Source/Editor/PropertyEditor/Private/MaterialList.cpp
@@ -636,3 +636,5 @@ void FMaterialList::AddMaterialItem( FDetailWidgetRow& Row, int32 CurrentSlot, c
 		RightSideContent.ToSharedRef()
 	];
 }
+
+#undef LOCTEXT_NAMESPACE
\ No newline at end of file
