UPROPERTY meta EditCondition now supports a full-fledged expression parser, allowing for multiple conditions to be used in EditConditions.

Previous functionality still works as it did before.

diff --git a/Engine/Source/Editor/PropertyEditor/Private/EditConditionParser.cpp b/Engine/Source/Editor/PropertyEditor/Private/EditConditionParser.cpp
new file mode 100644
--- /dev/null
+++ b/Engine/Source/Editor/PropertyEditor/Private/EditConditionParser.cpp
@@ -0,0 +1,903 @@
+// Copyright 1998-2019 Epic Games, Inc. All Rights Reserved.
+
+#include "EditConditionParser.h"
+#include "Misc/ExpressionParser.h"
+
+#include "PropertyNode.h"
+#include "PropertyEditorHelpers.h"
+
+#define LOCTEXT_NAMESPACE "EditConditionParser"
+
+namespace EditConditionParserNamespace
+{
+	const TCHAR* const FEqual::Moniker = TEXT("==");
+	const TCHAR* const FNotEqual::Moniker = TEXT("!=");
+	const TCHAR* const FGreater::Moniker = TEXT(">");
+	const TCHAR* const FGreaterEqual::Moniker = TEXT(">=");
+	const TCHAR* const FLess::Moniker = TEXT("<");
+	const TCHAR* const FLessEqual::Moniker = TEXT("<=");
+	const TCHAR* const FNot::Moniker = TEXT("!");
+	const TCHAR* const FAnd::Moniker = TEXT("&&");
+	const TCHAR* const FOr::Moniker = TEXT("||");
+	const TCHAR* const FAdd::Moniker = TEXT("+");
+	const TCHAR* const FSubtract::Moniker = TEXT("-");
+	const TCHAR* const FMultiply::Moniker = TEXT("*");
+	const TCHAR* const FDivide::Moniker = TEXT("/");
+}
+
+static const TCHAR PropertyBreakingChars[] = { '|', '=', '&', '>', '<', '!', '+', '-', '*', '/', ' ', '\t' };
+
+static TOptional<FExpressionError> ConsumeBool(FExpressionTokenConsumer& Consumer)
+{
+	int TrueChars = 0;
+	int FalseChars = 0;
+	TOptional<FStringToken> StringToken = Consumer.GetStream().ParseToken([&TrueChars, &FalseChars](TCHAR InC)
+	{
+		if (TEXT("true")[TrueChars] == InC)
+		{
+			++TrueChars;
+			if (TrueChars == 4)
+			{
+				return EParseState::StopAfter;
+			}
+			else
+			{
+				return EParseState::Continue;
+			}
+		}
+		else if (TEXT("false")[FalseChars] == InC)
+		{
+			++FalseChars;
+			if (FalseChars == 5)
+			{
+				return EParseState::StopAfter;
+			}
+			else
+			{
+				return EParseState::Continue;
+			}
+		}
+		return EParseState::Cancel;
+	});
+
+	if (StringToken.IsSet())
+	{
+		bool bIsTrue = TrueChars == 4;
+		bool bIsFalse = FalseChars == 5;
+		ensure(bIsTrue || bIsFalse);
+
+		Consumer.Add(StringToken.GetValue(), bIsTrue);
+	}
+
+	return TOptional<FExpressionError>();
+}
+
+static TOptional<FExpressionError> ConsumePropertyName(FExpressionTokenConsumer& Consumer)
+{
+	FString PropertyName;
+	bool bShouldBeEnum = false;
+			
+	TOptional<FStringToken> StringToken = Consumer.GetStream().ParseToken([&PropertyName, &bShouldBeEnum](TCHAR InC)
+	{
+		for (const TCHAR BreakingChar : PropertyBreakingChars)
+		{
+			if (InC == BreakingChar)
+			{
+				return EParseState::StopBefore;
+			}
+		}
+
+		if (InC == ':')
+		{
+			bShouldBeEnum = true;
+		}
+
+		PropertyName.AppendChar(InC);
+
+		return EParseState::Continue;
+	});
+
+	if (StringToken.IsSet())
+	{
+		if (bShouldBeEnum)
+		{
+			int32 DoubleColonIndex = PropertyName.Find("::");
+			if (DoubleColonIndex == INDEX_NONE)
+			{
+				return FExpressionError(FText::Format(LOCTEXT("PropertyContainsSingleColon", "EditCondition contains single colon in property name \"{0}\", expected double colons."), FText::FromString(PropertyName)));
+			}
+
+			if (DoubleColonIndex == 0)
+			{
+				return FExpressionError(FText::Format(LOCTEXT("PropertyDoubleColonAtStart", "EditCondition contained double colon at start of property name \"{0}\", expected enum type."), FText::FromString(PropertyName)));
+			}
+
+			FString EnumType = PropertyName.Left(DoubleColonIndex);
+			FString EnumValue = PropertyName.RightChop(DoubleColonIndex + 2);
+			
+			if (EnumValue.Len() == 0)
+			{
+				return FExpressionError(FText::Format(LOCTEXT("PropertyDoubleColonAtEnd", "EditCondition contained double colon at end of property name \"{0}\", expected enum value."), FText::FromString(PropertyName)));
+			}
+
+			Consumer.Add(StringToken.GetValue(), EditConditionParserNamespace::FEnumToken(MoveTemp(EnumType), MoveTemp(EnumValue)));
+		}
+		else
+		{
+			Consumer.Add(StringToken.GetValue(), EditConditionParserNamespace::FPropertyToken(MoveTemp(PropertyName)));
+		}
+	}
+
+	return TOptional<FExpressionError>();
+}
+
+template <typename T>
+TOptional<T> GetValueInternal(const IEditConditionContext& Context, const FString& PropertyName)
+{
+	return TOptional<T>();
+}
+
+template<>
+TOptional<bool> GetValueInternal<bool>(const IEditConditionContext& Context, const FString& PropertyName)
+{
+	return Context.GetBoolValue(PropertyName);
+}
+
+template<>
+TOptional<double> GetValueInternal<double>(const IEditConditionContext& Context, const FString& PropertyName)
+{
+	return Context.GetNumericValue(PropertyName);
+}
+
+template <typename T>
+struct TOperand
+{
+	TOperand(T InValue) : Value(InValue), Property(nullptr), Context(nullptr) {}
+	TOperand(const EditConditionParserNamespace::FPropertyToken& InProperty, const IEditConditionContext& InContext) : 
+		Property(&InProperty), Context(&InContext) {}
+
+	bool IsProperty() const { return Property != nullptr; }
+	TOptional<T> GetValue() const
+	{
+		if (IsProperty())
+		{
+			return GetValueInternal<T>(*Context, Property->PropertyName);	
+		}
+		
+		return TOptional<T>(Value);
+	}
+
+	const FString& GetName() const { check(IsProperty()); return Property->PropertyName; }
+
+private:
+	T Value;
+	const EditConditionParserNamespace::FPropertyToken* Property;
+	const IEditConditionContext* Context;
+};
+
+static FExpressionResult ApplyNot(TOperand<bool> A)
+{
+	TOptional<bool> Value = A.GetValue();
+	if (Value.IsSet())
+	{
+		return MakeValue(!Value.GetValue());
+	}
+	return MakeError(FText::Format(LOCTEXT("InvalidOperand", "EditCondition attempted to use an invalid operand \"{0}\"."), FText::FromString(A.GetName())));
+}
+
+template <typename T, typename Function>
+FExpressionResult ApplyBinary(TOperand<T> A, TOperand<T> B, Function Apply)
+{
+	TOptional<T> ValueA = A.GetValue();
+	if (!ValueA.IsSet())
+	{
+		return MakeError(FText::Format(LOCTEXT("InvalidOperand", "EditCondition attempted to use an invalid operand \"{0}\"."), FText::FromString(A.GetName())));
+	}
+
+	TOptional<T> ValueB = B.GetValue();
+	if (!ValueB.IsSet())
+	{
+		return MakeError(FText::Format(LOCTEXT("InvalidOperand", "EditCondition attempted to use an invalid operand \"{0}\"."), FText::FromString(B.GetName())));
+	}
+
+	return MakeValue(Apply(ValueA.GetValue(), ValueB.GetValue()));
+}
+
+static FExpressionResult ApplyPropertiesEqual(const EditConditionParserNamespace::FPropertyToken& A, const EditConditionParserNamespace::FPropertyToken& B, const IEditConditionContext& Context, bool bNegate)
+{
+	TOptional<FString> TypeNameA = Context.GetTypeName(A.PropertyName);
+	TOptional<FString> TypeNameB = Context.GetTypeName(B.PropertyName);
+	if (!TypeNameA.IsSet())
+	{
+		return MakeError(FText::Format(LOCTEXT("InvalidOperand", "EditCondition attempted to use an invalid operand \"{0}\"."), FText::FromString(A.PropertyName)));
+	}
+
+	if (!TypeNameB.IsSet())
+	{
+		return MakeError(FText::Format(LOCTEXT("InvalidOperand", "EditCondition attempted to use an invalid operand \"{0}\"."), FText::FromString(B.PropertyName)));
+	}
+
+	if (TypeNameA.GetValue() != TypeNameB.GetValue())
+	{
+		return MakeError(FText::Format(LOCTEXT("OperandTypeMismatch", "EditCondition attempted to compare operands of different types: \"{0}\" and \"{1}\"."), FText::FromString(A.PropertyName), FText::FromString(B.PropertyName)));
+	}
+
+	TOptional<bool> bEqual;
+
+	TOptional<bool> BoolA = Context.GetBoolValue(A.PropertyName);
+	TOptional<bool> BoolB = Context.GetBoolValue(B.PropertyName);
+	if (BoolA.IsSet() && BoolB.IsSet())
+	{
+		bEqual = BoolA.GetValue() == BoolB.GetValue();
+	}
+
+	TOptional<double> DoubleA = Context.GetNumericValue(A.PropertyName);
+	TOptional<double> DoubleB = Context.GetNumericValue(B.PropertyName);
+	if (DoubleA.IsSet() && DoubleB.IsSet())
+	{
+		bEqual = DoubleA.GetValue() == DoubleB.GetValue();
+	}
+
+	TOptional<FString> EnumA = Context.GetEnumValue(A.PropertyName);
+	TOptional<FString> EnumB = Context.GetEnumValue(B.PropertyName);
+	if (EnumA.IsSet() && EnumB.IsSet())
+	{
+		bEqual = EnumA.GetValue() == EnumB.GetValue();
+	}
+
+	if (bEqual.IsSet())
+	{
+		return MakeValue(bNegate ? !bEqual.GetValue() : bEqual.GetValue());
+	}
+
+	return MakeError(FText::Format(LOCTEXT("OperandTypeMismatch", "EditCondition attempted to compare operands of different types: \"{0}\" and \"{1}\"."), FText::FromString(A.PropertyName), FText::FromString(B.PropertyName)));
+}
+
+static void CreateBooleanOperators(TOperatorJumpTable<IEditConditionContext>& OperatorJumpTable)
+{
+	using namespace EditConditionParserNamespace;
+
+	OperatorJumpTable.MapPreUnary<FNot>([](bool A) { return !A; });
+	OperatorJumpTable.MapPreUnary<FNot>([](const FPropertyToken& A, const IEditConditionContext* Context)
+	{
+		return ApplyNot(TOperand<bool>(A, *Context));
+	});
+
+	// AND
+	{
+		auto ApplyAnd = [](bool First, bool Second) { return First && Second; };
+	
+		OperatorJumpTable.MapBinary<FAnd>(ApplyAnd);
+		OperatorJumpTable.MapBinary<FAnd>([ApplyAnd](const FPropertyToken& A, bool B, const IEditConditionContext* Context)
+		{
+			return ApplyBinary(TOperand<bool>(A, *Context), TOperand<bool>(B), ApplyAnd);
+		});
+		OperatorJumpTable.MapBinary<FAnd>([ApplyAnd](bool A, const FPropertyToken& B, const IEditConditionContext* Context)
+		{
+			return ApplyBinary(TOperand<bool>(A), TOperand<bool>(B, *Context), ApplyAnd);
+		});
+		OperatorJumpTable.MapBinary<FAnd>([ApplyAnd](const FPropertyToken& A, const FPropertyToken& B, const IEditConditionContext* Context)
+		{
+			return ApplyBinary(TOperand<bool>(A, *Context), TOperand<bool>(B, *Context), ApplyAnd);
+		});
+	}
+
+	// OR
+	{
+		auto ApplyOr = [](bool First, bool Second) { return First || Second; };
+
+		OperatorJumpTable.MapBinary<FOr>(ApplyOr);
+		OperatorJumpTable.MapBinary<FOr>([ApplyOr](const FPropertyToken& A, bool B, const IEditConditionContext* Context)
+		{
+			return ApplyBinary(TOperand<bool>(A, *Context), TOperand<bool>(B), ApplyOr);
+		});
+		OperatorJumpTable.MapBinary<FOr>([ApplyOr](bool A, const FPropertyToken& B, const IEditConditionContext* Context)
+		{
+			return ApplyBinary(TOperand<bool>(A), TOperand<bool>(B, *Context), ApplyOr);
+		});
+		OperatorJumpTable.MapBinary<FOr>([ApplyOr](const FPropertyToken& A, const FPropertyToken& B, const IEditConditionContext* Context)
+		{
+			return ApplyBinary(TOperand<bool>(A, *Context), TOperand<bool>(B, *Context), ApplyOr);
+		});
+	}
+
+	// EQUALS
+	{
+		auto ApplyEqual = [](bool First, bool Second) { return First == Second; };
+
+		OperatorJumpTable.MapBinary<FEqual>(ApplyEqual);
+		OperatorJumpTable.MapBinary<FEqual>([ApplyEqual](const FPropertyToken& A, bool B, const IEditConditionContext* Context)
+		{
+			return ApplyBinary(TOperand<bool>(A, *Context), TOperand<bool>(B), ApplyEqual);
+		});
+		OperatorJumpTable.MapBinary<FEqual>([ApplyEqual](bool A, const FPropertyToken& B, const IEditConditionContext* Context)
+		{
+			return ApplyBinary(TOperand<bool>(A), TOperand<bool>(B, *Context), ApplyEqual);
+		});
+	}
+
+	// NOT-EQUALS
+	{
+		auto ApplyNotEqual = [](bool First, bool Second) { return First != Second; };
+
+		OperatorJumpTable.MapBinary<FNotEqual>(ApplyNotEqual);
+		OperatorJumpTable.MapBinary<FNotEqual>([ApplyNotEqual](const FPropertyToken& A, bool B, const IEditConditionContext* Context)
+		{
+			return ApplyBinary(TOperand<bool>(A, *Context), TOperand<bool>(B), ApplyNotEqual);
+		});
+		OperatorJumpTable.MapBinary<FNotEqual>([ApplyNotEqual](bool A, const FPropertyToken& B, const IEditConditionContext* Context)
+		{
+			return ApplyBinary(TOperand<bool>(A), TOperand<bool>(B, *Context), ApplyNotEqual);
+		});
+	}
+}
+
+template <typename T>
+void CreateNumberOperators(TOperatorJumpTable<IEditConditionContext>& OperatorJumpTable)
+{
+	using namespace EditConditionParserNamespace;
+
+	// EQUAL
+	{
+		auto ApplyEqual = [](T First, T Second) { return First == Second; };
+
+		OperatorJumpTable.MapBinary<FEqual>(ApplyEqual);
+		OperatorJumpTable.MapBinary<FEqual>([ApplyEqual](const FPropertyToken& A, T B, const IEditConditionContext* Context)
+		{
+			return ApplyBinary(TOperand<T>(A, *Context), TOperand<T>(B), ApplyEqual);
+		});
+		OperatorJumpTable.MapBinary<FEqual>([ApplyEqual](T A, const FPropertyToken& B, const IEditConditionContext* Context)
+		{
+			return ApplyBinary(TOperand<T>(A), TOperand<T>(B, *Context), ApplyEqual);
+		});
+	}
+
+	// NOT-EQUAL
+	{
+		auto ApplyNotEqual = [](T First, T Second) { return First != Second; };
+
+		OperatorJumpTable.MapBinary<FNotEqual>(ApplyNotEqual);
+		OperatorJumpTable.MapBinary<FNotEqual>([ApplyNotEqual](const FPropertyToken& A, T B, const IEditConditionContext* Context)
+		{
+			return ApplyBinary(TOperand<T>(A, *Context), TOperand<T>(B), ApplyNotEqual);
+		});
+		OperatorJumpTable.MapBinary<FNotEqual>([ApplyNotEqual](T A, const FPropertyToken& B, const IEditConditionContext* Context)
+		{
+			return ApplyBinary(TOperand<T>(A), TOperand<T>(B, *Context), ApplyNotEqual);
+		});
+	}
+
+	// GREATER
+	{
+		auto ApplyGreater = [](T First, T Second) { return First > Second; };
+
+		OperatorJumpTable.MapBinary<FGreater>(ApplyGreater);
+		OperatorJumpTable.MapBinary<FGreater>([ApplyGreater](const FPropertyToken& A, T B, const IEditConditionContext* Context)
+		{
+			return ApplyBinary(TOperand<T>(A, *Context), TOperand<T>(B), ApplyGreater);
+		});
+		OperatorJumpTable.MapBinary<FGreater>([ApplyGreater](T A, const FPropertyToken& B, const IEditConditionContext* Context)
+		{
+			return ApplyBinary(TOperand<T>(A), TOperand<T>(B, *Context), ApplyGreater);
+		});
+		OperatorJumpTable.MapBinary<FGreater>([ApplyGreater](const FPropertyToken& A, const FPropertyToken& B, const IEditConditionContext* Context)
+		{
+			return ApplyBinary(TOperand<T>(A, *Context), TOperand<T>(B, *Context), ApplyGreater);
+		});
+	}
+
+	// GREATER-EQUAL
+	{
+		auto ApplyGreaterEqual = [](T First, T Second) { return First >= Second; };
+
+		OperatorJumpTable.MapBinary<FGreaterEqual>(ApplyGreaterEqual);
+		OperatorJumpTable.MapBinary<FGreaterEqual>([ApplyGreaterEqual](const FPropertyToken& A, T B, const IEditConditionContext* Context)
+		{
+			return ApplyBinary(TOperand<T>(A, *Context), TOperand<T>(B), ApplyGreaterEqual);
+		});
+		OperatorJumpTable.MapBinary<FGreaterEqual>([ApplyGreaterEqual](T A, const FPropertyToken& B, const IEditConditionContext* Context)
+		{
+			return ApplyBinary(TOperand<T>(A), TOperand<T>(B, *Context), ApplyGreaterEqual);
+		});
+		OperatorJumpTable.MapBinary<FGreaterEqual>([ApplyGreaterEqual](const FPropertyToken& A, const FPropertyToken& B, const IEditConditionContext* Context)
+		{
+			return ApplyBinary(TOperand<T>(A, *Context), TOperand<T>(B, *Context), ApplyGreaterEqual);
+		});
+	}
+
+	// LESS
+	{
+		auto ApplyLess = [](T First, T Second) { return First < Second; };
+
+		OperatorJumpTable.MapBinary<FLess>(ApplyLess);
+		OperatorJumpTable.MapBinary<FLess>([ApplyLess](const FPropertyToken& A, T B, const IEditConditionContext* Context)
+		{
+			return ApplyBinary(TOperand<T>(A, *Context), TOperand<T>(B), ApplyLess);
+		});
+		OperatorJumpTable.MapBinary<FLess>([ApplyLess](T A, const FPropertyToken& B, const IEditConditionContext* Context)
+		{
+			return ApplyBinary(TOperand<T>(A), TOperand<T>(B, *Context), ApplyLess);
+		});
+		OperatorJumpTable.MapBinary<FLess>([ApplyLess](const FPropertyToken& A, const FPropertyToken& B, const IEditConditionContext* Context)
+		{
+			return ApplyBinary(TOperand<T>(A, *Context), TOperand<T>(B, *Context), ApplyLess);
+		});
+	}
+
+	// LESS-EQUAL
+	{
+		auto ApplyLessEqual = [](T First, T Second) { return First <= Second; };
+
+		OperatorJumpTable.MapBinary<FLessEqual>(ApplyLessEqual);
+		OperatorJumpTable.MapBinary<FLessEqual>([ApplyLessEqual](const FPropertyToken& A, T B, const IEditConditionContext* Context)
+		{
+			return ApplyBinary(TOperand<T>(A, *Context), TOperand<T>(B), ApplyLessEqual);
+		});
+		OperatorJumpTable.MapBinary<FLessEqual>([ApplyLessEqual](T A, const FPropertyToken& B, const IEditConditionContext* Context)
+		{
+			return ApplyBinary(TOperand<T>(A), TOperand<T>(B, *Context), ApplyLessEqual);
+		});
+		OperatorJumpTable.MapBinary<FLessEqual>([ApplyLessEqual](const FPropertyToken& A, const FPropertyToken& B, const IEditConditionContext* Context)
+		{
+			return ApplyBinary(TOperand<T>(A, *Context), TOperand<T>(B, *Context), ApplyLessEqual);
+		});
+	}
+
+	// ADD
+	{
+		auto ApplyAdd = [](T First, T Second) { return First + Second; };
+
+		OperatorJumpTable.MapBinary<FAdd>(ApplyAdd);
+		OperatorJumpTable.MapBinary<FAdd>([ApplyAdd](const FPropertyToken& A, T B, const IEditConditionContext* Context)
+		{
+			return ApplyBinary(TOperand<T>(A, *Context), TOperand<T>(B), ApplyAdd);
+		});
+		OperatorJumpTable.MapBinary<FAdd>([ApplyAdd](T A, const FPropertyToken& B, const IEditConditionContext* Context)
+		{
+			return ApplyBinary(TOperand<T>(A), TOperand<T>(B, *Context), ApplyAdd);
+		});
+		OperatorJumpTable.MapBinary<FAdd>([ApplyAdd](const FPropertyToken& A, const FPropertyToken& B, const IEditConditionContext* Context)
+		{
+			return ApplyBinary(TOperand<T>(A, *Context), TOperand<T>(B, *Context), ApplyAdd);
+		});
+	}
+
+	// SUBTRACT
+	{
+		auto ApplySubtract = [](T First, T Second) { return First - Second; };
+
+		OperatorJumpTable.MapBinary<FSubtract>(ApplySubtract);
+		OperatorJumpTable.MapBinary<FSubtract>([ApplySubtract](const FPropertyToken& A, T B, const IEditConditionContext* Context)
+		{
+			return ApplyBinary(TOperand<T>(A, *Context), TOperand<T>(B), ApplySubtract);
+		});
+		OperatorJumpTable.MapBinary<FSubtract>([ApplySubtract](T A, const FPropertyToken& B, const IEditConditionContext* Context)
+		{
+			return ApplyBinary(TOperand<T>(A), TOperand<T>(B, *Context), ApplySubtract);
+		});
+		OperatorJumpTable.MapBinary<FSubtract>([ApplySubtract](const FPropertyToken& A, const FPropertyToken& B, const IEditConditionContext* Context)
+		{
+			return ApplyBinary(TOperand<T>(A, *Context), TOperand<T>(B, *Context), ApplySubtract);
+		});
+	}
+
+	// MULTIPLY
+	{
+		auto ApplyMultiply = [](T First, T Second) { return First * Second; };
+
+		OperatorJumpTable.MapBinary<FMultiply>(ApplyMultiply);
+		OperatorJumpTable.MapBinary<FMultiply>([ApplyMultiply](const FPropertyToken& A, T B, const IEditConditionContext* Context)
+		{
+			return ApplyBinary(TOperand<T>(A, *Context), TOperand<T>(B), ApplyMultiply);
+		});
+		OperatorJumpTable.MapBinary<FMultiply>([ApplyMultiply](T A, const FPropertyToken& B, const IEditConditionContext* Context)
+		{
+			return ApplyBinary(TOperand<T>(A), TOperand<T>(B, *Context), ApplyMultiply);
+		});
+		OperatorJumpTable.MapBinary<FMultiply>([ApplyMultiply](const FPropertyToken& A, const FPropertyToken& B, const IEditConditionContext* Context)
+		{
+			return ApplyBinary(TOperand<T>(A, *Context), TOperand<T>(B, *Context), ApplyMultiply);
+		});
+	}
+
+	// DIVIDE
+	{
+		auto ApplyDivide = [](T First, T Second) { return First / Second; };
+
+		OperatorJumpTable.MapBinary<FDivide>(ApplyDivide);
+		OperatorJumpTable.MapBinary<FDivide>([ApplyDivide](const FPropertyToken& A, T B, const IEditConditionContext* Context)
+		{
+			return ApplyBinary(TOperand<T>(A, *Context), TOperand<T>(B), ApplyDivide);
+		});
+		OperatorJumpTable.MapBinary<FDivide>([ApplyDivide](T A, const FPropertyToken& B, const IEditConditionContext* Context)
+		{
+			return ApplyBinary(TOperand<T>(A), TOperand<T>(B, *Context), ApplyDivide);
+		});
+		OperatorJumpTable.MapBinary<FDivide>([ApplyDivide](const FPropertyToken& A, const FPropertyToken& B, const IEditConditionContext* Context)
+		{
+			return ApplyBinary(TOperand<T>(A, *Context), TOperand<T>(B, *Context), ApplyDivide);
+		});
+	}
+}
+
+static FExpressionResult EnumPropertyEquals(const EditConditionParserNamespace::FEnumToken& Enum, const EditConditionParserNamespace::FPropertyToken& Property, const IEditConditionContext& Context, bool bNegate)
+{
+	TOptional<FString> TypeName = Context.GetTypeName(Property.PropertyName);
+	TOptional<FString> ValueProp = Context.GetEnumValue(Property.PropertyName);
+
+	if (TypeName.GetValue() != Enum.Type)
+	{
+		return MakeError(FText::Format(LOCTEXT("OperandTypeMismatch", "EditCondition attempted to compare operands of different types: \"{0}\" and \"{1}\"."), FText::FromString(Property.PropertyName), FText::FromString(Enum.Type + TEXT("::") + Enum.Value)));
+	}
+
+	if (!ValueProp.IsSet())
+	{
+		return MakeError(FText::Format(LOCTEXT("InvalidOperand", "EditCondition attempted to use an invalid operand \"{0}\"."), FText::FromString(Property.PropertyName)));
+	}
+
+	bool bEqual = ValueProp.GetValue() == Enum.Value;
+	return MakeValue(bNegate ? !bEqual : bEqual);
+}
+
+static void CreateEnumOperators(TOperatorJumpTable<IEditConditionContext>& OperatorJumpTable)
+{
+	using namespace EditConditionParserNamespace;
+
+	// EQUALS
+	{
+		OperatorJumpTable.MapBinary<FEqual>([](const FEnumToken& A, const FEnumToken& B, const IEditConditionContext* Context)
+		{
+			return A.Type == B.Type && A.Value == B.Value;
+		});
+		OperatorJumpTable.MapBinary<FEqual>([](const FPropertyToken& A, const FEnumToken& B, const IEditConditionContext* Context)
+		{
+			return EnumPropertyEquals(B, A, *Context, false);
+		});
+		OperatorJumpTable.MapBinary<FEqual>([](const FEnumToken& A, const FPropertyToken& B, const IEditConditionContext* Context)
+		{
+			return EnumPropertyEquals(A, B, *Context, false);
+		});
+	}
+
+	// NOT-EQUALS
+	{
+		OperatorJumpTable.MapBinary<FNotEqual>([](const FEnumToken& A, const FEnumToken& B, const IEditConditionContext* Context)
+		{
+			return A.Type != B.Type || A.Value != B.Value;
+		});
+		OperatorJumpTable.MapBinary<FNotEqual>([](const FPropertyToken& A, const FEnumToken& B, const IEditConditionContext* Context) -> FExpressionResult
+		{
+			return EnumPropertyEquals(B, A, *Context, true);
+		});
+		OperatorJumpTable.MapBinary<FNotEqual>([](const FEnumToken& A, const FPropertyToken& B, const IEditConditionContext* Context) -> FExpressionResult
+		{
+			return EnumPropertyEquals(A, B, *Context, true);
+		});
+	}
+}
+
+FEditConditionParser::FEditConditionParser()
+{
+	using namespace EditConditionParserNamespace;
+
+	TokenDefinitions.IgnoreWhitespace();
+	TokenDefinitions.DefineToken(&ExpressionParser::ConsumeSymbol<FEqual>);
+	TokenDefinitions.DefineToken(&ExpressionParser::ConsumeSymbol<FNotEqual>);
+	TokenDefinitions.DefineToken(&ExpressionParser::ConsumeSymbol<FLessEqual>);
+	TokenDefinitions.DefineToken(&ExpressionParser::ConsumeSymbol<FLess>);
+	TokenDefinitions.DefineToken(&ExpressionParser::ConsumeSymbol<FGreaterEqual>);
+	TokenDefinitions.DefineToken(&ExpressionParser::ConsumeSymbol<FGreater>);
+	TokenDefinitions.DefineToken(&ExpressionParser::ConsumeSymbol<FNot>);
+	TokenDefinitions.DefineToken(&ExpressionParser::ConsumeSymbol<FAnd>);
+	TokenDefinitions.DefineToken(&ExpressionParser::ConsumeSymbol<FOr>);
+	TokenDefinitions.DefineToken(&ExpressionParser::ConsumeSymbol<FAdd>);
+	TokenDefinitions.DefineToken(&ExpressionParser::ConsumeSymbol<FSubtract>);
+	TokenDefinitions.DefineToken(&ExpressionParser::ConsumeSymbol<FMultiply>);
+	TokenDefinitions.DefineToken(&ExpressionParser::ConsumeSymbol<FDivide>);
+	TokenDefinitions.DefineToken(&ExpressionParser::ConsumeNumber);
+	TokenDefinitions.DefineToken(&ConsumeBool);
+	TokenDefinitions.DefineToken(&ConsumePropertyName);
+				
+	ExpressionGrammar.DefineBinaryOperator<FAnd>(4);
+	ExpressionGrammar.DefineBinaryOperator<FOr>(4);
+	ExpressionGrammar.DefineBinaryOperator<FEqual>(3);
+	ExpressionGrammar.DefineBinaryOperator<FNotEqual>(3);
+	ExpressionGrammar.DefineBinaryOperator<FLess>(3);
+	ExpressionGrammar.DefineBinaryOperator<FLessEqual>(3);
+	ExpressionGrammar.DefineBinaryOperator<FGreater>(3);
+	ExpressionGrammar.DefineBinaryOperator<FGreaterEqual>(3);
+	ExpressionGrammar.DefineBinaryOperator<FAdd>(2);
+	ExpressionGrammar.DefineBinaryOperator<FSubtract>(2);
+	ExpressionGrammar.DefineBinaryOperator<FMultiply>(1);
+	ExpressionGrammar.DefineBinaryOperator<FDivide>(1);
+	ExpressionGrammar.DefinePreUnaryOperator<FNot>();
+
+	OperatorJumpTable.MapBinary<FEqual>([](const FPropertyToken& A, const FPropertyToken& B, const IEditConditionContext* Context) -> FExpressionResult
+	{
+		return ApplyPropertiesEqual(A, B, *Context, false);
+	});
+
+	OperatorJumpTable.MapBinary<FNotEqual>([](const FPropertyToken& A, const FPropertyToken& B, const IEditConditionContext* Context) -> FExpressionResult
+	{
+		return ApplyPropertiesEqual(A, B, *Context, true);
+	});
+
+	CreateBooleanOperators(OperatorJumpTable);
+	CreateNumberOperators<double>(OperatorJumpTable);
+	CreateEnumOperators(OperatorJumpTable);
+}
+
+TOptional<bool> FEditConditionParser::Evaluate(const FEditConditionExpression& Expression, const IEditConditionContext& Context) const
+{
+	using namespace EditConditionParserNamespace;
+	
+	FExpressionResult Result = ExpressionParser::Evaluate(Expression.Tokens, OperatorJumpTable, &Context);
+	if (Result.IsValid())
+	{
+		const bool* BoolResult = Result.GetValue().Cast<bool>();
+		if (BoolResult != nullptr)
+		{
+			return *BoolResult;
+		}
+
+		const FPropertyToken* PropertyResult = Result.GetValue().Cast<FPropertyToken>();
+		if (PropertyResult != nullptr)
+		{
+			TOptional<bool> PropertyValue = Context.GetBoolValue(PropertyResult->PropertyName);
+			if (PropertyValue.IsSet())
+			{
+				return PropertyValue.GetValue();
+			}
+		}
+	}
+
+	return TOptional<bool>();
+}
+
+TSharedPtr<FEditConditionExpression> FEditConditionParser::Parse(const FString& ExpressionString) const
+{
+	using namespace ExpressionParser;
+
+	LexResultType LexResult = ExpressionParser::Lex(*ExpressionString, TokenDefinitions);
+	if (LexResult.IsValid())
+	{
+		CompileResultType CompileResult = ExpressionParser::Compile(LexResult.StealValue(), ExpressionGrammar);
+		if (CompileResult.IsValid())
+		{
+			return TSharedPtr<FEditConditionExpression>(new FEditConditionExpression(CompileResult.StealValue()));
+		}
+	}
+
+	return TSharedPtr<FEditConditionExpression>();
+}
+
+FEditConditionContext::FEditConditionContext(FPropertyNode& InPropertyNode, TSharedPtr<FEditConditionExpression> Expression)
+{
+	ParentNode = InPropertyNode.GetParentNode();
+
+	const UProperty* Property = InPropertyNode.GetProperty();
+	check(Property);
+
+	if (PropertyEditorHelpers::IsStaticArray(InPropertyNode))
+	{
+		//in the case of conditional static arrays, we have to go up one more level to get the proper parent struct.
+		ParentNode = ParentNode->GetParentNode();
+	}
+	check(ParentNode);
+
+	FComplexPropertyNode* ComplexParentNode = ParentNode->FindComplexParent();
+	check(ComplexParentNode);
+
+	// populate the properties array
+	for (const FCompiledToken& Token : Expression->Tokens)
+	{
+		if (const EditConditionParserNamespace::FPropertyToken* PropertyToken = Token.Node.Cast<EditConditionParserNamespace::FPropertyToken>())
+		{
+			const UProperty* Field = FindField<UProperty>(Property->GetOwnerStruct(), *PropertyToken->PropertyName);
+			ensureMsgf(Field != nullptr, TEXT("EditCondition parsing failed: Field name %s was not found in struct %s."), *PropertyToken->PropertyName, *Property->GetOwnerStruct()->GetName());
+
+			if (Field != nullptr)
+			{
+				FReferencedProperty& RefProperty = Properties[Properties.Emplace()];
+				RefProperty.Property = Field;
+
+				for (int32 Index = 0; Index < ComplexParentNode->GetInstancesNum(); ++Index)
+				{
+					uint8* BaseAddress = ComplexParentNode->GetMemoryOfInstance(Index);
+					check(BaseAddress != nullptr);
+
+					FPropertyAddress& Address = RefProperty.Addresses[RefProperty.Addresses.Emplace()];
+					Address.BaseAddress = BaseAddress;
+					Address.ObjectPtr = ComplexParentNode->AsStructureNode() ?
+						TWeakObjectPtr<UObject>(ComplexParentNode->GetBaseStructure()) :
+						ComplexParentNode->GetInstanceAsUObject(Index);
+
+				}
+			}
+		}
+	}
+}
+
+const UBoolProperty* FEditConditionContext::GetSingleBoolProperty() const
+{
+	if (Properties.Num() == 1)
+	{
+		const UBoolProperty* BoolProperty = Cast<UBoolProperty>(Properties[0].Property);
+		return BoolProperty;
+	}
+	return nullptr;
+}
+
+const FEditConditionContext::FReferencedProperty* FEditConditionContext::FindProperty(const FString& PropertyName) const
+{
+	for (const FReferencedProperty& Property : Properties)
+	{
+		if (Property.Property->GetNameCPP() == PropertyName)
+		{
+			return &Property;
+		}
+	}
+	return nullptr;
+}
+
+TOptional<bool> FEditConditionContext::GetBoolValue(const FString& PropertyName) const
+{
+	const FReferencedProperty* RefProperty = FindProperty(PropertyName);
+	if (RefProperty == nullptr)
+	{
+		return TOptional<bool>();
+	}
+
+	const UBoolProperty* BoolProperty = Cast<UBoolProperty>(RefProperty->Property);
+	if (BoolProperty == nullptr)
+	{
+		return TOptional<bool>();
+	}
+
+	TOptional<bool> Result;
+
+	const FComplexPropertyNode* ComplexParentNode = ParentNode->FindComplexParent();
+	for (const FPropertyAddress& Address : RefProperty->Addresses)
+	{
+		if (!Address.ObjectPtr.IsValid())
+		{
+			return TOptional<bool>();
+		}
+
+		uint8* BaseOffset = ParentNode->GetValueAddress(Address.BaseAddress);
+		uint8* ValueAddr = BoolProperty->ContainerPtrToValuePtr<uint8>(BaseOffset);
+
+		bool bValue = BoolProperty->GetPropertyValue(ValueAddr);
+		if (!Result.IsSet())
+		{
+			Result = bValue;
+		}
+		else if (Result.GetValue() != bValue)
+		{
+			// all values aren't the same...
+			return TOptional<bool>();
+		}
+	}
+
+	return Result;
+}
+
+TOptional<double> FEditConditionContext::GetNumericValue(const FString& PropertyName) const
+{
+	const FReferencedProperty* RefProperty = FindProperty(PropertyName);
+	if (RefProperty == nullptr)
+	{
+		return TOptional<double>();
+	}
+
+	const UNumericProperty* NumericProperty = Cast<UNumericProperty>(RefProperty->Property);
+	if (NumericProperty == nullptr)
+	{
+		return TOptional<double>();
+	}
+
+	TOptional<double> Result;
+
+	for (const FPropertyAddress& Address : RefProperty->Addresses)
+	{
+		if (!Address.ObjectPtr.IsValid())
+		{
+			return TOptional<double>();
+		}
+
+		uint8* BaseOffset = ParentNode->GetValueAddress(Address.BaseAddress);
+		uint8* ValueAddr = NumericProperty->ContainerPtrToValuePtr<uint8>(BaseOffset);
+
+		double Value = 0;
+
+		if (NumericProperty->IsInteger())
+		{
+			Value = (double) NumericProperty->GetSignedIntPropertyValue(ValueAddr);
+		}
+		else if (NumericProperty->IsFloatingPoint())
+		{
+			Value = NumericProperty->GetFloatingPointPropertyValue(ValueAddr);
+		}
+
+		if (!Result.IsSet())
+		{
+			Result = Value;
+		}
+		else if (!FMath::IsNearlyEqual(Result.GetValue(), Value))
+		{
+			// all values aren't the same...
+			return TOptional<double>();
+		}
+	}
+
+	return Result;
+}
+
+TOptional<FString> FEditConditionContext::GetEnumValue(const FString& PropertyName) const
+{
+	const FReferencedProperty* RefProperty = FindProperty(PropertyName);
+	if (RefProperty == nullptr)
+	{
+		return TOptional<FString>();
+	}
+
+	const UEnumProperty* EnumProperty = Cast<UEnumProperty>(RefProperty->Property);
+	if (EnumProperty == nullptr)
+	{
+		return TOptional<FString>();
+	}
+
+	const UNumericProperty* NumProperty = EnumProperty->GetUnderlyingProperty();
+	if (NumProperty == nullptr)
+	{
+		return TOptional<FString>();
+	}
+
+	TOptional<int64> Result;
+
+	const FComplexPropertyNode* ComplexParentNode = ParentNode->FindComplexParent();
+	for (const FPropertyAddress& Address : RefProperty->Addresses)
+	{
+		if (!Address.ObjectPtr.IsValid())
+		{
+			return TOptional<FString>();
+		}
+
+		uint8* BaseOffset = ParentNode->GetValueAddress(Address.BaseAddress);
+		uint8* ValueAddr = NumProperty->ContainerPtrToValuePtr<uint8>(BaseOffset);
+
+		int64 Value = NumProperty->GetSignedIntPropertyValue(ValueAddr);
+
+		if (!Result.IsSet())
+		{
+			Result = Value;
+		}
+		else if (Result.GetValue() != Value)
+		{
+			// all values aren't the same...
+			return TOptional<FString>();
+		}
+	}
+
+	if (Result.IsSet())
+	{
+		return EnumProperty->GetEnum()->GetNameByValue(Result.GetValue()).ToString();
+	}
+
+	return TOptional<FString>();
+}
+
+TOptional<FString> FEditConditionContext::GetTypeName(const FString& PropertyName) const
+{
+	const FReferencedProperty* RefProperty = FindProperty(PropertyName);
+	if (RefProperty != nullptr)
+	{
+		const UProperty* Property = RefProperty->Property;
+		const UClass* Class = Property->GetClass();
+		return Class->GetName();
+	}
+
+	return TOptional<FString>();
+}
+
+#undef LOCTEXT_NAMESPACE
\ No newline at end of file
diff --git a/Engine/Source/Editor/PropertyEditor/Private/EditConditionParser.h b/Engine/Source/Editor/PropertyEditor/Private/EditConditionParser.h
new file mode 100644
--- /dev/null
+++ b/Engine/Source/Editor/PropertyEditor/Private/EditConditionParser.h
@@ -0,0 +1,152 @@
+// Copyright 1998-2019 Epic Games, Inc. All Rights Reserved.
+
+#pragma once
+
+namespace EditConditionParserNamespace
+{
+	struct FPropertyToken 
+	{
+		FPropertyToken(FString&& InProperty) :
+			PropertyName(InProperty) {}
+
+		FPropertyToken(const FPropertyToken& Other) :
+			PropertyName(Other.PropertyName) {}
+
+		FPropertyToken& operator=(const FPropertyToken& Other)
+		{
+			PropertyName = Other.PropertyName;
+			return *this;
+		}
+
+		FString PropertyName;
+	};
+
+	struct FEnumToken 
+	{
+		FEnumToken(FString&& InType, FString&& InValue) :
+			Type(InType), Value(InValue) {}
+
+		FEnumToken(const FEnumToken& Other) :
+			Type(Other.Type), Value(Other.Value) {}
+
+		FEnumToken& operator=(const FEnumToken& Other)
+		{
+			Type = Other.Type;
+			Value = Other.Value;
+			return *this;
+		}
+
+		FString Type;
+		FString Value;
+	};
+}
+
+#define DEFINE_EDIT_CONDITION_NODE(TYPE, ...) \
+	namespace EditConditionParserNamespace { struct TYPE { static const TCHAR* const Moniker; }; } \
+	DEFINE_EXPRESSION_NODE_TYPE(EditConditionParserNamespace::TYPE, __VA_ARGS__)
+
+	DEFINE_EDIT_CONDITION_NODE(FEqual, 0x3AF0EE1B, 0xC3F847C7, 0xA2B95102, 0x4EFC202D)
+	DEFINE_EDIT_CONDITION_NODE(FNotEqual, 0x5CDF3FA4, 0x35614D88, 0x8F94017C, 0xF9229625)
+	DEFINE_EDIT_CONDITION_NODE(FLessEqual, 0x0D097F12, 0x2268447B, 0xA9C6769E, 0x5FE086EA)
+	DEFINE_EDIT_CONDITION_NODE(FLess, 0xD3B1D8E9, 0x552249C7, 0x92646CF8, 0x6D065F45)
+	DEFINE_EDIT_CONDITION_NODE(FGreaterEqual, 0x5BC680B9, 0x148448FB, 0xA37A0EBB, 0x34C4BCA8)
+	DEFINE_EDIT_CONDITION_NODE(FGreater, 0xBD99934C, 0x5E7F4BE9, 0x99A1BD6E, 0x582F8FC9)
+	DEFINE_EDIT_CONDITION_NODE(FNot, 0xE0FB4BB7, 0x66A646C2, 0xB2360362, 0x0E1FAE33)
+	DEFINE_EDIT_CONDITION_NODE(FAnd, 0x71F839BC, 0xD75F6B47, 0x9C68017B, 0x33FAB080)
+	DEFINE_EDIT_CONDITION_NODE(FOr, 0x9C317A4C, 0xCE50304D, 0x8BF6D471, 0x9D791746)
+	DEFINE_EDIT_CONDITION_NODE(FAdd, 0xF4086317, 0x80428041, 0x903149EB, 0xC5A60ACF)
+	DEFINE_EDIT_CONDITION_NODE(FSubtract, 0xE80D53E9, 0x32F1E145, 0x8535B767, 0x0634DE9E)
+	DEFINE_EDIT_CONDITION_NODE(FMultiply, 0x0D4D37DB, 0x53C7EA41, 0x8297C200, 0xF5B28D27)
+	DEFINE_EDIT_CONDITION_NODE(FDivide, 0x48009789, 0x2EE84F40, 0x83467A99, 0x92617168)
+	DEFINE_EDIT_CONDITION_NODE(FSubExpressionStart, 0x78CE8411, 0x34076241, 0xA2BD0548, 0x54458A3A)
+	DEFINE_EDIT_CONDITION_NODE(FSubExpressionEnd, 0x5CB25466, 0x780CAE4D, 0x8A88E000, 0xD3695885)
+
+	DEFINE_EXPRESSION_NODE_TYPE(bool, 0xCACBC715, 0x505A6B4A, 0x8808809F, 0x897AA5F6)
+	DEFINE_EXPRESSION_NODE_TYPE(EditConditionParserNamespace::FPropertyToken, 0x9A3FAF6F, 0xB2E45E4D, 0xA80A70C6, 0x47A89BD7)
+	DEFINE_EXPRESSION_NODE_TYPE(EditConditionParserNamespace::FEnumToken, 0xC9A35C24, 0x21FC904B, 0x9F1B2B6A, 0xDF6F4BC4)
+
+#undef DEFINE_EDIT_CONDITION_NODE
+
+class IEditConditionContext
+{
+public:
+	virtual TOptional<bool> GetBoolValue(const FString& PropertyName) const = 0;
+	virtual TOptional<double> GetNumericValue(const FString& PropertyName) const = 0;
+	virtual TOptional<FString> GetEnumValue(const FString& PropertyName) const = 0;
+	virtual TOptional<FString> GetTypeName(const FString& PropertyName) const = 0;
+};
+
+class FEditConditionExpression
+{
+public:
+	FEditConditionExpression(TArray<FCompiledToken>&& InTokens) :
+		Tokens(MoveTemp(InTokens))
+	{}
+
+	TArray<FCompiledToken> Tokens;
+};
+
+class FEditConditionParser
+{
+public:
+	FEditConditionParser();
+
+	FEditConditionParser(const FEditConditionParser&) = delete;
+	FEditConditionParser(FEditConditionParser&&) = default;
+	FEditConditionParser& operator=(const FEditConditionParser&) = delete;
+	FEditConditionParser& operator=(FEditConditionParser&&) = default;
+
+	TSharedPtr<FEditConditionExpression> Parse(const FString& ExpressionString) const;
+
+	/** 
+	 * Evaluate the given expression within the given context.
+	 * @returns The result of the evaluated expression if valid, invalid TOptional if the evaluation failed or produced a non-bool result.
+	 */
+	TOptional<bool> Evaluate(const FEditConditionExpression& Expression, const IEditConditionContext& Context) const;
+
+private:
+	FTokenDefinitions TokenDefinitions;
+	FExpressionGrammar ExpressionGrammar;
+	TOperatorJumpTable<IEditConditionContext> OperatorJumpTable;
+};
+
+class UProperty;
+class FPropertyNode;
+class FComplexPropertyNode;
+
+class FEditConditionContext : public IEditConditionContext
+{
+public:
+	FEditConditionContext(FPropertyNode& InPropertyNode, TSharedPtr<FEditConditionExpression> Expression);
+	virtual ~FEditConditionContext() {}
+
+	virtual TOptional<bool> GetBoolValue(const FString& PropertyName) const override;
+	virtual TOptional<double> GetNumericValue(const FString& PropertyName) const override;
+	virtual TOptional<FString> GetEnumValue(const FString& PropertyName) const override;
+	virtual TOptional<FString> GetTypeName(const FString& PropertyName) const override;
+
+	/** 
+	 * Fetch the single boolean property referenced. 
+	 * Returns nullptr if more than one property is referenced. 
+	 */
+	const UBoolProperty* GetSingleBoolProperty() const;
+
+private:
+	struct FPropertyAddress
+	{
+		TWeakObjectPtr<UObject> ObjectPtr;
+		uint8* BaseAddress;
+	};
+
+	struct FReferencedProperty
+	{
+		const UProperty* Property;
+		TArray<FPropertyAddress> Addresses;
+	};
+
+	FPropertyNode* ParentNode;
+	TArray<FReferencedProperty> Properties;
+	TSharedPtr<FEditConditionExpression> Expression;
+
+	const FReferencedProperty* FindProperty(const FString& PropertyName) const;
+};
\ No newline at end of file
diff --git a/Engine/Source/Editor/PropertyEditor/Private/EditConditionParserTests.cpp b/Engine/Source/Editor/PropertyEditor/Private/EditConditionParserTests.cpp
new file mode 100644
--- /dev/null
+++ b/Engine/Source/Editor/PropertyEditor/Private/EditConditionParserTests.cpp
@@ -0,0 +1,300 @@
+// Copyright 1998-2019 Epic Games, Inc. All Rights Reserved.
+
+#include "CoreTypes.h"
+#include "EditConditionParser.h"
+#include "Misc/AutomationTest.h"
+
+#if WITH_DEV_AUTOMATION_TESTS
+
+PRAGMA_DISABLE_OPTIMIZATION
+
+struct TestEditConditionContext : IEditConditionContext
+{
+	TMap<FString, bool> BoolValues;
+	TMap<FString, double> DoubleValues;
+	TMap<FString, FString> EnumValues;
+	FString EnumTypeName;
+
+	TestEditConditionContext(){}
+	virtual ~TestEditConditionContext() {}
+
+	virtual TOptional<bool> GetBoolValue(const FString& PropertyName) const override
+	{
+		TOptional<bool> Result;
+		if (const bool* Value = BoolValues.Find(PropertyName))
+		{
+			Result = *Value;
+		}
+		return Result;
+	}
+
+	virtual TOptional<double> GetNumericValue(const FString& PropertyName) const override
+	{
+		TOptional<double> Result;
+		if (const double* Value = DoubleValues.Find(PropertyName))
+		{
+			Result = *Value;
+		}
+		return Result;
+	}
+
+	virtual TOptional<FString> GetEnumValue(const FString& PropertyName) const override
+	{
+		TOptional<FString> Result;
+		if (const FString* Value = EnumValues.Find(PropertyName))
+		{
+			Result = *Value;
+		}
+		return Result;
+	}
+
+	virtual TOptional<FString> GetTypeName(const FString& PropertyName) const override
+	{
+		TOptional<FString> Result;
+
+		if (BoolValues.Find(PropertyName) != nullptr)
+		{
+			Result = TEXT("bool");
+		}
+		else if (DoubleValues.Find(PropertyName) != nullptr)
+		{
+			Result = TEXT("double");
+		}
+		else if (EnumValues.Find(PropertyName) != nullptr)
+		{
+			Result = EnumTypeName;
+		}
+
+		return Result;
+	}
+
+	void SetupBool(const FString& PropertyName, bool Value)
+	{
+		BoolValues.Add(PropertyName, Value);
+	}
+
+	void SetupDouble(const FString& PropertyName, double Value)
+	{
+		DoubleValues.Add(PropertyName, Value);
+	}
+
+	void SetupEnum(const FString& PropertyName, const FString& Value)
+	{
+		EnumValues.Add(PropertyName, Value);
+	}
+
+	void SetupEnumType(const FString& EnumType)
+	{
+		EnumTypeName = EnumType;
+	}
+};
+
+static bool CanParse(const FEditConditionParser& Parser, const FString& Expression, int32 ExpectedTokens, int32 ExpectedProperties)
+{
+	 TSharedPtr<FEditConditionExpression> Parsed = Parser.Parse(Expression);
+
+	 if (!Parsed.IsValid())
+	 {
+		ensureMsgf(false, TEXT("Failed to parse expression: %s"), *Expression);
+		return false;
+	 }
+	 
+	 int PropertyCount = 0;
+
+	 for (const auto& Token : Parsed->Tokens)
+	 {
+		const EditConditionParserNamespace::FPropertyToken* PropertyToken = Token.Node.Cast<EditConditionParserNamespace::FPropertyToken>();
+		if (PropertyToken != nullptr)
+		{
+			++PropertyCount;
+		}
+	 }
+
+	 return Parsed->Tokens.Num() == ExpectedTokens &&
+		 PropertyCount == ExpectedProperties;
+}
+
+IMPLEMENT_SIMPLE_AUTOMATION_TEST(FEditConditionParser_Parse, "EditConditionParser.Parse", EAutomationTestFlags::EditorContext | EAutomationTestFlags::EngineFilter)
+bool FEditConditionParser_Parse::RunTest(const FString& Parameters)
+{
+	FEditConditionParser Parser;
+	bool bResult = true;
+
+	bResult &= CanParse(Parser, TEXT("bProperty"), 1, 1);
+	bResult &= CanParse(Parser, TEXT("!bProperty"), 2, 1);
+	bResult &= CanParse(Parser, TEXT("bProperty == true"), 3, 1);
+	bResult &= CanParse(Parser, TEXT("bProperty == false"), 3, 1);
+	bResult &= CanParse(Parser, TEXT("IntProperty == 0"), 3, 1);
+	bResult &= CanParse(Parser, TEXT("IntProperty != 0"), 3, 1);
+	bResult &= CanParse(Parser, TEXT("IntProperty > 0"), 3, 1);
+	bResult &= CanParse(Parser, TEXT("IntProperty < 0"), 3, 1);
+	bResult &= CanParse(Parser, TEXT("IntProperty <= 0"), 3, 1);
+	bResult &= CanParse(Parser, TEXT("IntProperty >= 0"), 3, 1);
+	bResult &= CanParse(Parser, TEXT("Foo > Bar"), 3, 2);
+	bResult &= CanParse(Parser, TEXT("Foo && Bar"), 3, 2);
+	bResult &= CanParse(Parser, TEXT("Foo || Bar"), 3, 2);
+	bResult &= CanParse(Parser, TEXT("Foo == Bar + 5"), 5, 2);
+	bResult &= CanParse(Parser, TEXT("Foo == Bar - 5"), 5, 2);
+	bResult &= CanParse(Parser, TEXT("Foo == Bar * 5"), 5, 2);
+	bResult &= CanParse(Parser, TEXT("Foo == Bar / 5"), 5, 2);
+	bResult &= CanParse(Parser, TEXT("Enum == EType::Value"), 3, 1);
+	bResult &= CanParse(Parser, TEXT("Enum != EType::Value"), 3, 1);
+	bResult &= CanParse(Parser, TEXT("Enum != EType::Value && bProperty"), 5, 2);
+	bResult &= CanParse(Parser, TEXT("Enum == EType::Value || bProperty == false"), 7, 2);
+	bResult &= CanParse(Parser, TEXT("Enum != EType::Value || bProperty == bFoo"), 7, 3);
+	bResult &= CanParse(Parser, TEXT("Enum == EType::Value && Foo != 5"), 7, 2);
+	bResult &= CanParse(Parser, TEXT("Enum != EType::Value && Foo == Bar"), 7, 3);
+
+	return bResult;
+}
+
+bool CanEvaluate(const FEditConditionParser& Parser, const IEditConditionContext& Context, const FString& Expression, bool Expected)
+{
+	TSharedPtr<FEditConditionExpression> Parsed = Parser.Parse(Expression);
+	if (!Parsed.IsValid())
+	{
+		ensureMsgf(false, TEXT("Failed to parse expression: %s"), *Expression);
+		return false;
+	}
+
+	TOptional<bool> Result = Parser.Evaluate(*Parsed.Get(), Context);
+	if (!Result.IsSet())
+	{
+		ensureMsgf(false, TEXT("Expression failed to evaluate: %s"), *Expression);
+		return false;
+	}
+
+	if (Result.GetValue() != Expected)
+	{
+		ensureMsgf(false, TEXT("Expression evaluated to unexpected value."), *Expression);
+		return false;
+	}
+
+	return true;
+}
+
+IMPLEMENT_SIMPLE_AUTOMATION_TEST(FEditConditionParser_EvaluateBool, "EditConditionParser.EvaluateBool", EAutomationTestFlags::EditorContext | EAutomationTestFlags::EngineFilter)
+bool FEditConditionParser_EvaluateBool::RunTest(const FString& Parameters)
+{
+	FEditConditionParser Parser;
+
+	TestEditConditionContext TestContext;
+	TestContext.SetupBool(TEXT("bProperty"), true);
+
+	bool bResult = true;
+
+	bResult &= CanEvaluate(Parser, TestContext, TEXT("true"), true);
+	bResult &= CanEvaluate(Parser, TestContext, TEXT("false"), false);
+	bResult &= CanEvaluate(Parser, TestContext, TEXT("!true"), false);
+	bResult &= CanEvaluate(Parser, TestContext, TEXT("!false"), true);
+	bResult &= CanEvaluate(Parser, TestContext, TEXT("bProperty"), true);
+	bResult &= CanEvaluate(Parser, TestContext, TEXT("!bProperty"), false);
+
+	bResult &= CanEvaluate(Parser, TestContext, TEXT("bProperty == true"), true);
+	bResult &= CanEvaluate(Parser, TestContext, TEXT("bProperty == false"), false);
+	bResult &= CanEvaluate(Parser, TestContext, TEXT("bProperty == bProperty"), true);
+
+	bResult &= CanEvaluate(Parser, TestContext, TEXT("bProperty != true"), false);
+	bResult &= CanEvaluate(Parser, TestContext, TEXT("bProperty != false"), true);
+
+	bResult &= CanEvaluate(Parser, TestContext, TEXT("true && true"), true);
+	bResult &= CanEvaluate(Parser, TestContext, TEXT("true && false"), false);
+	bResult &= CanEvaluate(Parser, TestContext, TEXT("false && true"), false);
+	bResult &= CanEvaluate(Parser, TestContext, TEXT("false && false"), false);
+	bResult &= CanEvaluate(Parser, TestContext, TEXT("true && true && true"), true);
+	bResult &= CanEvaluate(Parser, TestContext, TEXT("true && true && false"), false);
+	bResult &= CanEvaluate(Parser, TestContext, TEXT("bProperty && bProperty"), true);
+	bResult &= CanEvaluate(Parser, TestContext, TEXT("bProperty && false"), false);
+	bResult &= CanEvaluate(Parser, TestContext, TEXT("false && bProperty"), false);
+
+	bResult &= CanEvaluate(Parser, TestContext, TEXT("true || true"), true);
+	bResult &= CanEvaluate(Parser, TestContext, TEXT("true || false"), true);
+	bResult &= CanEvaluate(Parser, TestContext, TEXT("false || true"), true);
+	bResult &= CanEvaluate(Parser, TestContext, TEXT("false || false"), false);
+	bResult &= CanEvaluate(Parser, TestContext, TEXT("true || true || true"), true);
+	bResult &= CanEvaluate(Parser, TestContext, TEXT("true || true || false"), true);
+	bResult &= CanEvaluate(Parser, TestContext, TEXT("bProperty || bProperty"), true);
+	bResult &= CanEvaluate(Parser, TestContext, TEXT("bProperty || false"), true);
+	bResult &= CanEvaluate(Parser, TestContext, TEXT("false || bProperty"), true);
+
+	return bResult;
+}
+
+
+IMPLEMENT_SIMPLE_AUTOMATION_TEST(FEditConditionParser_EvaluateDouble, "EditConditionParser.EvaluateDouble", EAutomationTestFlags::EditorContext | EAutomationTestFlags::EngineFilter)
+bool FEditConditionParser_EvaluateDouble::RunTest(const FString& Parameters)
+{
+	FEditConditionParser Parser;
+
+	TestEditConditionContext TestContext;
+	TestContext.SetupDouble(TEXT("Property"), 5.0);
+
+	bool bResult = true;
+
+	bResult &= CanEvaluate(Parser, TestContext, TEXT("5 == 5"), true);
+	bResult &= CanEvaluate(Parser, TestContext, TEXT("5.0 == 5.0"), true);
+	bResult &= CanEvaluate(Parser, TestContext, TEXT("Property == 5.0"), true);
+	bResult &= CanEvaluate(Parser, TestContext, TEXT("Property == 5"), true);
+	bResult &= CanEvaluate(Parser, TestContext, TEXT("Property == Property"), true);
+
+	bResult &= CanEvaluate(Parser, TestContext, TEXT("Property != 5.0"), false);
+	bResult &= CanEvaluate(Parser, TestContext, TEXT("Property != 6.0"), true);
+	bResult &= CanEvaluate(Parser, TestContext, TEXT("Property != 6"), true);
+
+	bResult &= CanEvaluate(Parser, TestContext, TEXT("Property > 4.5"), true);
+	bResult &= CanEvaluate(Parser, TestContext, TEXT("Property > 5"), false);
+	bResult &= CanEvaluate(Parser, TestContext, TEXT("Property > 6"), false);
+	bResult &= CanEvaluate(Parser, TestContext, TEXT("Property > Property"), false);
+
+	bResult &= CanEvaluate(Parser, TestContext, TEXT("Property < 4.5"), false);
+	bResult &= CanEvaluate(Parser, TestContext, TEXT("Property < 5"), false);
+	bResult &= CanEvaluate(Parser, TestContext, TEXT("Property < 6"), true);
+	bResult &= CanEvaluate(Parser, TestContext, TEXT("Property < Property"), false);
+
+	bResult &= CanEvaluate(Parser, TestContext, TEXT("Property >= 4.5"), true);
+	bResult &= CanEvaluate(Parser, TestContext, TEXT("Property >= 5"), true);
+	bResult &= CanEvaluate(Parser, TestContext, TEXT("Property >= 6"), false);
+	bResult &= CanEvaluate(Parser, TestContext, TEXT("Property >= Property"), true);
+
+	bResult &= CanEvaluate(Parser, TestContext, TEXT("Property <= 4.5"), false);
+	bResult &= CanEvaluate(Parser, TestContext, TEXT("Property <= 5"), true);
+	bResult &= CanEvaluate(Parser, TestContext, TEXT("Property <= 6"), true);
+	bResult &= CanEvaluate(Parser, TestContext, TEXT("Property <= Property"), true);
+
+	bResult &= CanEvaluate(Parser, TestContext, TEXT("Property == 2 + 3"), true);
+	bResult &= CanEvaluate(Parser, TestContext, TEXT("Property == 6 - 1"), true);
+	bResult &= CanEvaluate(Parser, TestContext, TEXT("Property == 2.5 * 2"), true);
+	bResult &= CanEvaluate(Parser, TestContext, TEXT("Property == 10 / 2"), true);
+
+	return bResult;
+}
+
+IMPLEMENT_SIMPLE_AUTOMATION_TEST(FEditConditionParser_EvaluateEnum, "EditConditionParser.EvaluateEnum", EAutomationTestFlags::EditorContext | EAutomationTestFlags::EngineFilter)
+bool FEditConditionParser_EvaluateEnum::RunTest(const FString& Parameters)
+{
+	FEditConditionParser Parser;
+
+	TestEditConditionContext TestContext;
+	TestContext.SetupEnumType(TEXT("MyEnum"));
+	TestContext.SetupEnum(TEXT("Property"), TEXT("First"));
+
+	bool bResult = true;
+	
+	bResult &= CanEvaluate(Parser, TestContext, TEXT("MyEnum::First == MyEnum::First"), true);
+	bResult &= CanEvaluate(Parser, TestContext, TEXT("MyEnum::First == MyEnum::Second"), false);
+
+	bResult &= CanEvaluate(Parser, TestContext, TEXT("MyEnum::First != MyEnum::First"), false);
+	bResult &= CanEvaluate(Parser, TestContext, TEXT("MyEnum::First != MyEnum::Second"), true);
+
+	bResult &= CanEvaluate(Parser, TestContext, TEXT("Property == MyEnum::First"), true);
+	bResult &= CanEvaluate(Parser, TestContext, TEXT("MyEnum::First == Property"), true);
+	bResult &= CanEvaluate(Parser, TestContext, TEXT("Property == MyEnum::Second"), false);
+
+	bResult &= CanEvaluate(Parser, TestContext, TEXT("Property != MyEnum::Second"), true);
+	bResult &= CanEvaluate(Parser, TestContext, TEXT("Property != MyEnum::First"), false);
+	bResult &= CanEvaluate(Parser, TestContext, TEXT("MyEnum::Second != Property"), true);
+
+	return bResult;
+}
+
+#endif // WITH_DEV_AUTOMATION_TESTS
\ No newline at end of file
diff --git a/Engine/Source/Editor/PropertyEditor/Private/IDetailsViewPrivate.h b/Engine/Source/Editor/PropertyEditor/Private/IDetailsViewPrivate.h
--- a/Engine/Source/Editor/PropertyEditor/Private/IDetailsViewPrivate.h
+++ b/Engine/Source/Editor/PropertyEditor/Private/IDetailsViewPrivate.h
@@ -8,6 +8,7 @@
 #include "PropertyNode.h"
 #include "IDetailsView.h"
 
+class FEditConditionParser;
 class FNotifyHook;
 class IDetailPropertyExtensionHandler;
 class IDetailRootObjectCustomization;
@@ -128,4 +129,6 @@ public:
 	* Restores the expansion state of property nodes for the selected object set
 	*/
 	virtual void RestoreExpandedItems(TSharedRef<FPropertyNode> StartNode) = 0;
+
+	virtual TSharedPtr<FEditConditionParser> GetEditConditionParser() const = 0;
 };
diff --git a/Engine/Source/Editor/PropertyEditor/Private/Presentation/PropertyEditor/PropertyEditor.cpp b/Engine/Source/Editor/PropertyEditor/Private/Presentation/PropertyEditor/PropertyEditor.cpp
--- a/Engine/Source/Editor/PropertyEditor/Private/Presentation/PropertyEditor/PropertyEditor.cpp
+++ b/Engine/Source/Editor/PropertyEditor/Private/Presentation/PropertyEditor/PropertyEditor.cpp
@@ -28,11 +28,9 @@ TSharedRef< FPropertyEditor > FPropertyEditor::Create( const TSharedRef< class F
 }
 
 FPropertyEditor::FPropertyEditor( const TSharedRef<FPropertyNode>& InPropertyNode, const TSharedRef<IPropertyUtilities>& InPropertyUtilities )
-	: PropertyEditConditions()
-	, PropertyHandle( NULL )
+	: PropertyHandle( NULL )
 	, PropertyNode( InPropertyNode )
 	, PropertyUtilities( InPropertyUtilities )
-	, EditConditionProperty( NULL )
 {
 	// FPropertyEditor isn't built to handle CategoryNodes
 	check( InPropertyNode->AsCategoryNode() == NULL );
@@ -41,14 +39,19 @@ FPropertyEditor::FPropertyEditor( const TSharedRef<FPropertyNode>& InPropertyNod
 
 	if( Property )
 	{
-		//see if the property supports some kind of edit condition and this isn't the "parent" property of a static array
-		const bool bStaticArray = Property->ArrayDim > 1 && InPropertyNode->GetArrayIndex() == INDEX_NONE;
+		static const FName EditConditionName = TEXT("EditCondition");
 
-		if ( Property->HasMetaData( TEXT( "EditCondition" ) ) && !bStaticArray ) 
+		//see if the property supports some kind of edit condition and this isn't the "parent" property of a static array
+		if (Property->HasMetaData(EditConditionName) && !PropertyEditorHelpers::IsStaticArray(PropertyNode.Get()))
 		{
-			if ( !GetEditConditionPropertyAddress( /*OUT*/EditConditionProperty, *InPropertyNode, PropertyEditConditions ) )
+			TSharedPtr<FEditConditionParser> Parser = PropertyUtilities->GetEditConditionParser();
+			if (Parser.IsValid())
 			{
-				EditConditionProperty = NULL;
+				EditConditionExpression = Parser->Parse(Property->GetMetaData(EditConditionName));
+				if (EditConditionExpression.IsValid())
+				{
+					EditConditionContext = MakeShareable(new FEditConditionContext(PropertyNode.Get(), EditConditionExpression));
+				}
 			}
 		}
 	}
@@ -440,7 +443,7 @@ bool FPropertyEditor::IsEditConst() const
 	return PropertyNode->IsEditConst();
 }
 
-void FPropertyEditor::SetEditConditionState( bool bShouldEnable )
+void FPropertyEditor::ToggleEditConditionState()
 {
 	const FScopedTransaction Transaction(FText::Format(LOCTEXT("SetEditConditionState", "Set {0} edit condition state "), PropertyNode->GetDisplayName()));
 
@@ -451,17 +454,25 @@ void FPropertyEditor::SetEditConditionState( bool bShouldEnable )
 	check(ParentNode != nullptr);
 
 	PropertyNode->NotifyPreChange( PropertyNode->GetProperty(), PropertyUtilities->GetNotifyHook() );
-	for ( int32 ValueIdx = 0; ValueIdx < PropertyEditConditions.Num(); ValueIdx++ )
+
+	const UBoolProperty* EditConditionProperty = EditConditionContext->GetSingleBoolProperty();
+	check(EditConditionProperty != nullptr);
+
+	FComplexPropertyNode* ComplexParentNode = ParentNode->FindComplexParent();
+	for (int32 Index = 0; Index < ComplexParentNode->GetInstancesNum(); ++Index)
 	{
+		uint8* BaseAddress = ComplexParentNode->GetMemoryOfInstance(Index);
+		check(BaseAddress != nullptr);
+
 		// Get the address corresponding to the base of this property (i.e. if a struct property, set BaseOffset to the address of value for the whole struct)
-		uint8* BaseOffset = ParentNode->GetValueAddress(PropertyEditConditions[ValueIdx].BaseAddress);
+		uint8* BaseOffset = ParentNode->GetValueAddress(BaseAddress);
 		check(BaseOffset != NULL);
 
 		uint8* ValueAddr = EditConditionProperty->ContainerPtrToValuePtr<uint8>(BaseOffset);
 
-		const bool OldValue = EditConditionProperty->GetPropertyValue( ValueAddr );
-		const bool NewValue = XOR(bShouldEnable, PropertyEditConditions[ValueIdx].bNegateValue);
-		EditConditionProperty->SetPropertyValue( ValueAddr, NewValue );
+		const bool OldValue = EditConditionProperty->GetPropertyValue(ValueAddr);
+		const bool NewValue = !OldValue;
+		EditConditionProperty->SetPropertyValue(ValueAddr, NewValue);
 
 		if (ObjectNode != nullptr)
 		{
@@ -543,8 +554,7 @@ void FPropertyEditor::OnGetActorFiltersForSceneOutliner( TSharedPtr<SceneOutline
 
 bool FPropertyEditor::IsPropertyEditingEnabled() const
 {
-	return ( PropertyUtilities->IsPropertyEditingEnabled() ) && 
-		(EditConditionProperty == NULL || IsEditConditionMet( EditConditionProperty, PropertyEditConditions ));
+	return ( PropertyUtilities->IsPropertyEditingEnabled() ) && (!HasEditCondition() || IsEditConditionMet());
 }
 
 void FPropertyEditor::ForceRefresh()
@@ -559,17 +569,48 @@ void FPropertyEditor::RequestRefresh()
 
 bool FPropertyEditor::HasEditCondition() const 
 { 
-	return EditConditionProperty != NULL; 
+	return EditConditionExpression.IsValid();
 }
 
 bool FPropertyEditor::IsEditConditionMet() const 
 { 
-	return IsEditConditionMet( EditConditionProperty, PropertyEditConditions ); 
+	if (HasEditCondition())
+	{
+		TSharedPtr<FEditConditionParser> EditConditionParser = PropertyUtilities->GetEditConditionParser();
+		if (EditConditionParser.IsValid())
+		{
+			TOptional<bool> Result = EditConditionParser->Evaluate(*EditConditionExpression.Get(), *EditConditionContext.Get());
+			if (Result.IsSet())
+			{
+				return Result.GetValue();
+			}
+		}
+	}
+
+	return true;
 }
 
 bool FPropertyEditor::SupportsEditConditionToggle() const
 {
-	return SupportsEditConditionToggle( PropertyNode->GetProperty() );
+	UProperty* Property = PropertyNode->GetProperty();
+
+	static const FName Name_HideEditConditionToggle("HideEditConditionToggle");
+	if (!Property->HasMetaData(Name_HideEditConditionToggle) && EditConditionExpression.IsValid())
+	{
+		const UBoolProperty* ConditionalProperty = EditConditionContext->GetSingleBoolProperty();
+		if (ConditionalProperty != nullptr)
+		{
+			// If it's editable, then only put an inline toggle box if the metadata specifies it
+			static const FName Name_InlineEditConditionToggle("InlineEditConditionToggle");
+			if (ConditionalProperty->HasAllPropertyFlags(CPF_Edit) && 
+				ConditionalProperty->HasMetaData(Name_InlineEditConditionToggle))
+			{
+				return true;
+			}
+		}
+	}
+
+	return false;
 }
 
 void FPropertyEditor::AddPropertyEditorChild( const TSharedRef<FPropertyEditor>& Child )
@@ -602,145 +643,6 @@ TSharedRef< IPropertyHandle > FPropertyEditor::GetPropertyHandle() const
 	return PropertyHandle.ToSharedRef();
 }
 
-bool FPropertyEditor::IsEditConditionMet( UBoolProperty* ConditionProperty, const TArray<FPropertyConditionInfo>& ConditionValues ) const
-{
-	check(ConditionProperty);
-
-	bool bResult = false;
-
-	FPropertyNode* ParentNode = PropertyNode->GetParentNode();
-	if (ParentNode)
-	{
-		bool bAllConditionsMet = true;
-		for (int32 ValueIdx = 0; bAllConditionsMet && ValueIdx < ConditionValues.Num(); ValueIdx++)
-		{
-			if (!ConditionValues[ValueIdx].Object.IsValid())
-			{
-				bAllConditionsMet = false;
-				break;
-			}
-
-			uint8* BaseOffset = ParentNode->GetValueAddress(ConditionValues[ValueIdx].BaseAddress);
-			if (!BaseOffset)
-			{
-				bAllConditionsMet = false;
-				break;
-			}
-
-			uint8* ValueAddr = EditConditionProperty->ContainerPtrToValuePtr<uint8>(BaseOffset);
-
-			if (ConditionValues[ValueIdx].bNegateValue)
-			{
-				bAllConditionsMet = !ConditionProperty->GetPropertyValue(ValueAddr);
-			}
-			else
-			{
-				bAllConditionsMet = ConditionProperty->GetPropertyValue(ValueAddr);
-			}
-		}
-
-		bResult = bAllConditionsMet;
-	}
-
-	return bResult;
-}
-
-bool FPropertyEditor::GetEditConditionPropertyAddress( UBoolProperty*& ConditionProperty, FPropertyNode& InPropertyNode, TArray<FPropertyConditionInfo>& ConditionPropertyAddresses )
-{
-	bool bResult = false;
-	bool bNegate = false;
-	UBoolProperty* EditConditionProperty = PropertyCustomizationHelpers::GetEditConditionProperty(InPropertyNode.GetProperty(), bNegate);
-	if ( EditConditionProperty != NULL )
-	{
-		FPropertyNode* ParentNode = InPropertyNode.GetParentNode();
-		check(ParentNode);
-
-		UProperty* Property = InPropertyNode.GetProperty();
-		if (Property)
-		{
-			bool bStaticArray = (Property->ArrayDim > 1) && (InPropertyNode.GetArrayIndex() != INDEX_NONE);
-			if (bStaticArray)
-			{
-				//in the case of conditional static arrays, we have to go up one more level to get the proper parent struct.
-				ParentNode = ParentNode->GetParentNode();
-				check(ParentNode);
-			}
-		}
-
-		auto ComplexParentNode = ParentNode->FindComplexParent();
-		if (ComplexParentNode)
-		{
-			for (int32 Index = 0; Index < ComplexParentNode->GetInstancesNum(); ++Index)
-			{
-				uint8* BaseAddress = ComplexParentNode->GetMemoryOfInstance(Index);
-				if (BaseAddress)
-				{
-					// Get the address corresponding to the base of this property (i.e. if a struct property, set BaseOffset to the address of value for the whole struct)
-					uint8* BaseOffset = ParentNode->GetValueAddress(BaseAddress);
-					check(BaseOffset != NULL);
-
-					FPropertyConditionInfo NewCondition;
-					// now calculate the address of the property value being used as the condition and add it to the array.
-					NewCondition.Object = ComplexParentNode->AsStructureNode() ? TWeakObjectPtr<UObject>(ComplexParentNode->GetBaseStructure()) : ComplexParentNode->GetInstanceAsUObject(Index);
-					NewCondition.BaseAddress = BaseAddress;
-					NewCondition.bNegateValue = bNegate;
-					ConditionPropertyAddresses.Add(NewCondition);
-					bResult = true;
-				}
-			}
-		}
-	}
-
-	if ( bResult )
-	{
-		// set the output variable
-		ConditionProperty = EditConditionProperty;
-	}
-
-	return bResult;
-}
-
-bool FPropertyEditor::SupportsEditConditionToggle( UProperty* InProperty ) 
-{
-	bool bShowEditConditionToggle = false;
-
-	static const FName Name_HideEditConditionToggle("HideEditConditionToggle");
-	if (!InProperty->HasMetaData(Name_HideEditConditionToggle))
-	{
-		bool bNegateValue = false;
-		UBoolProperty* ConditionalProperty = PropertyCustomizationHelpers::GetEditConditionProperty( InProperty, bNegateValue );
-		if( ConditionalProperty != NULL )
-		{
-			if (!ConditionalProperty->HasAllPropertyFlags(CPF_Edit))
-			{
-				// Warn if the editcondition property is not marked as editable; this will break when the component is added to a Blueprint.
-				UObject* PropertyScope = InProperty->GetOwnerClass();
-				UObject* ConditionalPropertyScope = ConditionalProperty->GetOwnerClass();
-				if (!PropertyScope)
-				{
-					PropertyScope = InProperty->GetOwnerStruct();
-					ConditionalPropertyScope = ConditionalProperty->GetOwnerStruct();
-				}
-				const FString PropertyScopeName = PropertyScope ? PropertyScope->GetName() : FString();
-				const FString ConditionalPropertyScopeName = ConditionalPropertyScope ? ConditionalPropertyScope->GetName() : FString();
-
-				bShowEditConditionToggle = true;
-			}
-			else
-			{
-				// If it's editable, then only put an inline toggle box if the metadata specifies it
-				static const FName Name_InlineEditConditionToggle("InlineEditConditionToggle");
-				if (ConditionalProperty->HasMetaData(Name_InlineEditConditionToggle))
-				{
-					bShowEditConditionToggle = true;
-				}
-			}		
-		}
-	}
-
-	return bShowEditConditionToggle;
-}
-
 void FPropertyEditor::SyncToObjectsInNode( const TWeakPtr< FPropertyNode >& WeakPropertyNode )
 {
 #if WITH_EDITOR
diff --git a/Engine/Source/Editor/PropertyEditor/Private/Presentation/PropertyEditor/PropertyEditor.h b/Engine/Source/Editor/PropertyEditor/Private/Presentation/PropertyEditor/PropertyEditor.h
--- a/Engine/Source/Editor/PropertyEditor/Private/Presentation/PropertyEditor/PropertyEditor.h
+++ b/Engine/Source/Editor/PropertyEditor/Private/Presentation/PropertyEditor/PropertyEditor.h
@@ -5,6 +5,7 @@
 #include "PropertyNode.h"
 #include "PropertyHandle.h"
 #include "Editor/SceneOutliner/Public/SceneOutlinerFwd.h"
+#include "EditConditionParser.h"
 
 class FPropertyEditor : public TSharedFromThis< FPropertyEditor >	
 {
@@ -38,6 +39,7 @@ public:
 	void RequestRefresh();
 
 	bool SupportsEditConditionToggle() const;
+	void ToggleEditConditionState();
 
 	/**	@return Whether the property is editconst */
 	bool IsEditConst() const;
@@ -45,9 +47,6 @@ public:
 	/**	@return Whether the property has a condition which must be met before allowing editing of it's value */
 	bool HasEditCondition() const;
 
-	/**	Sets the state of the property's edit condition to the specified value */
-	void SetEditConditionState( bool bShouldEnable );
-
 	/**	@return Whether the condition has been met to allow editing of this property's value */
 	bool IsEditConditionMet() const;
 
@@ -97,17 +96,6 @@ public:
 private:
 	FPropertyEditor( const TSharedRef< class FPropertyNode >& InPropertyNode, const TSharedRef<class IPropertyUtilities >& InPropertyUtilities );
 
-	/** Stores information about property condition metadata. */
-	struct FPropertyConditionInfo
-	{
-	public:
-		TWeakObjectPtr<UObject> Object;
-
-		uint8* BaseAddress;
-		/** Whether the condition should be negated. */
-		bool bNegateValue;
-	};
-
 	void OnUseSelected();
 	void OnAddItem();
 	void OnAddGivenItem(const FString InGivenItem);
@@ -118,29 +106,8 @@ private:
 	void OnBrowseTo();
 	void OnEmptyArray();
 
-	/**
-	 * Returns true if the value of the conditional property matches the value required.  Indicates whether editing or otherwise interacting with this item's
-	 * associated property should be allowed.
-	 */
-	bool IsEditConditionMet( UBoolProperty* ConditionProperty, const TArray<FPropertyConditionInfo>& ConditionValues ) const;
-
-	/**
-	 * Finds the property being used to determine whether this item's associated property should be editable/expandable.
-	 * If the property is successfully found, ConditionPropertyAddresses will be filled with the addresses of the conditional properties' values.
-	 *
-	 * @param	ConditionProperty	receives the value of the property being used to control this item's ability to be edited.
-	 * @param	ConditionPropertyAddresses	receives the addresses of the actual property values for the conditional property
-	 *
-	 * @return	true if both the conditional property and property value addresses were successfully determined.
-	 */
-	static bool GetEditConditionPropertyAddress( UBoolProperty*& ConditionProperty, FPropertyNode& InPropertyNode, TArray<FPropertyConditionInfo>& ConditionPropertyAddresses );
-
-	static bool SupportsEditConditionToggle( UProperty* InProperty );
-
 private:
 
-	TArray< FPropertyConditionInfo > PropertyEditConditions;
-
 	TArray< TSharedRef< FPropertyEditor > >	ChildPropertyEditors;
 
 	/** Property handle for actually reading/writing the value of a property */
@@ -152,6 +119,7 @@ private:
 	/** The property view where this widget resides */
 	TSharedRef< class IPropertyUtilities > PropertyUtilities;
 
-	/** Edit condition property used to determine if this property editor can modify its property */
-	class UBoolProperty* EditConditionProperty;
+	/** Edit condition expression used to determine if this property editor can modify its property */
+	TSharedPtr<FEditConditionExpression> EditConditionExpression;
+	TSharedPtr<FEditConditionContext> EditConditionContext;
 };
diff --git a/Engine/Source/Editor/PropertyEditor/Private/Presentation/PropertyTable/PropertyTable.cpp b/Engine/Source/Editor/PropertyEditor/Private/Presentation/PropertyTable/PropertyTable.cpp
--- a/Engine/Source/Editor/PropertyEditor/Private/Presentation/PropertyTable/PropertyTable.cpp
+++ b/Engine/Source/Editor/PropertyEditor/Private/Presentation/PropertyTable/PropertyTable.cpp
@@ -11,6 +11,7 @@
 #include "Presentation/PropertyTable/PropertyTableObjectNameColumn.h"
 #include "Presentation/PropertyTable/PropertyTablePropertyNameColumn.h"
 
+#include "EditConditionParser.h"
 
 #define LOCTEXT_NAMESPACE "PropertyTable"
 
@@ -32,6 +33,7 @@ FPropertyTable::FPropertyTable()
 	, AllowUserToChangeRoot( true )
 	, bRefreshRequested( false )
 	, Orientation( EPropertyTableOrientation::AlignPropertiesInColumns )
+	, EditConditionParser( new FEditConditionParser )
 {
 }
 
@@ -103,6 +105,11 @@ TSharedPtr<class FAssetThumbnailPool> FPropertyTable::GetThumbnailPool() const
 	return NULL;
 }
 
+TSharedPtr<FEditConditionParser> FPropertyTable::GetEditConditionParser() const
+{
+	return EditConditionParser;
+}
+
 bool FPropertyTable::GetIsUserAllowedToChangeRoot()
 {
 	return AllowUserToChangeRoot;
diff --git a/Engine/Source/Editor/PropertyEditor/Private/Presentation/PropertyTable/PropertyTable.h b/Engine/Source/Editor/PropertyEditor/Private/Presentation/PropertyTable/PropertyTable.h
--- a/Engine/Source/Editor/PropertyEditor/Private/Presentation/PropertyTable/PropertyTable.h
+++ b/Engine/Source/Editor/PropertyEditor/Private/Presentation/PropertyTable/PropertyTable.h
@@ -28,6 +28,7 @@ public:
 	virtual void EnqueueDeferredAction( FSimpleDelegate DeferredAction ) override;
 	virtual TSharedPtr<class FAssetThumbnailPool> GetThumbnailPool() const override;
 	virtual void NotifyFinishedChangingProperties(const FPropertyChangedEvent& PropertyChangedEvent) override {}
+	virtual TSharedPtr<FEditConditionParser> GetEditConditionParser() const override;
 
 	virtual bool GetIsUserAllowedToChangeRoot() override;
 	virtual void SetIsUserAllowedToChangeRoot( bool InAllowUserToChangeRoot ) override;
@@ -218,5 +219,6 @@ private:
 
 	/** The Orientation of this table, I.e. do we swap columns and rows */
 	EPropertyTableOrientation::Type Orientation;
+	TSharedRef<FEditConditionParser> EditConditionParser;
 };
 
diff --git a/Engine/Source/Editor/PropertyEditor/Private/PropertyEditorHelpers.cpp b/Engine/Source/Editor/PropertyEditor/Private/PropertyEditorHelpers.cpp
--- a/Engine/Source/Editor/PropertyEditor/Private/PropertyEditorHelpers.cpp
+++ b/Engine/Source/Editor/PropertyEditor/Private/PropertyEditorHelpers.cpp
@@ -362,7 +362,7 @@ void SEditConditionWidget::OnEditConditionCheckChanged( ECheckBoxState CheckStat
 
 	if( PropertyEditor.IsValid() && PropertyEditor->HasEditCondition() && PropertyEditor->SupportsEditConditionToggle() )
 	{
-		PropertyEditor->SetEditConditionState( CheckState == ECheckBoxState::Checked );
+		PropertyEditor->ToggleEditConditionState();
 	}
 	else
 	{
diff --git a/Engine/Source/Editor/PropertyEditor/Private/PropertyRowGenerator.cpp b/Engine/Source/Editor/PropertyEditor/Private/PropertyRowGenerator.cpp
--- a/Engine/Source/Editor/PropertyEditor/Private/PropertyRowGenerator.cpp
+++ b/Engine/Source/Editor/PropertyEditor/Private/PropertyRowGenerator.cpp
@@ -12,13 +12,16 @@
 #include "DetailLayoutHelpers.h"
 #include "PropertyHandleImpl.h"
 #include "IPropertyGenerationUtilities.h"
+#include "EditConditionParser.h"
 
 class FPropertyRowGeneratorUtilities : public IPropertyUtilities
 {
 public:
 	FPropertyRowGeneratorUtilities(FPropertyRowGenerator& InGenerator)
 		: Generator(&InGenerator)
-	{}
+		, EditConditionParser(new FEditConditionParser)
+	{
+	}
 
 	void ResetGenerator()
 	{
@@ -87,8 +90,15 @@ public:
 	{
 		return Generator != nullptr && Generator->HasClassDefaultObject();
 	}
+
+	virtual TSharedPtr<FEditConditionParser> GetEditConditionParser() const override
+	{
+		return EditConditionParser;
+	}
+
 private:
 	FPropertyRowGenerator* Generator;
+	TSharedPtr<FEditConditionParser> EditConditionParser;
 };
 
 
diff --git a/Engine/Source/Editor/PropertyEditor/Private/SPropertyTreeViewImpl.cpp b/Engine/Source/Editor/PropertyEditor/Private/SPropertyTreeViewImpl.cpp
--- a/Engine/Source/Editor/PropertyEditor/Private/SPropertyTreeViewImpl.cpp
+++ b/Engine/Source/Editor/PropertyEditor/Private/SPropertyTreeViewImpl.cpp
@@ -31,6 +31,7 @@ public:
 	FPropertyUtilitiesTreeView( SPropertyTreeViewImpl& InView )
 		: View( InView )
 	{
+		EditConditionParser = MakeShareable(new FEditConditionParser());
 	}
 
 	virtual class FNotifyHook* GetNotifyHook() const override
@@ -95,8 +96,14 @@ public:
 	{
 		return false;
 	}
-private:
 
+	virtual TSharedPtr<FEditConditionParser> GetEditConditionParser() const
+	{
+		return EditConditionParser;
+	}
+
+private:
+	TSharedPtr<FEditConditionParser> EditConditionParser;
 	SPropertyTreeViewImpl& View;
 };
 
diff --git a/Engine/Source/Editor/PropertyEditor/Private/SSingleProperty.cpp b/Engine/Source/Editor/PropertyEditor/Private/SSingleProperty.cpp
--- a/Engine/Source/Editor/PropertyEditor/Private/SSingleProperty.cpp
+++ b/Engine/Source/Editor/PropertyEditor/Private/SSingleProperty.cpp
@@ -90,6 +90,11 @@ public:
 		return false;
 	}
 
+	virtual TSharedPtr<FEditConditionParser> GetEditConditionParser() const override
+	{
+		return nullptr;
+	}
+
 private:
 	TWeakPtr< SSingleProperty > View;
 };
diff --git a/Engine/Source/Editor/PropertyEditor/Private/UserInterface/PropertyDetails/PropertyDetailsUtilities.cpp b/Engine/Source/Editor/PropertyEditor/Private/UserInterface/PropertyDetails/PropertyDetailsUtilities.cpp
--- a/Engine/Source/Editor/PropertyEditor/Private/UserInterface/PropertyDetails/PropertyDetailsUtilities.cpp
+++ b/Engine/Source/Editor/PropertyEditor/Private/UserInterface/PropertyDetails/PropertyDetailsUtilities.cpp
@@ -73,4 +73,9 @@ const TArray<TWeakObjectPtr<UObject>>& FPropertyDetailsUtilities::GetSelectedObj
 bool FPropertyDetailsUtilities::HasClassDefaultObject() const
 {
 	return DetailsView.HasClassDefaultObject();
-}
\ No newline at end of file
+}
+
+TSharedPtr<FEditConditionParser> FPropertyDetailsUtilities::GetEditConditionParser() const
+{
+	return DetailsView.GetEditConditionParser();
+}
diff --git a/Engine/Source/Editor/PropertyEditor/Private/UserInterface/PropertyDetails/PropertyDetailsUtilities.h b/Engine/Source/Editor/PropertyEditor/Private/UserInterface/PropertyDetails/PropertyDetailsUtilities.h
--- a/Engine/Source/Editor/PropertyEditor/Private/UserInterface/PropertyDetails/PropertyDetailsUtilities.h
+++ b/Engine/Source/Editor/PropertyEditor/Private/UserInterface/PropertyDetails/PropertyDetailsUtilities.h
@@ -26,6 +26,8 @@ public:
 	virtual bool DontUpdateValueWhileEditing() const override;
 	const TArray<TWeakObjectPtr<UObject>>& GetSelectedObjects() const override;
 	virtual bool HasClassDefaultObject() const override;
+	virtual TSharedPtr<FEditConditionParser> GetEditConditionParser() const override;
+
 private:
 	IDetailsViewPrivate& DetailsView;
 };
diff --git a/Engine/Source/Editor/PropertyEditor/Private/UserInterface/PropertyEditor/SPropertyEditorTableRow.cpp b/Engine/Source/Editor/PropertyEditor/Private/UserInterface/PropertyEditor/SPropertyEditorTableRow.cpp
--- a/Engine/Source/Editor/PropertyEditor/Private/UserInterface/PropertyEditor/SPropertyEditorTableRow.cpp
+++ b/Engine/Source/Editor/PropertyEditor/Private/UserInterface/PropertyEditor/SPropertyEditorTableRow.cpp
@@ -194,7 +194,7 @@ const FSlateBrush* SPropertyEditorTableRow::OnGetFavoriteImage() const
 
 void SPropertyEditorTableRow::OnEditConditionCheckChanged( ECheckBoxState CheckState )
 {
-	PropertyEditor->SetEditConditionState( CheckState == ECheckBoxState::Checked );
+	PropertyEditor->ToggleEditConditionState();
 }
 
 ECheckBoxState SPropertyEditorTableRow::OnGetEditConditionCheckState() const
diff --git a/Engine/Source/Editor/PropertyEditor/Public/IPropertyUtilities.h b/Engine/Source/Editor/PropertyEditor/Public/IPropertyUtilities.h
--- a/Engine/Source/Editor/PropertyEditor/Public/IPropertyUtilities.h
+++ b/Engine/Source/Editor/PropertyEditor/Public/IPropertyUtilities.h
@@ -5,6 +5,7 @@
 #include "AssetThumbnail.h"
 
 struct FPropertyChangedEvent;
+class FEditConditionParser;
 
 /**
  * Settings for property editor widgets that call up to the base container for the widgets
@@ -28,4 +29,5 @@ public:
 	virtual const TArray<TWeakObjectPtr<UObject>>& GetSelectedObjects() const = 0;
 	/** If a customization standalone widget is used, the value should be update only once, when its window is closed */
 	virtual bool DontUpdateValueWhileEditing() const = 0;
+	virtual TSharedPtr<FEditConditionParser> GetEditConditionParser() const = 0;
 };
