Missed files in previous submit.

diff --git a/Engine/Source/Editor/PropertyEditor/Private/SDetailsViewBase.cpp b/Engine/Source/Editor/PropertyEditor/Private/SDetailsViewBase.cpp
--- a/Engine/Source/Editor/PropertyEditor/Private/SDetailsViewBase.cpp
+++ b/Engine/Source/Editor/PropertyEditor/Private/SDetailsViewBase.cpp
@@ -23,6 +23,7 @@
 #include "Widgets/Input/SSearchBox.h"
 #include "Classes/EditorStyleSettings.h"
 #include "DetailLayoutHelpers.h"
+#include "EditConditionParser.h"
 
 SDetailsViewBase::~SDetailsViewBase()
 {
@@ -477,6 +478,16 @@ TSharedPtr<FAssetThumbnailPool> SDetailsViewBase::GetThumbnailPool() const
 	return ThumbnailPool;
 }
 
+TSharedPtr<FEditConditionParser> SDetailsViewBase::GetEditConditionParser() const
+{
+	if (!EditConditionParser.IsValid())
+	{
+		EditConditionParser = MakeShareable(new FEditConditionParser);
+	}
+
+	return EditConditionParser;
+}
+
 void SDetailsViewBase::NotifyFinishedChangingProperties(const FPropertyChangedEvent& PropertyChangedEvent)
 {
 	OnFinishedChangingPropertiesDelegate.Broadcast(PropertyChangedEvent);
@@ -1002,7 +1013,7 @@ void SDetailsViewBase::UpdateFilteredDetails()
 
 	FDetailNodeList InitialRootNodeList;
 	
-	NumVisbleTopLevelObjectNodes = 0;
+	NumVisibleTopLevelObjectNodes = 0;
 	FRootPropertyNodeList& RootPropertyNodes = GetRootNodes();
 
 	if( GetDefault<UEditorStyleSettings>()->bShowAllAdvancedDetails )
@@ -1037,7 +1048,7 @@ void SDetailsViewBase::UpdateFilteredDetails()
 				if (LayoutRoots.Num() > 0)
 				{
 					// A top level object nodes has a non-filtered away root so add one to the total number we have
-					++NumVisbleTopLevelObjectNodes;
+					++NumVisibleTopLevelObjectNodes;
 
 					InitialRootNodeList.Append(LayoutRoots);
 				}
diff --git a/Engine/Source/Editor/PropertyEditor/Private/SDetailsViewBase.h b/Engine/Source/Editor/PropertyEditor/Private/SDetailsViewBase.h
--- a/Engine/Source/Editor/PropertyEditor/Private/SDetailsViewBase.h
+++ b/Engine/Source/Editor/PropertyEditor/Private/SDetailsViewBase.h
@@ -86,7 +86,7 @@ public:
 		, bIsLocked(false)
 		, bHasOpenColorPicker(false)
 		, bDisableCustomDetailLayouts( false )
-		, NumVisbleTopLevelObjectNodes(0)
+		, NumVisibleTopLevelObjectNodes(0)
 	{
 	}
 
@@ -107,10 +107,9 @@ public:
 
 	virtual int32 GetNumVisibleTopLevelObjects() const override
 	{
-		return NumVisbleTopLevelObjectNodes;
+		return NumVisibleTopLevelObjectNodes;
 	}
 
-
 	/** @return The identifier for this details view, or NAME_None is this view is anonymous */
 	virtual FName GetIdentifier() const override
 	{
@@ -162,6 +161,7 @@ public:
 	virtual void NotifyFinishedChangingProperties(const FPropertyChangedEvent& PropertyChangedEvent) override;
 	void RefreshTree() override;
 	TSharedPtr<class FAssetThumbnailPool> GetThumbnailPool() const override;
+	TSharedPtr<FEditConditionParser> GetEditConditionParser() const override;
 	TSharedPtr<IPropertyUtilities> GetPropertyUtilities() override;
 	void CreateColorPickerWindow(const TSharedRef< class FPropertyEditor >& PropertyEditor, bool bUseAlpha) override;
 	virtual void UpdateSinglePropertyMap(TSharedPtr<FComplexPropertyNode> InRootPropertyNode, FDetailLayoutData& LayoutData, bool bIsExternal) override;
@@ -422,7 +422,7 @@ protected:
 	/** True if we want to skip generation of custom layouts for displayed object */
 	bool bDisableCustomDetailLayouts;
 
-	int32 NumVisbleTopLevelObjectNodes;
+	int32 NumVisibleTopLevelObjectNodes;
 
 	/** Delegate for overriding the show modified filter */
 	FSimpleDelegate CustomFilterDelegate;
@@ -431,4 +431,6 @@ protected:
 	bool bCustomFilterActive;
 
 	FText CustomFilterLabel;
+
+	mutable TSharedPtr<FEditConditionParser> EditConditionParser;
 };
