Removed optimization pragma.

diff --git a/Engine/Source/Editor/PropertyEditor/Private/EditConditionParserTests.cpp b/Engine/Source/Editor/PropertyEditor/Private/EditConditionParserTests.cpp
--- a/Engine/Source/Editor/PropertyEditor/Private/EditConditionParserTests.cpp
+++ b/Engine/Source/Editor/PropertyEditor/Private/EditConditionParserTests.cpp
@@ -6,8 +6,6 @@
 
 #if WITH_DEV_AUTOMATION_TESTS
 
-PRAGMA_DISABLE_OPTIMIZATION
-
 struct TestEditConditionContext : IEditConditionContext
 {
 	TMap<FString, bool> BoolValues;
