Missed include causing Clang build to fail.

diff --git a/Engine/Source/Editor/PropertyEditor/Private/EditConditionParser.h b/Engine/Source/Editor/PropertyEditor/Private/EditConditionParser.h
--- a/Engine/Source/Editor/PropertyEditor/Private/EditConditionParser.h
+++ b/Engine/Source/Editor/PropertyEditor/Private/EditConditionParser.h
@@ -2,6 +2,8 @@
 
 #pragma once
 
+#include "Misc/ExpressionParserTypes.h"
+
 namespace EditConditionParserNamespace
 {
 	struct FPropertyToken 
