Fixed EditCondition triggering ensures when an EditCondition uses a property that does not exist in a struct (anymore).

Fixed enum value evaluation when a property is stored as a TEnumAsByte. Split EditConditionContext into its own file. Added unit tests using UObject to confirm that edit conditions don't break in the future.

diff --git a/Engine/Source/Editor/PropertyEditor/Private/EditConditionContext.cpp b/Engine/Source/Editor/PropertyEditor/Private/EditConditionContext.cpp
new file mode 100644
--- /dev/null
+++ b/Engine/Source/Editor/PropertyEditor/Private/EditConditionContext.cpp
@@ -0,0 +1,245 @@
+// Copyright 1998-2019 Epic Games, Inc. All Rights Reserved.
+
+#include "EditConditionContext.h"
+#include "EditConditionParser.h"
+
+#include "PropertyNode.h"
+#include "PropertyEditorHelpers.h"
+
+DEFINE_LOG_CATEGORY_STATIC(LogEditCondition, Log, All);
+
+FEditConditionContext::FEditConditionContext(FPropertyNode& InPropertyNode)
+{
+	PropertyNode = InPropertyNode.AsShared();
+
+	FComplexPropertyNode* ComplexParentNode = FindComplexParent();
+	check(ComplexParentNode);
+
+	const UProperty* Property = InPropertyNode.GetProperty();
+	check(Property);
+}
+
+FComplexPropertyNode* FEditConditionContext::FindComplexParent() const
+{
+	if (!PropertyNode.IsValid())
+	{
+		return nullptr;
+	}
+
+	TSharedPtr<FPropertyNode> PinnedNode = PropertyNode.Pin();
+
+	FPropertyNode* SearchStart = PinnedNode.Get();
+	if (PropertyEditorHelpers::IsStaticArray(*SearchStart))
+	{
+		//in the case of conditional static arrays, we have to go up one more level to get the proper parent struct.
+		SearchStart = SearchStart->GetParentNode();
+	}
+
+	return SearchStart->FindComplexParent();
+}
+
+const UBoolProperty* FEditConditionContext::GetSingleBoolProperty(const TSharedPtr<FEditConditionExpression>& Expression) const
+{
+	if (!PropertyNode.IsValid())
+	{
+		return nullptr;
+	}
+
+	const UProperty* Property = PropertyNode.Pin()->GetProperty();
+
+	const UBoolProperty* BoolProperty = nullptr;
+	for (const FCompiledToken& Token : Expression->Tokens)
+	{
+		if (const EditConditionParserTokens::FPropertyToken* PropertyToken = Token.Node.Cast<EditConditionParserTokens::FPropertyToken>())
+		{
+			if (BoolProperty != nullptr)
+			{
+				// second property token
+				return nullptr;
+			}
+
+			const UProperty* Field = FindField<UProperty>(Property->GetOwnerStruct(), *PropertyToken->PropertyName);
+			BoolProperty = Cast<UBoolProperty>(Field);
+			
+			// not a bool
+			if (BoolProperty == nullptr)
+			{
+				return nullptr;
+			}
+		}
+	}
+
+	return BoolProperty;
+}
+
+template<typename T>
+T* FindTypedField(const TWeakPtr<FPropertyNode>& PropertyNode, const FString& PropertyName)
+{
+	if (PropertyNode.IsValid())
+	{
+		TSharedPtr<FPropertyNode> PinnedNode = PropertyNode.Pin();
+		const UProperty* Property = PinnedNode->GetProperty();
+
+		UProperty* Field = FindField<UProperty>(Property->GetOwnerStruct(), *PropertyName);
+		if (Field == nullptr)
+		{
+			UE_LOG(LogEditCondition, Error, TEXT("EditCondition parsing failed: Field name %s was not found in struct %s."), *PropertyName, *Property->GetOwnerStruct()->GetName());
+			return nullptr;
+		}
+
+		return Cast<T>(Field);
+	}
+
+	return nullptr;
+}
+
+TOptional<bool> FEditConditionContext::GetBoolValue(const FString& PropertyName) const
+{
+	const UBoolProperty* BoolProperty = FindTypedField<UBoolProperty>(PropertyNode, PropertyName);
+	if (BoolProperty == nullptr)
+	{
+		return TOptional<bool>();
+	}
+
+	TSharedPtr<FPropertyNode> PinnedNode = PropertyNode.Pin();
+	
+	TOptional<bool> Result;
+
+	FComplexPropertyNode* ComplexParentNode = FindComplexParent();
+	for (int32 Index = 0; Index < ComplexParentNode->GetInstancesNum(); ++Index)
+	{
+		uint8* BasePtr = ComplexParentNode->GetMemoryOfInstance(Index);
+		uint8* ValuePtr = BoolProperty->ContainerPtrToValuePtr<uint8>(BasePtr);
+
+		bool bValue = BoolProperty->GetPropertyValue(ValuePtr);
+		if (!Result.IsSet())
+		{
+			Result = bValue;
+		}
+		else if (Result.GetValue() != bValue)
+		{
+			// all values aren't the same...
+			return TOptional<bool>();
+		}
+	}
+
+	return Result;
+}
+
+TOptional<double> FEditConditionContext::GetNumericValue(const FString& PropertyName) const
+{
+	const UNumericProperty* NumericProperty = FindTypedField<UNumericProperty>(PropertyNode, PropertyName);
+	if (NumericProperty == nullptr)
+	{
+		return TOptional<double>();
+	}
+
+	TSharedPtr<FPropertyNode> PinnedNode = PropertyNode.Pin();
+
+	TOptional<double> Result;
+
+	FComplexPropertyNode* ComplexParentNode = FindComplexParent();
+	for (int32 Index = 0; Index < ComplexParentNode->GetInstancesNum(); ++Index)
+	{
+		uint8* BasePtr = ComplexParentNode->GetMemoryOfInstance(Index);
+		uint8* ValuePtr = NumericProperty->ContainerPtrToValuePtr<uint8>(BasePtr);
+
+		double Value = 0;
+
+		if (NumericProperty->IsInteger())
+		{
+			Value = (double) NumericProperty->GetSignedIntPropertyValue(ValuePtr);
+		}
+		else if (NumericProperty->IsFloatingPoint())
+		{
+			Value = NumericProperty->GetFloatingPointPropertyValue(ValuePtr);
+		}
+
+		if (!Result.IsSet())
+		{
+			Result = Value;
+		}
+		else if (!FMath::IsNearlyEqual(Result.GetValue(), Value))
+		{
+			// all values aren't the same...
+			return TOptional<double>();
+		}
+	}
+
+	return Result;
+}
+
+TOptional<FString> FEditConditionContext::GetEnumValue(const FString& PropertyName) const
+{
+	const UProperty* Property = FindTypedField<UProperty>(PropertyNode, PropertyName);
+	if (Property == nullptr)
+	{
+		return TOptional<FString>();
+	}
+
+	const UEnum* EnumType = nullptr;
+	const UNumericProperty* NumericProperty = nullptr;
+	if (const UEnumProperty* EnumProperty = Cast<UEnumProperty>(Property))
+	{
+		NumericProperty = EnumProperty->GetUnderlyingProperty();
+		EnumType = EnumProperty->GetEnum();
+	}
+	else if (const UByteProperty* ByteProperty = Cast<UByteProperty>(Property))
+	{
+		NumericProperty = ByteProperty;
+		EnumType = ByteProperty->GetIntPropertyEnum();
+	}
+	else
+	{
+		return TOptional<FString>();
+	}
+	
+	TSharedPtr<FPropertyNode> PinnedNode = PropertyNode.Pin();
+
+	TOptional<int64> Result;
+
+	FComplexPropertyNode* ComplexParentNode = FindComplexParent();
+	for (int32 Index = 0; Index < ComplexParentNode->GetInstancesNum(); ++Index)
+	{
+		uint8* BasePtr = ComplexParentNode->GetMemoryOfInstance(Index);
+		uint8* ValuePtr = Property->ContainerPtrToValuePtr<uint8>(BasePtr);
+
+		int64 Value = NumericProperty->GetSignedIntPropertyValue(ValuePtr);
+		if (!Result.IsSet())
+		{
+			Result = Value;
+		}
+		else if (Result.GetValue() != Value)
+		{
+			// all values aren't the same...
+			return TOptional<FString>();
+		}
+	}
+
+	if (Result.IsSet())
+	{
+		return EnumType->GetNameStringByValue(Result.GetValue());
+	}
+
+	return TOptional<FString>();
+}
+
+TOptional<FString> FEditConditionContext::GetTypeName(const FString& PropertyName) const
+{
+	const UProperty* Property = FindTypedField<UProperty>(PropertyNode, PropertyName);
+	if (Property == nullptr)
+	{
+		return TOptional<FString>();
+	}
+
+	if (const UEnumProperty* EnumProperty = Cast<UEnumProperty>(Property))
+	{
+		return EnumProperty->GetEnum()->GetName();
+	}
+	else if (const UByteProperty* ByteProperty = Cast<UByteProperty>(Property))
+	{
+		return ByteProperty->GetIntPropertyEnum()->GetName();
+	}
+
+	return Property->GetCPPType();
+}
\ No newline at end of file
diff --git a/Engine/Source/Editor/PropertyEditor/Private/EditConditionContext.h b/Engine/Source/Editor/PropertyEditor/Private/EditConditionContext.h
new file mode 100644
--- /dev/null
+++ b/Engine/Source/Editor/PropertyEditor/Private/EditConditionContext.h
@@ -0,0 +1,40 @@
+// Copyright 1998-2019 Epic Games, Inc. All Rights Reserved.
+
+#pragma once
+
+class IEditConditionContext
+{
+public:
+	virtual TOptional<bool> GetBoolValue(const FString& PropertyName) const = 0;
+	virtual TOptional<double> GetNumericValue(const FString& PropertyName) const = 0;
+	virtual TOptional<FString> GetEnumValue(const FString& PropertyName) const = 0;
+	virtual TOptional<FString> GetTypeName(const FString& PropertyName) const = 0;
+};
+
+class UProperty;
+class FPropertyNode;
+class FComplexPropertyNode;
+class FEditConditionExpression;
+
+class FEditConditionContext : public IEditConditionContext
+{
+public:
+	FEditConditionContext(FPropertyNode& InPropertyNode);
+	virtual ~FEditConditionContext() {}
+
+	virtual TOptional<bool> GetBoolValue(const FString& PropertyName) const override;
+	virtual TOptional<double> GetNumericValue(const FString& PropertyName) const override;
+	virtual TOptional<FString> GetEnumValue(const FString& PropertyName) const override;
+	virtual TOptional<FString> GetTypeName(const FString& PropertyName) const override;
+
+	/**
+	 * Fetch the single boolean property referenced.
+	 * Returns nullptr if more than one property is referenced.
+	 */
+	const UBoolProperty* GetSingleBoolProperty(const TSharedPtr<FEditConditionExpression>& Expression) const;
+
+private:
+	TWeakPtr<FPropertyNode> PropertyNode;
+
+	FComplexPropertyNode* FindComplexParent() const;
+};
\ No newline at end of file
diff --git a/Engine/Source/Editor/PropertyEditor/Private/EditConditionParser.cpp b/Engine/Source/Editor/PropertyEditor/Private/EditConditionParser.cpp
--- a/Engine/Source/Editor/PropertyEditor/Private/EditConditionParser.cpp
+++ b/Engine/Source/Editor/PropertyEditor/Private/EditConditionParser.cpp
@@ -1,16 +1,14 @@
 // Copyright 1998-2019 Epic Games, Inc. All Rights Reserved.
 
 #include "EditConditionParser.h"
+#include "EditConditionContext.h"
 
 #include "Math/BasicMathExpressionEvaluator.h"
 #include "Misc/ExpressionParser.h"
 
-#include "PropertyNode.h"
-#include "PropertyEditorHelpers.h"
-
 #define LOCTEXT_NAMESPACE "EditConditionParser"
 
-namespace EditConditionParserNamespace
+namespace EditConditionParserTokens
 {
 	const TCHAR* const FEqual::Moniker = TEXT("==");
 	const TCHAR* const FNotEqual::Moniker = TEXT("!=");
@@ -122,11 +120,11 @@ static TOptional<FExpressionError> ConsumePropertyName(FExpressionTokenConsumer&
 				return FExpressionError(FText::Format(LOCTEXT("PropertyDoubleColonAtEnd", "EditCondition contained double colon at end of property name \"{0}\", expected enum value."), FText::FromString(PropertyName)));
 			}
 
-			Consumer.Add(StringToken.GetValue(), EditConditionParserNamespace::FEnumToken(MoveTemp(EnumType), MoveTemp(EnumValue)));
+			Consumer.Add(StringToken.GetValue(), EditConditionParserTokens::FEnumToken(MoveTemp(EnumType), MoveTemp(EnumValue)));
 		}
 		else
 		{
-			Consumer.Add(StringToken.GetValue(), EditConditionParserNamespace::FPropertyToken(MoveTemp(PropertyName)));
+			Consumer.Add(StringToken.GetValue(), EditConditionParserTokens::FPropertyToken(MoveTemp(PropertyName)));
 		}
 	}
 
@@ -155,7 +153,7 @@ template <typename T>
 struct TOperand
 {
 	TOperand(T InValue) : Value(InValue), Property(nullptr), Context(nullptr) {}
-	TOperand(const EditConditionParserNamespace::FPropertyToken& InProperty, const IEditConditionContext& InContext) : 
+	TOperand(const EditConditionParserTokens::FPropertyToken& InProperty, const IEditConditionContext& InContext) : 
 		Property(&InProperty), Context(&InContext) {}
 
 	bool IsProperty() const { return Property != nullptr; }
@@ -173,7 +171,7 @@ struct TOperand
 
 private:
 	T Value;
-	const EditConditionParserNamespace::FPropertyToken* Property;
+	const EditConditionParserTokens::FPropertyToken* Property;
 	const IEditConditionContext* Context;
 };
 
@@ -205,7 +203,7 @@ FExpressionResult ApplyBinary(TOperand<T> A, TOperand<T> B, Function Apply)
 	return MakeValue(Apply(ValueA.GetValue(), ValueB.GetValue()));
 }
 
-static FExpressionResult ApplyPropertiesEqual(const EditConditionParserNamespace::FPropertyToken& A, const EditConditionParserNamespace::FPropertyToken& B, const IEditConditionContext& Context, bool bNegate)
+static FExpressionResult ApplyPropertiesEqual(const EditConditionParserTokens::FPropertyToken& A, const EditConditionParserTokens::FPropertyToken& B, const IEditConditionContext& Context, bool bNegate)
 {
 	TOptional<FString> TypeNameA = Context.GetTypeName(A.PropertyName);
 	TOptional<FString> TypeNameB = Context.GetTypeName(B.PropertyName);
@@ -257,7 +255,7 @@ static FExpressionResult ApplyPropertiesEqual(const EditConditionParserNamespace
 
 static void CreateBooleanOperators(TOperatorJumpTable<IEditConditionContext>& OperatorJumpTable)
 {
-	using namespace EditConditionParserNamespace;
+	using namespace EditConditionParserTokens;
 
 	OperatorJumpTable.MapPreUnary<FNot>([](bool A) { return !A; });
 	OperatorJumpTable.MapPreUnary<FNot>([](const FPropertyToken& A, const IEditConditionContext* Context)
@@ -337,7 +335,7 @@ static void CreateBooleanOperators(TOperatorJumpTable<IEditConditionContext>& Op
 template <typename T>
 void CreateNumberOperators(TOperatorJumpTable<IEditConditionContext>& OperatorJumpTable)
 {
-	using namespace EditConditionParserNamespace;
+	using namespace EditConditionParserTokens;
 
 	// EQUAL
 	{
@@ -522,16 +520,17 @@ void CreateNumberOperators(TOperatorJumpTable<IEditConditionContext>& OperatorJu
 	}
 }
 
-static FExpressionResult EnumPropertyEquals(const EditConditionParserNamespace::FEnumToken& Enum, const EditConditionParserNamespace::FPropertyToken& Property, const IEditConditionContext& Context, bool bNegate)
+static FExpressionResult EnumPropertyEquals(const EditConditionParserTokens::FEnumToken& Enum, const EditConditionParserTokens::FPropertyToken& Property, const IEditConditionContext& Context, bool bNegate)
 {
 	TOptional<FString> TypeName = Context.GetTypeName(Property.PropertyName);
-	TOptional<FString> ValueProp = Context.GetEnumValue(Property.PropertyName);
 
 	if (TypeName.GetValue() != Enum.Type)
 	{
 		return MakeError(FText::Format(LOCTEXT("OperandTypeMismatch", "EditCondition attempted to compare operands of different types: \"{0}\" and \"{1}\"."), FText::FromString(Property.PropertyName), FText::FromString(Enum.Type + TEXT("::") + Enum.Value)));
 	}
 
+	TOptional<FString> ValueProp = Context.GetEnumValue(Property.PropertyName);
+
 	if (!ValueProp.IsSet())
 	{
 		return MakeError(FText::Format(LOCTEXT("InvalidOperand", "EditCondition attempted to use an invalid operand \"{0}\"."), FText::FromString(Property.PropertyName)));
@@ -543,7 +542,7 @@ static FExpressionResult EnumPropertyEquals(const EditConditionParserNamespace::
 
 static void CreateEnumOperators(TOperatorJumpTable<IEditConditionContext>& OperatorJumpTable)
 {
-	using namespace EditConditionParserNamespace;
+	using namespace EditConditionParserTokens;
 
 	// EQUALS
 	{
@@ -580,7 +579,7 @@ static void CreateEnumOperators(TOperatorJumpTable<IEditConditionContext>& Opera
 
 FEditConditionParser::FEditConditionParser()
 {
-	using namespace EditConditionParserNamespace;
+	using namespace EditConditionParserTokens;
 
 	TokenDefinitions.IgnoreWhitespace();
 	TokenDefinitions.DefineToken(&ExpressionParser::ConsumeSymbol<FEqual>);
@@ -631,7 +630,7 @@ FEditConditionParser::FEditConditionParser()
 
 TOptional<bool> FEditConditionParser::Evaluate(const FEditConditionExpression& Expression, const IEditConditionContext& Context) const
 {
-	using namespace EditConditionParserNamespace;
+	using namespace EditConditionParserTokens;
 	
 	FExpressionResult Result = ExpressionParser::Evaluate(Expression.Tokens, OperatorJumpTable, &Context);
 	if (Result.IsValid())
@@ -673,233 +672,4 @@ TSharedPtr<FEditConditionExpression> FEditConditionParser::Parse(const FString&
 	return TSharedPtr<FEditConditionExpression>();
 }
 
-FEditConditionContext::FEditConditionContext(FPropertyNode& InPropertyNode, TSharedPtr<FEditConditionExpression> Expression)
-{
-	ParentNode = InPropertyNode.GetParentNode();
-
-	const UProperty* Property = InPropertyNode.GetProperty();
-	check(Property);
-
-	if (PropertyEditorHelpers::IsStaticArray(InPropertyNode))
-	{
-		//in the case of conditional static arrays, we have to go up one more level to get the proper parent struct.
-		ParentNode = ParentNode->GetParentNode();
-	}
-	check(ParentNode);
-
-	FComplexPropertyNode* ComplexParentNode = ParentNode->FindComplexParent();
-	check(ComplexParentNode);
-
-	// populate the properties array
-	for (const FCompiledToken& Token : Expression->Tokens)
-	{
-		if (const EditConditionParserNamespace::FPropertyToken* PropertyToken = Token.Node.Cast<EditConditionParserNamespace::FPropertyToken>())
-		{
-			const UProperty* Field = FindField<UProperty>(Property->GetOwnerStruct(), *PropertyToken->PropertyName);
-			ensureMsgf(Field != nullptr, TEXT("EditCondition parsing failed: Field name %s was not found in struct %s."), *PropertyToken->PropertyName, *Property->GetOwnerStruct()->GetName());
-
-			if (Field != nullptr)
-			{
-				FReferencedProperty& RefProperty = Properties[Properties.Emplace()];
-				RefProperty.Property = Field;
-
-				for (int32 Index = 0; Index < ComplexParentNode->GetInstancesNum(); ++Index)
-				{
-					uint8* BaseAddress = ComplexParentNode->GetMemoryOfInstance(Index);
-					check(BaseAddress != nullptr);
-
-					FPropertyAddress& Address = RefProperty.Addresses[RefProperty.Addresses.Emplace()];
-					Address.BaseAddress = BaseAddress;
-					Address.ObjectPtr = ComplexParentNode->AsStructureNode() ?
-						TWeakObjectPtr<UObject>(ComplexParentNode->GetBaseStructure()) :
-						ComplexParentNode->GetInstanceAsUObject(Index);
-
-				}
-			}
-		}
-	}
-}
-
-const UBoolProperty* FEditConditionContext::GetSingleBoolProperty() const
-{
-	if (Properties.Num() == 1)
-	{
-		const UBoolProperty* BoolProperty = Cast<UBoolProperty>(Properties[0].Property);
-		return BoolProperty;
-	}
-	return nullptr;
-}
-
-const FEditConditionContext::FReferencedProperty* FEditConditionContext::FindProperty(const FString& PropertyName) const
-{
-	for (const FReferencedProperty& Property : Properties)
-	{
-		if (Property.Property->GetNameCPP() == PropertyName)
-		{
-			return &Property;
-		}
-	}
-	return nullptr;
-}
-
-TOptional<bool> FEditConditionContext::GetBoolValue(const FString& PropertyName) const
-{
-	const FReferencedProperty* RefProperty = FindProperty(PropertyName);
-	if (RefProperty == nullptr)
-	{
-		return TOptional<bool>();
-	}
-
-	const UBoolProperty* BoolProperty = Cast<UBoolProperty>(RefProperty->Property);
-	if (BoolProperty == nullptr)
-	{
-		return TOptional<bool>();
-	}
-
-	TOptional<bool> Result;
-
-	const FComplexPropertyNode* ComplexParentNode = ParentNode->FindComplexParent();
-	for (const FPropertyAddress& Address : RefProperty->Addresses)
-	{
-		if (!Address.ObjectPtr.IsValid())
-		{
-			return TOptional<bool>();
-		}
-
-		uint8* BaseOffset = ParentNode->GetValueAddress(Address.BaseAddress);
-		uint8* ValueAddr = BoolProperty->ContainerPtrToValuePtr<uint8>(BaseOffset);
-
-		bool bValue = BoolProperty->GetPropertyValue(ValueAddr);
-		if (!Result.IsSet())
-		{
-			Result = bValue;
-		}
-		else if (Result.GetValue() != bValue)
-		{
-			// all values aren't the same...
-			return TOptional<bool>();
-		}
-	}
-
-	return Result;
-}
-
-TOptional<double> FEditConditionContext::GetNumericValue(const FString& PropertyName) const
-{
-	const FReferencedProperty* RefProperty = FindProperty(PropertyName);
-	if (RefProperty == nullptr)
-	{
-		return TOptional<double>();
-	}
-
-	const UNumericProperty* NumericProperty = Cast<UNumericProperty>(RefProperty->Property);
-	if (NumericProperty == nullptr)
-	{
-		return TOptional<double>();
-	}
-
-	TOptional<double> Result;
-
-	for (const FPropertyAddress& Address : RefProperty->Addresses)
-	{
-		if (!Address.ObjectPtr.IsValid())
-		{
-			return TOptional<double>();
-		}
-
-		uint8* BaseOffset = ParentNode->GetValueAddress(Address.BaseAddress);
-		uint8* ValueAddr = NumericProperty->ContainerPtrToValuePtr<uint8>(BaseOffset);
-
-		double Value = 0;
-
-		if (NumericProperty->IsInteger())
-		{
-			Value = (double) NumericProperty->GetSignedIntPropertyValue(ValueAddr);
-		}
-		else if (NumericProperty->IsFloatingPoint())
-		{
-			Value = NumericProperty->GetFloatingPointPropertyValue(ValueAddr);
-		}
-
-		if (!Result.IsSet())
-		{
-			Result = Value;
-		}
-		else if (!FMath::IsNearlyEqual(Result.GetValue(), Value))
-		{
-			// all values aren't the same...
-			return TOptional<double>();
-		}
-	}
-
-	return Result;
-}
-
-TOptional<FString> FEditConditionContext::GetEnumValue(const FString& PropertyName) const
-{
-	const FReferencedProperty* RefProperty = FindProperty(PropertyName);
-	if (RefProperty == nullptr)
-	{
-		return TOptional<FString>();
-	}
-
-	const UEnumProperty* EnumProperty = Cast<UEnumProperty>(RefProperty->Property);
-	if (EnumProperty == nullptr)
-	{
-		return TOptional<FString>();
-	}
-
-	const UNumericProperty* NumProperty = EnumProperty->GetUnderlyingProperty();
-	if (NumProperty == nullptr)
-	{
-		return TOptional<FString>();
-	}
-
-	TOptional<int64> Result;
-
-	const FComplexPropertyNode* ComplexParentNode = ParentNode->FindComplexParent();
-	for (const FPropertyAddress& Address : RefProperty->Addresses)
-	{
-		if (!Address.ObjectPtr.IsValid())
-		{
-			return TOptional<FString>();
-		}
-
-		uint8* BaseOffset = ParentNode->GetValueAddress(Address.BaseAddress);
-		uint8* ValueAddr = NumProperty->ContainerPtrToValuePtr<uint8>(BaseOffset);
-
-		int64 Value = NumProperty->GetSignedIntPropertyValue(ValueAddr);
-
-		if (!Result.IsSet())
-		{
-			Result = Value;
-		}
-		else if (Result.GetValue() != Value)
-		{
-			// all values aren't the same...
-			return TOptional<FString>();
-		}
-	}
-
-	if (Result.IsSet())
-	{
-		return EnumProperty->GetEnum()->GetNameByValue(Result.GetValue()).ToString();
-	}
-
-	return TOptional<FString>();
-}
-
-TOptional<FString> FEditConditionContext::GetTypeName(const FString& PropertyName) const
-{
-	const FReferencedProperty* RefProperty = FindProperty(PropertyName);
-	if (RefProperty != nullptr)
-	{
-		const UProperty* Property = RefProperty->Property;
-		const UClass* Class = Property->GetClass();
-		return Class->GetName();
-	}
-
-	return TOptional<FString>();
-}
-
 #undef LOCTEXT_NAMESPACE
\ No newline at end of file
diff --git a/Engine/Source/Editor/PropertyEditor/Private/EditConditionParser.h b/Engine/Source/Editor/PropertyEditor/Private/EditConditionParser.h
--- a/Engine/Source/Editor/PropertyEditor/Private/EditConditionParser.h
+++ b/Engine/Source/Editor/PropertyEditor/Private/EditConditionParser.h
@@ -7,7 +7,7 @@
 #include "UObject/WeakObjectPtr.h"
 #include "UObject/WeakObjectPtrTemplates.h"
 
-namespace EditConditionParserNamespace
+namespace EditConditionParserTokens
 {
 	struct FPropertyToken 
 	{
@@ -47,8 +47,8 @@ namespace EditConditionParserNamespace
 }
 
 #define DEFINE_EDIT_CONDITION_NODE(TYPE, ...) \
-	namespace EditConditionParserNamespace { struct TYPE { static const TCHAR* const Moniker; }; } \
-	DEFINE_EXPRESSION_NODE_TYPE(EditConditionParserNamespace::TYPE, __VA_ARGS__)
+	namespace EditConditionParserTokens { struct TYPE { static const TCHAR* const Moniker; }; } \
+	DEFINE_EXPRESSION_NODE_TYPE(EditConditionParserTokens::TYPE, __VA_ARGS__)
 
 	DEFINE_EDIT_CONDITION_NODE(FEqual, 0x3AF0EE1B, 0xC3F847C7, 0xA2B95102, 0x4EFC202D)
 	DEFINE_EDIT_CONDITION_NODE(FNotEqual, 0x5CDF3FA4, 0x35614D88, 0x8F94017C, 0xF9229625)
@@ -67,19 +67,12 @@ namespace EditConditionParserNamespace
 	DEFINE_EDIT_CONDITION_NODE(FSubExpressionEnd, 0x5CB25466, 0x780CAE4D, 0x8A88E000, 0xD3695885)
 
 	DEFINE_EXPRESSION_NODE_TYPE(bool, 0xCACBC715, 0x505A6B4A, 0x8808809F, 0x897AA5F6)
-	DEFINE_EXPRESSION_NODE_TYPE(EditConditionParserNamespace::FPropertyToken, 0x9A3FAF6F, 0xB2E45E4D, 0xA80A70C6, 0x47A89BD7)
-	DEFINE_EXPRESSION_NODE_TYPE(EditConditionParserNamespace::FEnumToken, 0xC9A35C24, 0x21FC904B, 0x9F1B2B6A, 0xDF6F4BC4)
+	DEFINE_EXPRESSION_NODE_TYPE(EditConditionParserTokens::FPropertyToken, 0x9A3FAF6F, 0xB2E45E4D, 0xA80A70C6, 0x47A89BD7)
+	DEFINE_EXPRESSION_NODE_TYPE(EditConditionParserTokens::FEnumToken, 0xC9A35C24, 0x21FC904B, 0x9F1B2B6A, 0xDF6F4BC4)
 
 #undef DEFINE_EDIT_CONDITION_NODE
 
-class IEditConditionContext
-{
-public:
-	virtual TOptional<bool> GetBoolValue(const FString& PropertyName) const = 0;
-	virtual TOptional<double> GetNumericValue(const FString& PropertyName) const = 0;
-	virtual TOptional<FString> GetEnumValue(const FString& PropertyName) const = 0;
-	virtual TOptional<FString> GetTypeName(const FString& PropertyName) const = 0;
-};
+class IEditConditionContext;
 
 class FEditConditionExpression
 {
@@ -114,44 +107,3 @@ private:
 	FExpressionGrammar ExpressionGrammar;
 	TOperatorJumpTable<IEditConditionContext> OperatorJumpTable;
 };
-
-class UProperty;
-class FPropertyNode;
-class FComplexPropertyNode;
-
-class FEditConditionContext : public IEditConditionContext
-{
-public:
-	FEditConditionContext(FPropertyNode& InPropertyNode, TSharedPtr<FEditConditionExpression> Expression);
-	virtual ~FEditConditionContext() {}
-
-	virtual TOptional<bool> GetBoolValue(const FString& PropertyName) const override;
-	virtual TOptional<double> GetNumericValue(const FString& PropertyName) const override;
-	virtual TOptional<FString> GetEnumValue(const FString& PropertyName) const override;
-	virtual TOptional<FString> GetTypeName(const FString& PropertyName) const override;
-
-	/** 
-	 * Fetch the single boolean property referenced. 
-	 * Returns nullptr if more than one property is referenced. 
-	 */
-	const UBoolProperty* GetSingleBoolProperty() const;
-
-private:
-	struct FPropertyAddress
-	{
-		TWeakObjectPtr<UObject> ObjectPtr;
-		uint8* BaseAddress;
-	};
-
-	struct FReferencedProperty
-	{
-		const UProperty* Property;
-		TArray<FPropertyAddress> Addresses;
-	};
-
-	FPropertyNode* ParentNode;
-	TArray<FReferencedProperty> Properties;
-	TSharedPtr<FEditConditionExpression> Expression;
-
-	const FReferencedProperty* FindProperty(const FString& PropertyName) const;
-};
\ No newline at end of file
diff --git a/Engine/Source/Editor/PropertyEditor/Private/EditConditionParserTests.cpp b/Engine/Source/Editor/PropertyEditor/Private/EditConditionParserTests.cpp
--- a/Engine/Source/Editor/PropertyEditor/Private/EditConditionParserTests.cpp
+++ b/Engine/Source/Editor/PropertyEditor/Private/EditConditionParserTests.cpp
@@ -1,20 +1,26 @@
 // Copyright 1998-2019 Epic Games, Inc. All Rights Reserved.
 
-#include "CoreTypes.h"
+#include "EditConditionParserTests.h"
 #include "EditConditionParser.h"
+#include "EditConditionContext.h"
+#include "ObjectPropertyNode.h"
 #include "Misc/AutomationTest.h"
 
-#if WITH_DEV_AUTOMATION_TESTS
+UEditConditionTestObject::UEditConditionTestObject(const FObjectInitializer& ObjectInitializer)
+	: Super(ObjectInitializer)
+{
+}
 
-struct TestEditConditionContext : IEditConditionContext
+#if WITH_DEV_AUTOMATION_TESTS
+struct FTestEditConditionContext : IEditConditionContext
 {
 	TMap<FString, bool> BoolValues;
 	TMap<FString, double> DoubleValues;
 	TMap<FString, FString> EnumValues;
 	FString EnumTypeName;
 
-	TestEditConditionContext(){}
-	virtual ~TestEditConditionContext() {}
+	FTestEditConditionContext(){}
+	virtual ~FTestEditConditionContext() {}
 
 	virtual TOptional<bool> GetBoolValue(const FString& PropertyName) const override
 	{
@@ -101,7 +107,7 @@ static bool CanParse(const FEditConditionParser& Parser, const FString& Expressi
 
 	 for (const auto& Token : Parsed->Tokens)
 	 {
-		const EditConditionParserNamespace::FPropertyToken* PropertyToken = Token.Node.Cast<EditConditionParserNamespace::FPropertyToken>();
+		const EditConditionParserTokens::FPropertyToken* PropertyToken = Token.Node.Cast<EditConditionParserTokens::FPropertyToken>();
 		if (PropertyToken != nullptr)
 		{
 			++PropertyCount;
@@ -118,10 +124,10 @@ bool FEditConditionParser_Parse::RunTest(const FString& Parameters)
 	FEditConditionParser Parser;
 	bool bResult = true;
 
-	bResult &= CanParse(Parser, TEXT("bProperty"), 1, 1);
-	bResult &= CanParse(Parser, TEXT("!bProperty"), 2, 1);
-	bResult &= CanParse(Parser, TEXT("bProperty == true"), 3, 1);
-	bResult &= CanParse(Parser, TEXT("bProperty == false"), 3, 1);
+	bResult &= CanParse(Parser, TEXT("BoolProperty"), 1, 1);
+	bResult &= CanParse(Parser, TEXT("!BoolProperty"), 2, 1);
+	bResult &= CanParse(Parser, TEXT("BoolProperty == true"), 3, 1);
+	bResult &= CanParse(Parser, TEXT("BoolProperty == false"), 3, 1);
 	bResult &= CanParse(Parser, TEXT("IntProperty == 0"), 3, 1);
 	bResult &= CanParse(Parser, TEXT("IntProperty != 0"), 3, 1);
 	bResult &= CanParse(Parser, TEXT("IntProperty > 0"), 3, 1);
@@ -137,16 +143,16 @@ bool FEditConditionParser_Parse::RunTest(const FString& Parameters)
 	bResult &= CanParse(Parser, TEXT("Foo == Bar / 5"), 5, 2);
 	bResult &= CanParse(Parser, TEXT("Enum == EType::Value"), 3, 1);
 	bResult &= CanParse(Parser, TEXT("Enum != EType::Value"), 3, 1);
-	bResult &= CanParse(Parser, TEXT("Enum != EType::Value && bProperty"), 5, 2);
-	bResult &= CanParse(Parser, TEXT("Enum == EType::Value || bProperty == false"), 7, 2);
-	bResult &= CanParse(Parser, TEXT("Enum != EType::Value || bProperty == bFoo"), 7, 3);
+	bResult &= CanParse(Parser, TEXT("Enum != EType::Value && BoolProperty"), 5, 2);
+	bResult &= CanParse(Parser, TEXT("Enum == EType::Value || BoolProperty == false"), 7, 2);
+	bResult &= CanParse(Parser, TEXT("Enum != EType::Value || BoolProperty == bFoo"), 7, 3);
 	bResult &= CanParse(Parser, TEXT("Enum == EType::Value && Foo != 5"), 7, 2);
 	bResult &= CanParse(Parser, TEXT("Enum != EType::Value && Foo == Bar"), 7, 3);
 
 	return bResult;
 }
 
-bool CanEvaluate(const FEditConditionParser& Parser, const IEditConditionContext& Context, const FString& Expression, bool Expected)
+static bool CanEvaluate(const FEditConditionParser& Parser, const IEditConditionContext& Context, const FString& Expression, bool Expected)
 {
 	TSharedPtr<FEditConditionExpression> Parsed = Parser.Parse(Expression);
 	if (!Parsed.IsValid())
@@ -171,49 +177,99 @@ bool CanEvaluate(const FEditConditionParser& Parser, const IEditConditionContext
 	return true;
 }
 
+static bool RunBoolTests(const IEditConditionContext& Context)
+{
+	FEditConditionParser Parser;
+	bool bResult = true;
+
+	bResult &= CanEvaluate(Parser, Context, TEXT("true"), true);
+	bResult &= CanEvaluate(Parser, Context, TEXT("false"), false);
+	bResult &= CanEvaluate(Parser, Context, TEXT("!true"), false);
+	bResult &= CanEvaluate(Parser, Context, TEXT("!false"), true);
+	bResult &= CanEvaluate(Parser, Context, TEXT("BoolProperty"), true);
+	bResult &= CanEvaluate(Parser, Context, TEXT("!BoolProperty"), false);
+	
+	bResult &= CanEvaluate(Parser, Context, TEXT("BoolProperty == true"), true);
+	bResult &= CanEvaluate(Parser, Context, TEXT("BoolProperty == false"), false);
+
+	bResult &= CanEvaluate(Parser, Context, TEXT("BoolProperty == BoolProperty"), true);
+	bResult &= CanEvaluate(Parser, Context, TEXT("BoolProperty != BoolProperty"), false);
+
+	bResult &= CanEvaluate(Parser, Context, TEXT("BoolProperty != true"), false);
+	bResult &= CanEvaluate(Parser, Context, TEXT("BoolProperty != false"), true);
+	
+	bResult &= CanEvaluate(Parser, Context, TEXT("true && true"), true);
+	bResult &= CanEvaluate(Parser, Context, TEXT("true && false"), false);
+	bResult &= CanEvaluate(Parser, Context, TEXT("false && true"), false);
+	bResult &= CanEvaluate(Parser, Context, TEXT("false && false"), false);
+	bResult &= CanEvaluate(Parser, Context, TEXT("true && true && true"), true);
+	bResult &= CanEvaluate(Parser, Context, TEXT("true && true && false"), false);
+	bResult &= CanEvaluate(Parser, Context, TEXT("BoolProperty && BoolProperty"), true);
+	bResult &= CanEvaluate(Parser, Context, TEXT("BoolProperty && false"), false);
+	bResult &= CanEvaluate(Parser, Context, TEXT("false && BoolProperty"), false);
+	
+	bResult &= CanEvaluate(Parser, Context, TEXT("true || true"), true);
+	bResult &= CanEvaluate(Parser, Context, TEXT("true || false"), true);
+	bResult &= CanEvaluate(Parser, Context, TEXT("false || true"), true);
+	bResult &= CanEvaluate(Parser, Context, TEXT("false || false"), false);
+	bResult &= CanEvaluate(Parser, Context, TEXT("true || true || true"), true);
+	bResult &= CanEvaluate(Parser, Context, TEXT("true || true || false"), true);
+	bResult &= CanEvaluate(Parser, Context, TEXT("BoolProperty || BoolProperty"), true);
+	bResult &= CanEvaluate(Parser, Context, TEXT("BoolProperty || false"), true);
+	bResult &= CanEvaluate(Parser, Context, TEXT("false || BoolProperty"), true);
+
+	return bResult;
+}
+
 IMPLEMENT_SIMPLE_AUTOMATION_TEST(FEditConditionParser_EvaluateBool, "EditConditionParser.EvaluateBool", EAutomationTestFlags::EditorContext | EAutomationTestFlags::EngineFilter)
 bool FEditConditionParser_EvaluateBool::RunTest(const FString& Parameters)
 {
-	FEditConditionParser Parser;
+	FTestEditConditionContext TestContext;
+	TestContext.SetupBool(TEXT("BoolProperty"), true);
 
-	TestEditConditionContext TestContext;
-	TestContext.SetupBool(TEXT("bProperty"), true);
+	return RunBoolTests(TestContext);
+}
 
+static bool RunNumericTests(const IEditConditionContext& Context)
+{
+	FEditConditionParser Parser;
 	bool bResult = true;
 
-	bResult &= CanEvaluate(Parser, TestContext, TEXT("true"), true);
-	bResult &= CanEvaluate(Parser, TestContext, TEXT("false"), false);
-	bResult &= CanEvaluate(Parser, TestContext, TEXT("!true"), false);
-	bResult &= CanEvaluate(Parser, TestContext, TEXT("!false"), true);
-	bResult &= CanEvaluate(Parser, TestContext, TEXT("bProperty"), true);
-	bResult &= CanEvaluate(Parser, TestContext, TEXT("!bProperty"), false);
-
-	bResult &= CanEvaluate(Parser, TestContext, TEXT("bProperty == true"), true);
-	bResult &= CanEvaluate(Parser, TestContext, TEXT("bProperty == false"), false);
-	bResult &= CanEvaluate(Parser, TestContext, TEXT("bProperty == bProperty"), true);
-
-	bResult &= CanEvaluate(Parser, TestContext, TEXT("bProperty != true"), false);
-	bResult &= CanEvaluate(Parser, TestContext, TEXT("bProperty != false"), true);
-
-	bResult &= CanEvaluate(Parser, TestContext, TEXT("true && true"), true);
-	bResult &= CanEvaluate(Parser, TestContext, TEXT("true && false"), false);
-	bResult &= CanEvaluate(Parser, TestContext, TEXT("false && true"), false);
-	bResult &= CanEvaluate(Parser, TestContext, TEXT("false && false"), false);
-	bResult &= CanEvaluate(Parser, TestContext, TEXT("true && true && true"), true);
-	bResult &= CanEvaluate(Parser, TestContext, TEXT("true && true && false"), false);
-	bResult &= CanEvaluate(Parser, TestContext, TEXT("bProperty && bProperty"), true);
-	bResult &= CanEvaluate(Parser, TestContext, TEXT("bProperty && false"), false);
-	bResult &= CanEvaluate(Parser, TestContext, TEXT("false && bProperty"), false);
-
-	bResult &= CanEvaluate(Parser, TestContext, TEXT("true || true"), true);
-	bResult &= CanEvaluate(Parser, TestContext, TEXT("true || false"), true);
-	bResult &= CanEvaluate(Parser, TestContext, TEXT("false || true"), true);
-	bResult &= CanEvaluate(Parser, TestContext, TEXT("false || false"), false);
-	bResult &= CanEvaluate(Parser, TestContext, TEXT("true || true || true"), true);
-	bResult &= CanEvaluate(Parser, TestContext, TEXT("true || true || false"), true);
-	bResult &= CanEvaluate(Parser, TestContext, TEXT("bProperty || bProperty"), true);
-	bResult &= CanEvaluate(Parser, TestContext, TEXT("bProperty || false"), true);
-	bResult &= CanEvaluate(Parser, TestContext, TEXT("false || bProperty"), true);
+	bResult &= CanEvaluate(Parser, Context, TEXT("5 == 5"), true);
+	bResult &= CanEvaluate(Parser, Context, TEXT("5.0 == 5.0"), true);
+	bResult &= CanEvaluate(Parser, Context, TEXT("DoubleProperty == 5.0"), true);
+	bResult &= CanEvaluate(Parser, Context, TEXT("DoubleProperty == 5"), true);
+	bResult &= CanEvaluate(Parser, Context, TEXT("DoubleProperty == DoubleProperty"), true);
+	
+	bResult &= CanEvaluate(Parser, Context, TEXT("DoubleProperty != 5.0"), false);
+	bResult &= CanEvaluate(Parser, Context, TEXT("DoubleProperty != 6.0"), true);
+	bResult &= CanEvaluate(Parser, Context, TEXT("DoubleProperty != 6"), true);
+	bResult &= CanEvaluate(Parser, Context, TEXT("DoubleProperty != DoubleProperty"), false);
+
+	bResult &= CanEvaluate(Parser, Context, TEXT("DoubleProperty > 4.5"), true);
+	bResult &= CanEvaluate(Parser, Context, TEXT("DoubleProperty > 5"), false);
+	bResult &= CanEvaluate(Parser, Context, TEXT("DoubleProperty > 6"), false);
+	bResult &= CanEvaluate(Parser, Context, TEXT("DoubleProperty > DoubleProperty"), false);
+	
+	bResult &= CanEvaluate(Parser, Context, TEXT("DoubleProperty < 4.5"), false);
+	bResult &= CanEvaluate(Parser, Context, TEXT("DoubleProperty < 5"), false);
+	bResult &= CanEvaluate(Parser, Context, TEXT("DoubleProperty < 6"), true);
+	bResult &= CanEvaluate(Parser, Context, TEXT("DoubleProperty < DoubleProperty"), false);
+	
+	bResult &= CanEvaluate(Parser, Context, TEXT("DoubleProperty >= 4.5"), true);
+	bResult &= CanEvaluate(Parser, Context, TEXT("DoubleProperty >= 5"), true);
+	bResult &= CanEvaluate(Parser, Context, TEXT("DoubleProperty >= 6"), false);
+	bResult &= CanEvaluate(Parser, Context, TEXT("DoubleProperty >= DoubleProperty"), true);
+	
+	bResult &= CanEvaluate(Parser, Context, TEXT("DoubleProperty <= 4.5"), false);
+	bResult &= CanEvaluate(Parser, Context, TEXT("DoubleProperty <= 5"), true);
+	bResult &= CanEvaluate(Parser, Context, TEXT("DoubleProperty <= 6"), true);
+	bResult &= CanEvaluate(Parser, Context, TEXT("DoubleProperty <= DoubleProperty"), true);
+	
+	bResult &= CanEvaluate(Parser, Context, TEXT("DoubleProperty == 2 + 3"), true);
+	bResult &= CanEvaluate(Parser, Context, TEXT("DoubleProperty == 6 - 1"), true);
+	bResult &= CanEvaluate(Parser, Context, TEXT("DoubleProperty == 2.5 * 2"), true);
+	bResult &= CanEvaluate(Parser, Context, TEXT("DoubleProperty == 10 / 2"), true);
 
 	return bResult;
 }
@@ -222,47 +278,33 @@ bool FEditConditionParser_EvaluateBool::RunTest(const FString& Parameters)
 IMPLEMENT_SIMPLE_AUTOMATION_TEST(FEditConditionParser_EvaluateDouble, "EditConditionParser.EvaluateDouble", EAutomationTestFlags::EditorContext | EAutomationTestFlags::EngineFilter)
 bool FEditConditionParser_EvaluateDouble::RunTest(const FString& Parameters)
 {
-	FEditConditionParser Parser;
+	FTestEditConditionContext TestContext;
+	TestContext.SetupDouble(TEXT("DoubleProperty"), 5.0);
 
-	TestEditConditionContext TestContext;
-	TestContext.SetupDouble(TEXT("Property"), 5.0);
+	return RunNumericTests(TestContext);
+}
 
+static bool RunEnumTests(const IEditConditionContext& Context, const FString& EnumName, const FString& PropertyName)
+{
+	FEditConditionParser Parser;
 	bool bResult = true;
 
-	bResult &= CanEvaluate(Parser, TestContext, TEXT("5 == 5"), true);
-	bResult &= CanEvaluate(Parser, TestContext, TEXT("5.0 == 5.0"), true);
-	bResult &= CanEvaluate(Parser, TestContext, TEXT("Property == 5.0"), true);
-	bResult &= CanEvaluate(Parser, TestContext, TEXT("Property == 5"), true);
-	bResult &= CanEvaluate(Parser, TestContext, TEXT("Property == Property"), true);
-
-	bResult &= CanEvaluate(Parser, TestContext, TEXT("Property != 5.0"), false);
-	bResult &= CanEvaluate(Parser, TestContext, TEXT("Property != 6.0"), true);
-	bResult &= CanEvaluate(Parser, TestContext, TEXT("Property != 6"), true);
-
-	bResult &= CanEvaluate(Parser, TestContext, TEXT("Property > 4.5"), true);
-	bResult &= CanEvaluate(Parser, TestContext, TEXT("Property > 5"), false);
-	bResult &= CanEvaluate(Parser, TestContext, TEXT("Property > 6"), false);
-	bResult &= CanEvaluate(Parser, TestContext, TEXT("Property > Property"), false);
-
-	bResult &= CanEvaluate(Parser, TestContext, TEXT("Property < 4.5"), false);
-	bResult &= CanEvaluate(Parser, TestContext, TEXT("Property < 5"), false);
-	bResult &= CanEvaluate(Parser, TestContext, TEXT("Property < 6"), true);
-	bResult &= CanEvaluate(Parser, TestContext, TEXT("Property < Property"), false);
-
-	bResult &= CanEvaluate(Parser, TestContext, TEXT("Property >= 4.5"), true);
-	bResult &= CanEvaluate(Parser, TestContext, TEXT("Property >= 5"), true);
-	bResult &= CanEvaluate(Parser, TestContext, TEXT("Property >= 6"), false);
-	bResult &= CanEvaluate(Parser, TestContext, TEXT("Property >= Property"), true);
-
-	bResult &= CanEvaluate(Parser, TestContext, TEXT("Property <= 4.5"), false);
-	bResult &= CanEvaluate(Parser, TestContext, TEXT("Property <= 5"), true);
-	bResult &= CanEvaluate(Parser, TestContext, TEXT("Property <= 6"), true);
-	bResult &= CanEvaluate(Parser, TestContext, TEXT("Property <= Property"), true);
-
-	bResult &= CanEvaluate(Parser, TestContext, TEXT("Property == 2 + 3"), true);
-	bResult &= CanEvaluate(Parser, TestContext, TEXT("Property == 6 - 1"), true);
-	bResult &= CanEvaluate(Parser, TestContext, TEXT("Property == 2.5 * 2"), true);
-	bResult &= CanEvaluate(Parser, TestContext, TEXT("Property == 10 / 2"), true);
+	bResult &= CanEvaluate(Parser, Context, EnumName + TEXT("::First == ") + EnumName + TEXT("::First"), true);
+	bResult &= CanEvaluate(Parser, Context, EnumName + TEXT("::First == ") + EnumName + TEXT("::Second"), false);
+
+	bResult &= CanEvaluate(Parser, Context, EnumName + TEXT("::First != ") + EnumName + TEXT("::First"), false);
+	bResult &= CanEvaluate(Parser, Context, EnumName + TEXT("::First != ") + EnumName + TEXT("::Second"), true);
+
+	bResult &= CanEvaluate(Parser, Context, PropertyName + TEXT(" == ") + PropertyName, true);
+	bResult &= CanEvaluate(Parser, Context, PropertyName + TEXT(" != ") + PropertyName, false);
+
+	bResult &= CanEvaluate(Parser, Context, PropertyName + TEXT(" == ") + EnumName + TEXT("::First"), true);
+	bResult &= CanEvaluate(Parser, Context, EnumName + TEXT("::First == ") + PropertyName, true);
+	bResult &= CanEvaluate(Parser, Context, PropertyName + TEXT(" == ") + EnumName + TEXT("::Second"), false);
+
+	bResult &= CanEvaluate(Parser, Context, PropertyName + TEXT(" != ") + EnumName + TEXT("::Second"), true);
+	bResult &= CanEvaluate(Parser, Context, PropertyName + TEXT(" != ") + EnumName + TEXT("::First"), false);
+	bResult &= CanEvaluate(Parser, Context, EnumName + TEXT("::Second != ") + PropertyName, true);
 
 	return bResult;
 }
@@ -270,29 +312,96 @@ bool FEditConditionParser_EvaluateDouble::RunTest(const FString& Parameters)
 IMPLEMENT_SIMPLE_AUTOMATION_TEST(FEditConditionParser_EvaluateEnum, "EditConditionParser.EvaluateEnum", EAutomationTestFlags::EditorContext | EAutomationTestFlags::EngineFilter)
 bool FEditConditionParser_EvaluateEnum::RunTest(const FString& Parameters)
 {
-	FEditConditionParser Parser;
+	FTestEditConditionContext TestContext;
 
-	TestEditConditionContext TestContext;
-	TestContext.SetupEnumType(TEXT("MyEnum"));
-	TestContext.SetupEnum(TEXT("Property"), TEXT("First"));
+	const FString EnumType = TEXT("EditConditionTestEnum");
+	TestContext.SetupEnumType(EnumType);
 
-	bool bResult = true;
-	
-	bResult &= CanEvaluate(Parser, TestContext, TEXT("MyEnum::First == MyEnum::First"), true);
-	bResult &= CanEvaluate(Parser, TestContext, TEXT("MyEnum::First == MyEnum::Second"), false);
+	const FString PropertyName = TEXT("EnumProperty");
+	TestContext.SetupEnum(PropertyName, TEXT("First"));
 
-	bResult &= CanEvaluate(Parser, TestContext, TEXT("MyEnum::First != MyEnum::First"), false);
-	bResult &= CanEvaluate(Parser, TestContext, TEXT("MyEnum::First != MyEnum::Second"), true);
+	return RunEnumTests(TestContext, EnumType, PropertyName);
+}
 
-	bResult &= CanEvaluate(Parser, TestContext, TEXT("Property == MyEnum::First"), true);
-	bResult &= CanEvaluate(Parser, TestContext, TEXT("MyEnum::First == Property"), true);
-	bResult &= CanEvaluate(Parser, TestContext, TEXT("Property == MyEnum::Second"), false);
+IMPLEMENT_SIMPLE_AUTOMATION_TEST(FEditConditionParser_EvaluateUObject, "EditConditionParser.EvaluateUObject", EAutomationTestFlags::EditorContext | EAutomationTestFlags::EngineFilter)
+bool FEditConditionParser_EvaluateUObject::RunTest(const FString& Parameters)
+{
+	UEditConditionTestObject* TestObject = NewObject<UEditConditionTestObject>();
+	TestObject->AddToRoot();
 
-	bResult &= CanEvaluate(Parser, TestContext, TEXT("Property != MyEnum::Second"), true);
-	bResult &= CanEvaluate(Parser, TestContext, TEXT("Property != MyEnum::First"), false);
-	bResult &= CanEvaluate(Parser, TestContext, TEXT("MyEnum::Second != Property"), true);
+	TSharedPtr<FObjectPropertyNode> ObjectNode(new FObjectPropertyNode);
+	ObjectNode->AddObject(TestObject);
 
-	return bResult;
+	FPropertyNodeInitParams InitParams;
+	ObjectNode->InitNode(InitParams);
+
+	TOptional<bool> bResult;
+	bool bAllResults = true;
+
+	// enum comparison
+	{
+		static const FName EnumPropertyName = TEXT("EnumProperty");
+		TSharedPtr<FPropertyNode> PropertyNode = ObjectNode->FindChildPropertyNode(EnumPropertyName, true);
+		FEditConditionContext Context(*PropertyNode.Get());
+
+		TestObject->EnumProperty = EditConditionTestEnum::First;
+		TestObject->ByteEnumProperty = EditConditionByteEnum::First;
+
+		static const FString EnumType = TEXT("EditConditionTestEnum");
+		bAllResults &= RunEnumTests(Context, EnumType, EnumPropertyName.ToString());
+
+		static const FString ByteEnumType = TEXT("EditConditionByteEnum");
+		static const FString ByteEnumPropertyName = TEXT("ByteEnumProperty");
+		bAllResults &= RunEnumTests(Context, ByteEnumType, ByteEnumPropertyName);
+	}
+
+	// bool comparison
+	{
+		static const FName BoolPropertyName(TEXT("BoolProperty"));
+		TSharedPtr<FPropertyNode> PropertyNode = ObjectNode->FindChildPropertyNode(BoolPropertyName, true);
+		FEditConditionContext Context(*PropertyNode.Get());
+
+		TestObject->BoolProperty = true;
+
+		bAllResults &= RunBoolTests(Context);
+	}
+
+	// double comparison
+	{
+		static const FName DoublePropertyName(TEXT("DoubleProperty"));
+		TSharedPtr<FPropertyNode> PropertyNode = ObjectNode->FindChildPropertyNode(DoublePropertyName, true);
+		FEditConditionContext Context(*PropertyNode.Get());
+
+		TestObject->DoubleProperty = 5.0;
+
+		bAllResults &= RunNumericTests(Context);
+	}
+
+	// integer comparison
+	{
+		static const FName IntegerPropertyName(TEXT("IntegerProperty"));
+		TSharedPtr<FPropertyNode> PropertyNode = ObjectNode->FindChildPropertyNode(IntegerPropertyName, true);
+		FEditConditionContext Context(*PropertyNode.Get());
+
+		TestObject->IntegerProperty = 5;
+
+		bAllResults &= RunNumericTests(Context);
+	}
+
+	{
+		static const FName DoublePropertyName(TEXT("DoubleProperty"));
+		TSharedPtr<FPropertyNode> PropertyNode = ObjectNode->FindChildPropertyNode(DoublePropertyName, true);
+		FEditConditionContext Context(*PropertyNode.Get());
+
+		TestEqual(TEXT("Boolean Type Name"), Context.GetTypeName(TEXT("BoolProperty")).GetValue(), TEXT("bool"));
+		TestEqual(TEXT("Enum Type Name"), Context.GetTypeName(TEXT("EnumProperty")).GetValue(), TEXT("EditConditionTestEnum"));
+		TestEqual(TEXT("Byte Enum Type Name"), Context.GetTypeName(TEXT("ByteEnumProperty")).GetValue(), TEXT("EditConditionByteEnum"));
+		TestEqual(TEXT("Double Type Name"), Context.GetTypeName(TEXT("DoubleProperty")).GetValue(), TEXT("double"));
+	}
+
+	TestObject->RemoveFromRoot();
+
+	return bAllResults;
 }
 
 #endif // WITH_DEV_AUTOMATION_TESTS
\ No newline at end of file
diff --git a/Engine/Source/Editor/PropertyEditor/Private/EditConditionParserTests.h b/Engine/Source/Editor/PropertyEditor/Private/EditConditionParserTests.h
new file mode 100644
--- /dev/null
+++ b/Engine/Source/Editor/PropertyEditor/Private/EditConditionParserTests.h
@@ -0,0 +1,41 @@
+// Copyright 1998-2019 Epic Games, Inc. All Rights Reserved.
+
+#pragma once
+
+#include "UObject/Object.h"
+#include "EditConditionParserTests.generated.h"
+
+UENUM()
+enum class EditConditionTestEnum
+{
+	First = 15,
+	Second = 31
+};
+
+UENUM()
+enum EditConditionByteEnum
+{
+	First = 15,
+	Second = 31
+};
+
+UCLASS(transient)
+class UEditConditionTestObject : public UObject
+{
+	GENERATED_UCLASS_BODY()
+
+	UPROPERTY(EditAnywhere, Category=Test)
+	bool BoolProperty;
+
+	UPROPERTY(EditAnywhere, Category=Test)
+	EditConditionTestEnum EnumProperty;
+
+	UPROPERTY(EditAnywhere, Category=Test)
+	TEnumAsByte<EditConditionByteEnum> ByteEnumProperty;
+
+	UPROPERTY(EditAnywhere, Category=Test)
+	double DoubleProperty;
+
+	UPROPERTY(EditAnywhere, Category=Test)
+	int32 IntegerProperty;
+};
diff --git a/Engine/Source/Editor/PropertyEditor/Private/Presentation/PropertyEditor/PropertyEditor.cpp b/Engine/Source/Editor/PropertyEditor/Private/Presentation/PropertyEditor/PropertyEditor.cpp
--- a/Engine/Source/Editor/PropertyEditor/Private/Presentation/PropertyEditor/PropertyEditor.cpp
+++ b/Engine/Source/Editor/PropertyEditor/Private/Presentation/PropertyEditor/PropertyEditor.cpp
@@ -17,6 +17,8 @@
 #include "IConfigEditorModule.h"
 #include "PropertyNode.h"
 #include "Kismet2/BlueprintEditorUtils.h"
+#include "EditConditionParser.h"
+#include "EditConditionContext.h"
 
 #define LOCTEXT_NAMESPACE "PropertyEditor"
 
@@ -50,7 +52,7 @@ FPropertyEditor::FPropertyEditor( const TSharedRef<FPropertyNode>& InPropertyNod
 				EditConditionExpression = Parser->Parse(Property->GetMetaData(EditConditionName));
 				if (EditConditionExpression.IsValid())
 				{
-					EditConditionContext = MakeShareable(new FEditConditionContext(PropertyNode.Get(), EditConditionExpression));
+					EditConditionContext = MakeShareable(new FEditConditionContext(PropertyNode.Get()));
 				}
 			}
 		}
@@ -455,7 +457,7 @@ void FPropertyEditor::ToggleEditConditionState()
 
 	PropertyNode->NotifyPreChange( PropertyNode->GetProperty(), PropertyUtilities->GetNotifyHook() );
 
-	const UBoolProperty* EditConditionProperty = EditConditionContext->GetSingleBoolProperty();
+	const UBoolProperty* EditConditionProperty = EditConditionContext->GetSingleBoolProperty(EditConditionExpression);
 	check(EditConditionProperty != nullptr);
 
 	FComplexPropertyNode* ComplexParentNode = ParentNode->FindComplexParent();
@@ -597,7 +599,7 @@ bool FPropertyEditor::SupportsEditConditionToggle() const
 	static const FName Name_HideEditConditionToggle("HideEditConditionToggle");
 	if (!Property->HasMetaData(Name_HideEditConditionToggle) && EditConditionExpression.IsValid())
 	{
-		const UBoolProperty* ConditionalProperty = EditConditionContext->GetSingleBoolProperty();
+		const UBoolProperty* ConditionalProperty = EditConditionContext->GetSingleBoolProperty(EditConditionExpression);
 		if (ConditionalProperty != nullptr)
 		{
 			// If it's editable, then only put an inline toggle box if the metadata specifies it
diff --git a/Engine/Source/Editor/PropertyEditor/Private/Presentation/PropertyEditor/PropertyEditor.h b/Engine/Source/Editor/PropertyEditor/Private/Presentation/PropertyEditor/PropertyEditor.h
--- a/Engine/Source/Editor/PropertyEditor/Private/Presentation/PropertyEditor/PropertyEditor.h
+++ b/Engine/Source/Editor/PropertyEditor/Private/Presentation/PropertyEditor/PropertyEditor.h
@@ -5,7 +5,9 @@
 #include "PropertyNode.h"
 #include "PropertyHandle.h"
 #include "Editor/SceneOutliner/Public/SceneOutlinerFwd.h"
-#include "EditConditionParser.h"
+
+class FEditConditionExpression;
+class FEditConditionContext;
 
 class FPropertyEditor : public TSharedFromThis< FPropertyEditor >	
 {
diff --git a/Engine/Source/Editor/UnrealEd/Classes/Editor/PropertyEditorTestObject.h b/Engine/Source/Editor/UnrealEd/Classes/Editor/PropertyEditorTestObject.h
--- a/Engine/Source/Editor/UnrealEd/Classes/Editor/PropertyEditorTestObject.h
+++ b/Engine/Source/Editor/UnrealEd/Classes/Editor/PropertyEditorTestObject.h
@@ -19,7 +19,7 @@ class UStaticMeshComponent;
 class UTexture;
 
 UENUM()
-enum PropertEditorTestEnum
+enum PropertyEditorTestEnum
 {	
 	/** This comment should appear above enum 1 */
 	PropertyEditorTest_Enum1 UMETA(Hidden),
@@ -124,13 +124,11 @@ class UPropertyEditorTestObject : public UObject
 {
     GENERATED_UCLASS_BODY()
 
-
-
 	UPROPERTY(EditAnywhere, Category = BasicProperties)
 	int8 Int8Property;
 
 	UPROPERTY(EditAnywhere, Category = BasicProperties)
-	int16 Int16roperty;
+	int16 Int16Property;
 
 	UPROPERTY(EditAnywhere, Category = BasicProperties)
 	int32 Int32Property;
@@ -196,7 +194,10 @@ class UPropertyEditorTestObject : public UObject
 	FColor ColorProperty;
 
 	UPROPERTY(EditAnywhere, Category=BasicProperties)
-	TEnumAsByte<enum PropertEditorTestEnum> EnumProperty;
+	TEnumAsByte<enum PropertyEditorTestEnum> EnumByteProperty;
+
+	UPROPERTY(EditAnywhere, Category=BasicProperties)
+	EditColor EnumProperty;
 
 	UPROPERTY(EditAnywhere, Category = BasicProperties)
 	FMatrix MatrixProperty;
@@ -252,7 +253,7 @@ class UPropertyEditorTestObject : public UObject
 	TArray<FColor> ColorPropertyArray;
 
 	UPROPERTY(EditAnywhere, Category=ArraysOfProperties)
-	TArray<TEnumAsByte<enum PropertEditorTestEnum> > EnumPropertyArray;
+	TArray<TEnumAsByte<enum PropertyEditorTestEnum> > EnumPropertyArray;
 
 	UPROPERTY(EditAnywhere, Category=ArraysOfProperties)
 	TArray<FPropertyEditorTestBasicStruct> StructPropertyArray;
@@ -266,7 +267,7 @@ class UPropertyEditorTestObject : public UObject
 	UPROPERTY(EditAnywhere, Category=ArraysOfProperties)
 	int32 StaticArrayOfIntsWithEnumLabels[ArrayIndex_MAX];
 
-	/** This is a float property tooltip that is overridden */
+	// This is a float property tooltip that is overridden
 	UPROPERTY(EditAnywhere, Category=AdvancedProperties, meta=(ClampMin = "0.0", ClampMax = "100.0", UIMin = "0.0", UIMax = "50.0", ToolTip = "This is a custom tooltip that should be shown"))
 	float FloatPropertyWithClampedRange;
 
@@ -385,7 +386,7 @@ class UPropertyEditorTestObject : public UObject
 	TMap<UObject*, FLinearColor> ObjectToColorMap;
 
 	UPROPERTY(EditAnywhere, Category=TMapTests)
-	TMap<int32, TEnumAsByte<PropertEditorTestEnum> > IntToEnumMap;
+	TMap<int32, TEnumAsByte<PropertyEditorTestEnum> > IntToEnumMap;
 
 	UPROPERTY(EditAnywhere, Category=TMapTests)
 	TMap<FName, FName> NameToNameMap;
@@ -419,4 +420,67 @@ class UPropertyEditorTestObject : public UObject
 	// filter for AllowedClasses correctly.
 	UPROPERTY(EditAnywhere, Category=ObjectPropertyAllowedClasses, meta=(AllowedClasses="Texture,BlendableInterface"))
 	UObject* TextureOrBlendableInterface;
+
+	UPROPERTY(EditAnywhere, Category = "Subcategory")
+	bool bSubcategory;
+
+	UPROPERTY(EditAnywhere, AdvancedDisplay, Category = "Subcategory")
+	bool bSubcategoryAdvanced;
+
+	UPROPERTY(EditAnywhere, Category = "Subcategory|Foo")
+	bool bSubcategoryFooSimple;
+
+	UPROPERTY(EditAnywhere, AdvancedDisplay, Category = "Subcategory|Foo")
+	bool bSubcategoryFooAdvanced;
+
+	UPROPERTY(EditAnywhere, Category = "Subcategory|Bar")
+	bool bSubcategoryBarSimple;
+
+	UPROPERTY(EditAnywhere, AdvancedDisplay, Category = "Subcategory|Bar")
+	bool bSubcategoryBarAdvanced;
+
+	UPROPERTY(EditAnywhere, Category = "Subcategory")
+	bool bSubcategoryLast;
+
+	UPROPERTY(EditAnywhere, Category = EditCondition)
+	bool bEnablesNext;
+
+	UPROPERTY(EditAnywhere, Category = EditCondition, meta = (EditCondition = "bEnablesNext == true"))
+	bool bEnabledByPrevious;
+
+	UPROPERTY(EditAnywhere, Category = EditCondition)
+	EditColor EnumEditCondition;
+
+	UPROPERTY(EditAnywhere, Category = EditCondition, meta = (EditCondition = "EnumEditCondition == EditColor::Blue"))
+	bool bEnabledWhenBlue;
+
+	UPROPERTY(EditAnywhere, Category = EditCondition, meta = (EditCondition = "EnumEditCondition == EditColor::Pink"))
+	bool bEnabledWhenPink;
+
+	UPROPERTY(EditAnywhere, Category = EditCondition)
+	TEnumAsByte<PropertyEditorTestEnum> EnumAsByteEditCondition;
+
+	UPROPERTY(EditAnywhere, Category = EditCondition, meta = (EditCondition = "EnumAsByteEditCondition == PropertyEditorTestEnum::PropertyEditorTest_Enum2"))
+	bool bEnabledWhenEnumIs2;
+
+	UPROPERTY(EditAnywhere, Category = EditCondition, meta = (EditCondition = "EnumAsByteEditCondition == PropertyEditorTestEnum::PropertyEditorTest_Enum4"))
+	bool bEnabledWhenEnumIs4;
+
+	UPROPERTY(EditAnywhere, Category = EditCondition)
+	int32 IntegerEditCondition;
+
+	UPROPERTY(EditAnywhere, Category = EditCondition, meta = (EditCondition = "IntegerEditCondition >= 5"))
+	bool bEnabledWhenIntGreaterOrEqual5;
+
+	UPROPERTY(EditAnywhere, Category = EditCondition, meta = (EditCondition = "IntegerEditCondition <= 10"))
+	bool bEnabledWhenIntLessOrEqual10;
+
+	UPROPERTY(EditAnywhere, Category = EditCondition)
+	float FloatEditCondition;
+
+	UPROPERTY(EditAnywhere, Category = EditCondition, meta = (EditCondition = "FloatEditCondition > 5"))
+	bool bEnabledWhenFloatGreaterThan5;
+
+	UPROPERTY(EditAnywhere, Category = EditCondition, meta = (EditCondition = "FloatEditCondition < 10"))
+	bool bEnabledWhenFloatLessThan10;
 };
