Added documentation for FEditConditionParser::Parse().

diff --git a/Engine/Source/Editor/PropertyEditor/Private/EditConditionParser.h b/Engine/Source/Editor/PropertyEditor/Private/EditConditionParser.h
--- a/Engine/Source/Editor/PropertyEditor/Private/EditConditionParser.h
+++ b/Engine/Source/Editor/PropertyEditor/Private/EditConditionParser.h
@@ -94,6 +94,10 @@ public:
 	FEditConditionParser& operator=(const FEditConditionParser&) = delete;
 	FEditConditionParser& operator=(FEditConditionParser&&) = default;
 
+	/**
+	 * Parse the given string into an expression.
+	 * @returns The parsed expression if valid, invalid TSharedPtr if the string is not a valid expression string.
+	 */
 	TSharedPtr<FEditConditionExpression> Parse(const FString& ExpressionString) const;
 
 	/** 
