Added support for InlineCategoryProperty metadata, which makes the property it is applied to inline next to the category. Useful for toggling editability of an entire category of properties, for example.

Currently restricted to bools and enums.

diff --git a/Engine/Source/Editor/PropertyEditor/Private/DetailCategoryBuilderImpl.cpp b/Engine/Source/Editor/PropertyEditor/Private/DetailCategoryBuilderImpl.cpp
--- a/Engine/Source/Editor/PropertyEditor/Private/DetailCategoryBuilderImpl.cpp
+++ b/Engine/Source/Editor/PropertyEditor/Private/DetailCategoryBuilderImpl.cpp
@@ -8,6 +8,7 @@
 #include "DetailItemNode.h"
 #include "DetailAdvancedDropdownNode.h"
 #include "DetailPropertyRow.h"
+#include "SDetailSingleItemRow.h"
 #include "DetailGroup.h"
 #include "StructurePropertyNode.h"
 #include "ItemPropertyNode.h"
@@ -125,7 +126,6 @@ FDetailCategoryImpl::FDetailCategoryImpl(FName InCategoryName, TSharedRef<FDetai
 	GConfig->GetBool(TEXT("DetailCategoriesAdvanced"), *CategoryPathName, bUserShowAdvancedConfigValue, GEditorPerProjectIni);
 
 	bUserShowAdvanced = bUserShowAdvancedConfigValue;
-
 }
 
 FDetailCategoryImpl::~FDetailCategoryImpl()
@@ -244,6 +244,7 @@ IDetailCategoryBuilder& FDetailCategoryImpl::RestoreExpansionState(bool bRestore
 
 IDetailCategoryBuilder& FDetailCategoryImpl::HeaderContent(TSharedRef<SWidget> InHeaderContent)
 {
+	ensureMsgf(!this->HeaderContentWidget.IsValid(), TEXT("Category already has a header content widget defined!"));
 	this->HeaderContentWidget = InHeaderContent;
 	return *this;
 }
@@ -291,7 +292,6 @@ IDetailPropertyRow& FDetailCategoryImpl::AddProperty(TSharedPtr<IPropertyHandle>
 
 	bool bForAdvanced = false;
 
-
 	if (Location == EPropertyLocation::Default)
 	{
 		// Get the default location of this property
@@ -303,7 +303,6 @@ IDetailPropertyRow& FDetailCategoryImpl::AddProperty(TSharedPtr<IPropertyHandle>
 		bForAdvanced = true;
 	}
 
-
 	AddCustomLayout(NewCustomization, bForAdvanced);
 
 	return *NewCustomization.PropertyRow;
@@ -629,16 +628,33 @@ void FDetailCategoryImpl::SetDisplayName(FName InCategoryName, const FText& Loca
 
 TSharedRef<ITableRow> FDetailCategoryImpl::GenerateWidgetForTableView(const TSharedRef<STableViewBase>& OwnerTable, const FDetailColumnSizeData& ColumnSizeData, bool bAllowFavoriteSystem)
 {
+	TSharedPtr<SWidget> HeaderContent = HeaderContentWidget;
+	if (InlinePropertyNode.IsValid())
+	{
+		FDetailWidgetRow Row;
+		InlinePropertyNode->GenerateStandaloneWidget(Row);
+		HeaderContent = Row.ValueWidget.Widget;
+	}
+
 	return
 		SNew(SDetailCategoryTableRow, AsShared(), OwnerTable)
 		.InnerCategory(DetailLayoutBuilder.IsValid() ? DetailLayoutBuilder.Pin()->IsLayoutForExternalRoot() : false)
 		.DisplayName(GetDisplayName())
-		.HeaderContent(HeaderContentWidget);
+		.HeaderContent(HeaderContent)
+		.ColumnSizeData(&ColumnSizeData);
 }
 
 
 bool FDetailCategoryImpl::GenerateStandaloneWidget(FDetailWidgetRow& OutRow) const
 {
+	TSharedPtr<SWidget> HeaderContent = HeaderContentWidget;
+	if (InlinePropertyNode.IsValid())
+	{
+		FDetailWidgetRow Row;
+		InlinePropertyNode->GenerateStandaloneWidget(Row);
+		HeaderContent = Row.ValueWidget.Widget;
+	}
+
 	const bool bIsInnerCategory = DetailLayoutBuilder.Pin()->IsLayoutForExternalRoot();
 
 	OutRow.NameContent()
@@ -655,7 +671,7 @@ bool FDetailCategoryImpl::GenerateStandaloneWidget(FDetailWidgetRow& OutRow) con
 			.ValueContent()
 			.HAlign(HAlign_Fill)
 			[
-				HeaderContentWidget.ToSharedRef()
+				HeaderContent.ToSharedRef()
 			];
 	}
 
@@ -700,10 +716,11 @@ bool FDetailCategoryImpl::ShouldBeExpanded() const
 
 ENodeVisibility FDetailCategoryImpl::GetVisibility() const
 {
-	return bHasVisibleDetails && bIsCategoryVisible ? ENodeVisibility::Visible : ENodeVisibility::ForcedHidden;
+	return bHasVisibleDetails && bIsCategoryVisible ? 
+		ENodeVisibility::Visible : ENodeVisibility::ForcedHidden;
 }
 
-bool IsCustomProperty(const TSharedPtr<FPropertyNode>& PropertyNode)
+static bool IsCustomProperty(const TSharedPtr<FPropertyNode>& PropertyNode)
 {
 	// The property node is custom if it has a custom layout or if its a struct and any of its children have a custom layout
 	bool bIsCustom = !PropertyNode.IsValid() || PropertyNode->HasNodeFlags(EPropertyNodeFlags::IsCustomized) != 0;
@@ -711,6 +728,33 @@ bool IsCustomProperty(const TSharedPtr<FPropertyNode>& PropertyNode)
 	return bIsCustom;
 }
 
+static bool ShouldBeInlineNode(const TSharedRef<FDetailItemNode>& Node)
+{
+	TSharedPtr<FPropertyNode> PropertyNode = Node->GetPropertyNode();
+	if (PropertyNode.IsValid())
+	{
+		const UProperty* Property = PropertyNode->GetProperty();
+		if (Property != nullptr)
+		{
+			const UBoolProperty* BoolProperty = Cast<UBoolProperty>(Property);
+			const UEnumProperty* EnumProperty = Cast<UEnumProperty>(Property);
+			const UByteProperty* ByteProperty = Cast<UByteProperty>(Property);
+
+			// Only allow bools and enums as inline nodes.
+			if (BoolProperty != nullptr || EnumProperty != nullptr ||
+				(ByteProperty != nullptr && ByteProperty->IsEnum()))
+			{
+				static const FName Name_InlineCategoryProperty("InlineCategoryProperty");
+				if (Property->HasMetaData(Name_InlineCategoryProperty))
+				{
+					return true;
+				}
+			}
+		}
+	}
+
+	return false;
+}
 
 void FDetailCategoryImpl::GenerateNodesFromCustomizations(const FCustomizationList& InCustomizationList, bool bDefaultLayouts, FDetailNodeList& OutNodeList, bool &bOutLastItemHasMultipleColumns)
 {
@@ -726,6 +770,13 @@ void FDetailCategoryImpl::GenerateNodesFromCustomizations(const FCustomizationLi
 			TSharedRef<FDetailItemNode> NewNode = MakeShareable(new FDetailItemNode(Customization, AsShared(), IsParentEnabled));
 			NewNode->Initialize();
 
+			if (ShouldBeInlineNode(NewNode))
+			{
+				ensureMsgf(!InlinePropertyNode.IsValid(), TEXT("Multiple properties marked InlineCategoryProperty detected in category %s."), *DisplayName.ToString());
+				InlinePropertyNode = NewNode;
+				continue;
+			}
+
 			// Add the node unless only its children should be visible or it didnt generate any children or if it is a custom builder which can generate children at any point
 			if (!NewNode->ShouldShowOnlyChildren() || NewNode->HasGeneratedChildren() || Customization.HasCustomBuilder())
 			{
@@ -823,8 +874,6 @@ void FDetailCategoryImpl::GenerateChildrenForLayouts()
 			const bool bShouldShowGroup = LayoutMap.ShouldShowGroup(RequiredGroupName);
 			bGeneratedAnyChildren |= GenerateChildrenForSingleLayout(RequiredGroupName, bDefaultLayout, bShouldShowGroup, Layout.GetDefaultAdvancedLayouts(), AdvancedChildNodes, bLastItemHasMultipleColumns);
 		}
-
-
 	}
 
 	// Generate nodes for advanced dropdowns
@@ -849,9 +898,8 @@ void FDetailCategoryImpl::GetChildren(FDetailNodeList& OutChildren)
 
 void FDetailCategoryImpl::GetGeneratedChildren(FDetailNodeList& OutChildren, bool bIgnoreVisibility, bool bIgnoreAdvancedDropdown)
 {
-	for (int32 ChildIndex = 0; ChildIndex < SimpleChildNodes.Num(); ++ChildIndex)
+	for (TSharedRef<FDetailTreeNode>& Child : SimpleChildNodes)
 	{
-		TSharedRef<FDetailTreeNode>& Child = SimpleChildNodes[ChildIndex];
 		if (bIgnoreVisibility || Child->GetVisibility() == ENodeVisibility::Visible)
 		{
 			if (Child->ShouldShowOnlyChildren())
@@ -872,10 +920,8 @@ void FDetailCategoryImpl::GetGeneratedChildren(FDetailNodeList& OutChildren, boo
 			OutChildren.Add(AdvancedDropdownNodeTop.ToSharedRef());
 		}
 
-		for (int32 ChildIndex = 0; ChildIndex < AdvancedChildNodes.Num(); ++ChildIndex)
+		for (TSharedRef<FDetailTreeNode>& Child : AdvancedChildNodes)
 		{
-			TSharedRef<FDetailTreeNode>& Child = AdvancedChildNodes[ChildIndex];
-
 			if (bIgnoreVisibility || Child->GetVisibility() == ENodeVisibility::Visible)
 			{
 				if (Child->ShouldShowOnlyChildren())
@@ -903,9 +949,13 @@ void FDetailCategoryImpl::FilterNode(const FDetailFilter& InFilter)
 
 	bHasVisibleDetails = false;
 
-	for (int32 ChildIndex = 0; ChildIndex < SimpleChildNodes.Num(); ++ChildIndex)
+	if (InlinePropertyNode.IsValid())
+	{
+		bHasVisibleDetails = true;
+	}
+
+	for (TSharedRef<FDetailTreeNode>& Child : SimpleChildNodes)
 	{
-		TSharedRef<FDetailTreeNode>& Child = SimpleChildNodes[ChildIndex];
 		Child->FilterNode(InFilter);
 
 		if (Child->GetVisibility() == ENodeVisibility::Visible)
@@ -915,9 +965,8 @@ void FDetailCategoryImpl::FilterNode(const FDetailFilter& InFilter)
 		}
 	}
 
-	for (int32 ChildIndex = 0; ChildIndex < AdvancedChildNodes.Num(); ++ChildIndex)
+	for (TSharedRef<FDetailTreeNode>& Child : AdvancedChildNodes)
 	{
-		TSharedRef<FDetailTreeNode>& Child = AdvancedChildNodes[ChildIndex];
 		Child->FilterNode(InFilter);
 
 		if (Child->GetVisibility() == ENodeVisibility::Visible)
@@ -945,6 +994,7 @@ void FDetailCategoryImpl::GenerateLayout()
 	AdvancedChildNodes.Empty();
 	AdvancedDropdownNodeTop.Reset();
 	AdvancedDropdownNodeBottom.Reset();
+	InlinePropertyNode.Reset();
 
 	GenerateChildrenForLayouts();
 
diff --git a/Engine/Source/Editor/PropertyEditor/Private/DetailCategoryBuilderImpl.h b/Engine/Source/Editor/PropertyEditor/Private/DetailCategoryBuilderImpl.h
--- a/Engine/Source/Editor/PropertyEditor/Private/DetailCategoryBuilderImpl.h
+++ b/Engine/Source/Editor/PropertyEditor/Private/DetailCategoryBuilderImpl.h
@@ -410,6 +410,8 @@ private:
 	FString CategoryPathName;
 	/** Custom header content displayed to the right of the category name */
 	TSharedPtr<SWidget> HeaderContentWidget;
+
+	TSharedPtr<FDetailTreeNode> InlinePropertyNode;
 	/** The parent detail builder */
 	TWeakPtr<FDetailLayoutBuilderImpl> DetailLayoutBuilder;
 	/** The category identifier */
diff --git a/Engine/Source/Editor/PropertyEditor/Private/DetailCategoryGroupNode.cpp b/Engine/Source/Editor/PropertyEditor/Private/DetailCategoryGroupNode.cpp
--- a/Engine/Source/Editor/PropertyEditor/Private/DetailCategoryGroupNode.cpp
+++ b/Engine/Source/Editor/PropertyEditor/Private/DetailCategoryGroupNode.cpp
@@ -69,10 +69,9 @@ void SDetailCategoryTableRow::Construct( const FArguments& InArgs, TSharedRef<FD
 			.OnSlotResized(OptionalSizeData->OnWidthChanged)
 			[
 				SNew(SHorizontalBox)
-
 				+ SHorizontalBox::Slot()
 				[
-					SNullWidget::NullWidget
+					InArgs._HeaderContent.IsValid() ? InArgs._HeaderContent.ToSharedRef() : SNullWidget::NullWidget
 				]
 			];
 	}
@@ -93,15 +92,6 @@ void SDetailCategoryTableRow::Construct( const FArguments& InArgs, TSharedRef<FD
 		]
 	];
 
-	if( InArgs._HeaderContent.IsValid() )
-	{
-		MyContent->AddSlot()
-		.VAlign(VAlign_Center)
-		[	
-			InArgs._HeaderContent.ToSharedRef()
-		];
-	}
-
 	STableRow< TSharedPtr< FDetailTreeNode > >::ConstructInternal(
 		STableRow::FArguments()
 			.Style(FEditorStyle::Get(), "DetailsView.TreeView.TableRow")
diff --git a/Engine/Source/Editor/PropertyEditor/Private/DetailItemNode.cpp b/Engine/Source/Editor/PropertyEditor/Private/DetailItemNode.cpp
--- a/Engine/Source/Editor/PropertyEditor/Private/DetailItemNode.cpp
+++ b/Engine/Source/Editor/PropertyEditor/Private/DetailItemNode.cpp
@@ -226,6 +226,7 @@ TSharedRef< ITableRow > FDetailItemNode::GenerateWidgetForTableView( const TShar
 			SNew(SDetailCategoryTableRow, AsShared(), OwnerTable)
 			.DisplayName(Customization.GetPropertyNode()->GetDisplayName())
 			.AddMetaData<FTagMetaData>(TagMeta)
+			.ColumnSizeData(&ColumnSizeData)
 			.InnerCategory( true );
 	}
 	else
diff --git a/Engine/Source/Editor/UnrealEd/Classes/Editor/PropertyEditorTestObject.h b/Engine/Source/Editor/UnrealEd/Classes/Editor/PropertyEditorTestObject.h
--- a/Engine/Source/Editor/UnrealEd/Classes/Editor/PropertyEditorTestObject.h
+++ b/Engine/Source/Editor/UnrealEd/Classes/Editor/PropertyEditorTestObject.h
@@ -8,6 +8,7 @@
 #include "UObject/ScriptInterface.h"
 #include "UObject/SoftObjectPath.h"
 #include "Engine/BlendableInterface.h"
+#include "Engine/EngineTypes.h"
 #include "PropertyEditorTestObject.generated.h"
 
 class AActor;
@@ -122,7 +123,7 @@ struct FPropertyEditorTestBasicStruct
 UCLASS(transient)
 class UPropertyEditorTestObject : public UObject
 {
-    GENERATED_UCLASS_BODY()
+    GENERATED_BODY()
 
 	UPROPERTY(EditAnywhere, Category = BasicProperties)
 	int8 Int8Property;
@@ -483,4 +484,7 @@ class UPropertyEditorTestObject : public UObject
 
 	UPROPERTY(EditAnywhere, Category = EditCondition, meta = (EditCondition = "FloatEditCondition < 10"))
 	bool bEnabledWhenFloatLessThan10;
+
+	UPROPERTY(EditAnywhere, Category=InlineProperty, meta=(InlineCategoryProperty))
+	TEnumAsByte<EComponentMobility::Type> InlineProperty;
 };
