Added several const-correct getters for DetailTreeNode and derived classes.

diff --git a/Engine/Source/Editor/PropertyEditor/Private/CategoryPropertyNode.cpp b/Engine/Source/Editor/PropertyEditor/Private/CategoryPropertyNode.cpp
--- a/Engine/Source/Editor/PropertyEditor/Private/CategoryPropertyNode.cpp
+++ b/Engine/Source/Editor/PropertyEditor/Private/CategoryPropertyNode.cpp
@@ -18,7 +18,7 @@ FCategoryPropertyNode::~FCategoryPropertyNode(void)
 
 bool FCategoryPropertyNode::IsSubcategory() const
 {
-	return GetParentNode() != NULL && const_cast<FPropertyNode*>( GetParentNode() )->AsCategoryNode() != NULL;
+	return GetParentNode() != nullptr && GetParentNode()->AsCategoryNode() != nullptr;
 }
 
 FText FCategoryPropertyNode::GetDisplayName() const 
diff --git a/Engine/Source/Editor/PropertyEditor/Private/CategoryPropertyNode.h b/Engine/Source/Editor/PropertyEditor/Private/CategoryPropertyNode.h
--- a/Engine/Source/Editor/PropertyEditor/Private/CategoryPropertyNode.h
+++ b/Engine/Source/Editor/PropertyEditor/Private/CategoryPropertyNode.h
@@ -15,6 +15,7 @@ public:
 	 * Overridden function to get the derived object node
 	 */
 	virtual FCategoryPropertyNode* AsCategoryNode() override { return this; }
+	virtual const FCategoryPropertyNode* AsCategoryNode() const override { return this; }
 
 	/**
 	 * Accessor functions for getting a category name
diff --git a/Engine/Source/Editor/PropertyEditor/Private/DetailItemNode.h b/Engine/Source/Editor/PropertyEditor/Private/DetailItemNode.h
--- a/Engine/Source/Editor/PropertyEditor/Private/DetailItemNode.h
+++ b/Engine/Source/Editor/PropertyEditor/Private/DetailItemNode.h
@@ -64,7 +64,7 @@ public:
 	virtual void Tick( float DeltaTime ) override;
 	virtual bool ShouldShowOnlyChildren() const override;
 	virtual FName GetNodeName() const override;
-	virtual TSharedPtr<FDetailCategoryImpl> GetParentCategory() override { return ParentCategory.Pin(); }
+	virtual TSharedPtr<FDetailCategoryImpl> GetParentCategory() const override { return ParentCategory.Pin(); }
 	virtual FPropertyPath GetPropertyPath() const override;
 	virtual void SetIsHighlighted(bool bInIsHighlighted) override { bIsHighlighted = bInIsHighlighted; }
 	virtual bool IsHighlighted() const override { return bIsHighlighted; }
diff --git a/Engine/Source/Editor/PropertyEditor/Private/DetailLayoutBuilderImpl.cpp b/Engine/Source/Editor/PropertyEditor/Private/DetailLayoutBuilderImpl.cpp
--- a/Engine/Source/Editor/PropertyEditor/Private/DetailLayoutBuilderImpl.cpp
+++ b/Engine/Source/Editor/PropertyEditor/Private/DetailLayoutBuilderImpl.cpp
@@ -216,12 +216,12 @@ FDetailCategoryImpl& FDetailLayoutBuilderImpl::DefaultCategory( FName CategoryNa
 	return *CategoryImpl;
 }
 
-TSharedPtr<FDetailCategoryImpl> FDetailLayoutBuilderImpl::GetSubCategoryImpl(FName CategoryName)
+TSharedPtr<FDetailCategoryImpl> FDetailLayoutBuilderImpl::GetSubCategoryImpl(FName CategoryName) const
 {
 	return SubCategoryMap.FindRef(CategoryName);
 }
 
-bool FDetailLayoutBuilderImpl::HasCategory(FName CategoryName)
+bool FDetailLayoutBuilderImpl::HasCategory(FName CategoryName) const
 {
 	return DefaultCategoryMap.Contains(CategoryName);
 }
diff --git a/Engine/Source/Editor/PropertyEditor/Private/DetailLayoutBuilderImpl.h b/Engine/Source/Editor/PropertyEditor/Private/DetailLayoutBuilderImpl.h
--- a/Engine/Source/Editor/PropertyEditor/Private/DetailLayoutBuilderImpl.h
+++ b/Engine/Source/Editor/PropertyEditor/Private/DetailLayoutBuilderImpl.h
@@ -58,12 +58,15 @@ public:
 	 */
 	FDetailCategoryImpl& DefaultCategory(FName CategoryName);
 
-	TSharedPtr<FDetailCategoryImpl> GetSubCategoryImpl(FName CategoryName);
+	/** 
+	 * Find a subcategory with the given name, if it exists
+	 */
+	TSharedPtr<FDetailCategoryImpl> GetSubCategoryImpl(FName CategoryName) const;
 
 	/**
-	* Return true if the category exist
-	*/
-	bool HasCategory(FName CategoryName);
+	 * @returns true if a category with the given name exists
+	 */
+	bool HasCategory(FName CategoryName) const;
 
 	/**
 	 * Generates the layout for this detail builder
@@ -92,7 +95,6 @@ public:
 	 *
 	 * @param PropertyPath	The path to the property
 	 * @param ClassOutermost	The outer class of the property
-	 * @param bMarkPropertyAsCustomized	Whether or not to mark the property as customized so it does not appear in the default layout
 	 * @return The found property node
 	 */
 	TSharedPtr<FPropertyNode> GetPropertyNode(FName PropertyPath, const UClass* ClassOutermost, FName InstanceName) const;
@@ -101,7 +103,6 @@ public:
 	 * Gets the property node from the provided handle
 	 *
 	 * @param PropertyHandle	The property handle to get the node from
-	 * @param bMarkAsCustomized	Whether or not to mark the value as customized
 	 */
 	TSharedPtr<FPropertyNode> GetPropertyNode(TSharedPtr<IPropertyHandle> PropertyHandle) const;
 
diff --git a/Engine/Source/Editor/PropertyEditor/Private/DetailTreeNode.h b/Engine/Source/Editor/PropertyEditor/Private/DetailTreeNode.h
--- a/Engine/Source/Editor/PropertyEditor/Private/DetailTreeNode.h
+++ b/Engine/Source/Editor/PropertyEditor/Private/DetailTreeNode.h
@@ -99,7 +99,7 @@ public:
 	/**
 	 * @return The category node that this node is nested in, if any:
 	 */
-	virtual TSharedPtr<FDetailCategoryImpl> GetParentCategory() { return TSharedPtr<FDetailCategoryImpl>(); }
+	virtual TSharedPtr<FDetailCategoryImpl> GetParentCategory() const { return TSharedPtr<FDetailCategoryImpl>(); }
 
 	/**
 	 * @return The property path that this node is associate with, if any:
diff --git a/Engine/Source/Editor/PropertyEditor/Private/ItemPropertyNode.h b/Engine/Source/Editor/PropertyEditor/Private/ItemPropertyNode.h
--- a/Engine/Source/Editor/PropertyEditor/Private/ItemPropertyNode.h
+++ b/Engine/Source/Editor/PropertyEditor/Private/ItemPropertyNode.h
@@ -32,6 +32,7 @@ public:
 	 * Overridden function to get the derived object node
 	 */
 	virtual FItemPropertyNode* AsItemPropertyNode() override { return this; }
+	virtual const FItemPropertyNode* AsItemPropertyNode() const override { return this; }
 
 	/** Display name override to use instead of the property name */
 	void SetDisplayNameOverride( const FText& InDisplayNameOverride ) override;
diff --git a/Engine/Source/Editor/PropertyEditor/Private/PropertyNode.h b/Engine/Source/Editor/PropertyEditor/Private/PropertyNode.h
--- a/Engine/Source/Editor/PropertyEditor/Private/PropertyNode.h
+++ b/Engine/Source/Editor/PropertyEditor/Private/PropertyNode.h
@@ -295,32 +295,40 @@ public:
 	void ClearCachedReadAddresses(bool bRecursive = true);
 
 	/**
-	 * Interface function to get at the dervied FObjectPropertyNodeWx class
+	 * Interface function to get at the derived FObjectPropertyNode class
 	 */
 	virtual class FObjectPropertyNode* AsObjectNode() { return nullptr; }
 	virtual const FObjectPropertyNode* AsObjectNode() const { return nullptr; }
 
-	virtual FComplexPropertyNode* AsComplexNode() { return nullptr; }
+	/**
+	 * Interface function to get at the derived FComplexPropertyNode class
+	 */
+	virtual class FComplexPropertyNode* AsComplexNode() { return nullptr; }
 	virtual const FComplexPropertyNode* AsComplexNode() const { return nullptr; }
 
 	/**
-	 * Interface function to get at the dervied FCategoryPropertyNodeWx class
+	 * Interface function to get at the derived FCategoryPropertyNode class
 	 */
 	virtual class FCategoryPropertyNode* AsCategoryNode() { return nullptr; }
+	virtual const FCategoryPropertyNode* AsCategoryNode() const { return nullptr; }
 
 	/**
-	 * Interface function to get at the dervied FItemPropertyNodeWx class
+	 * Interface function to get at the derived FItemPropertyNode class
 	 */
 	virtual class FItemPropertyNode* AsItemPropertyNode() { return nullptr; }
+	virtual const FItemPropertyNode* AsItemPropertyNode() const { return nullptr; }
 
 	/**
-	 * Follows the chain of items upwards until it finds the object window that houses this item.
+	 * Follows the chain of items upwards until it finds the complex property that houses this item.
 	 */
 	class FComplexPropertyNode* FindComplexParent();
-	const class FComplexPropertyNode* FindComplexParent() const;
+	const FComplexPropertyNode* FindComplexParent() const;
 
+	/**
+	 * Follows the chain of items upwards until it finds the object property that houses this item.
+	 */
 	class FObjectPropertyNode* FindObjectItemParent();
-	const class FObjectPropertyNode* FindObjectItemParent() const;
+	const FObjectPropertyNode* FindObjectItemParent() const;
 
 	/**
 	 * Follows the top-most object window that contains this property window item.
