Added ShouldBeVisible to PropertyEditorHelpers that contains the most common checks that are duplicated to determine visibility for details nodes.

diff --git a/Engine/Source/Editor/PropertyEditor/Private/CategoryPropertyNode.cpp b/Engine/Source/Editor/PropertyEditor/Private/CategoryPropertyNode.cpp
--- a/Engine/Source/Editor/PropertyEditor/Private/CategoryPropertyNode.cpp
+++ b/Engine/Source/Editor/PropertyEditor/Private/CategoryPropertyNode.cpp
@@ -43,7 +43,7 @@ void FCategoryPropertyNode::InitChildNodes()
 
 	TArray<UProperty*> Properties;
 	// The parent of a category window has to be an object window.
-	auto ComplexNode = FindComplexParent();
+	FComplexPropertyNode* ComplexNode = FindComplexParent();
 	if (ComplexNode)
 	{
 		// Get a list of properties that are in the same category
@@ -64,14 +64,9 @@ void FCategoryPropertyNode::InitChildNodes()
 
 			if (bMetaDataAllowVisible)
 			{
-				static const FName Name_InlineEditConditionToggle("InlineEditConditionToggle");
-				const bool bOnlyShowAsInlineEditCondition = (*It)->HasMetaData(Name_InlineEditConditionToggle);
-				const bool bShowIfEditableProperty = (*It)->HasAnyPropertyFlags(CPF_Edit);
-				const bool bShowIfDisableEditOnInstance = !(*It)->HasAnyPropertyFlags(CPF_DisableEditOnInstance) || bShouldShowDisableEditOnInstance;
-
 				// Add if we are showing non-editable props and this is the 'None' category, 
-				// or if this is the right category, and we are either showing non-editable
-				if (FObjectEditorUtils::GetCategoryFName(*It) == CategoryName && (bShowHiddenProperties || (bShowIfEditableProperty && !bOnlyShowAsInlineEditCondition && bShowIfDisableEditOnInstance)))
+				// or if this is the right category, and we are showing non-editable
+				if (FObjectEditorUtils::GetCategoryFName(*It) == CategoryName && PropertyEditorHelpers::ShouldBeVisible(*this, *It))
 				{
 					Properties.Add(*It);
 				}
diff --git a/Engine/Source/Editor/PropertyEditor/Private/ItemPropertyNode.cpp b/Engine/Source/Editor/PropertyEditor/Private/ItemPropertyNode.cpp
--- a/Engine/Source/Editor/PropertyEditor/Private/ItemPropertyNode.cpp
+++ b/Engine/Source/Editor/PropertyEditor/Private/ItemPropertyNode.cpp
@@ -302,11 +302,8 @@ void FItemPropertyNode::InitChildNodes()
 		for (TFieldIterator<UProperty> It(StructProperty->Struct); It; ++It)
 		{
 			UProperty* StructMember = *It;
-			static const FName Name_InlineEditConditionToggle("InlineEditConditionToggle");
-			const bool bOnlyShowAsInlineEditCondition = StructMember->HasMetaData(Name_InlineEditConditionToggle);
-			const bool bShowIfEditableProperty = StructMember->HasAnyPropertyFlags(CPF_Edit);
-			const bool bShowIfDisableEditOnInstance = !StructMember->HasAnyPropertyFlags(CPF_DisableEditOnInstance) || bShouldShowDisableEditOnInstance;
-			if (bShouldShowHiddenProperties || (bShowIfEditableProperty && !bOnlyShowAsInlineEditCondition && bShowIfDisableEditOnInstance))
+			
+			if (PropertyEditorHelpers::ShouldBeVisible(*this, StructMember))
 			{
 				StructMembers.Add(StructMember);
 			}
@@ -316,7 +313,7 @@ void FItemPropertyNode::InitChildNodes()
 
 		for (UProperty* StructMember : StructMembers)
 		{
-			TSharedPtr<FItemPropertyNode> NewItemNode( new FItemPropertyNode );//;//CreatePropertyItem(StructMember,INDEX_NONE,this);
+			TSharedPtr<FItemPropertyNode> NewItemNode( new FItemPropertyNode );
 		
 			FPropertyNodeInitParams InitParams;
 			InitParams.ParentNode = SharedThis(this);
diff --git a/Engine/Source/Editor/PropertyEditor/Private/ObjectPropertyNode.cpp b/Engine/Source/Editor/PropertyEditor/Private/ObjectPropertyNode.cpp
--- a/Engine/Source/Editor/PropertyEditor/Private/ObjectPropertyNode.cpp
+++ b/Engine/Source/Editor/PropertyEditor/Private/ObjectPropertyNode.cpp
@@ -7,6 +7,7 @@
 #include "ItemPropertyNode.h"
 #include "ObjectEditorUtils.h"
 #include "EditorCategoryUtils.h"
+#include "PropertyEditorHelpers.h"
 
 #if WITH_EDITORONLY_DATA
 #include "Engine/Blueprint.h"
@@ -396,6 +397,23 @@ void FObjectPropertyNode::InitChildNodes()
 	InternalInitChildNodes();
 }
 
+static TSharedPtr<FPropertyNode> FindChildCategory( TSharedPtr<FPropertyNode> ParentNode, FName CategoryName )
+{
+	for( int32 CurNodeIndex = 0; CurNodeIndex < ParentNode->GetNumChildNodes(); ++CurNodeIndex )
+	{
+		TSharedPtr<FPropertyNode>& ChildNode = ParentNode->GetChildNode(CurNodeIndex);
+		check( ChildNode.IsValid() );
+
+		// Is this a category node?
+		FCategoryPropertyNode* ChildCategoryNode = ChildNode->AsCategoryNode();
+		if( ChildCategoryNode != nullptr && ChildCategoryNode->GetCategoryName() == CategoryName )
+		{
+			return ChildNode;
+		}
+	}
+	return nullptr;
+}
+
 void FObjectPropertyNode::InternalInitChildNodes( FName SinglePropertyName )
 {
 	HiddenCategories.Empty();
@@ -447,33 +465,22 @@ void FObjectPropertyNode::InternalInitChildNodes( FName SinglePropertyName )
 			if( FEditorCategoryUtils::IsCategoryHiddenFromClass(Class, CategoryName.ToString()) )
 			{
 				HiddenCategories.Add( CategoryName );
-
-				bHidden = true;
 				break;
 			}
 		}
 
-		// 
-		static const FName Name_InlineEditConditionToggle("InlineEditConditionToggle");
-		if (It->HasMetaData(Name_InlineEditConditionToggle))
-		{
-			bHidden = true;
-		}
-
 		bool bMetaDataAllowVisible = true;
 		static const FName Name_bShowOnlyWhenTrue("bShowOnlyWhenTrue");
-		const FString& MetaDataVisibilityCheckString = It->GetMetaData(Name_bShowOnlyWhenTrue);
-		if (MetaDataVisibilityCheckString.Len())
+		const FString& ShowOnlyWhenTrueString = It->GetMetaData(Name_bShowOnlyWhenTrue);
+		if (ShowOnlyWhenTrueString.Len())
 		{
 			//ensure that the metadata visibility string is actually set to true in order to show this property
-			GConfig->GetBool(TEXT("UnrealEd.PropertyFilters"), *MetaDataVisibilityCheckString, bMetaDataAllowVisible, GEditorPerProjectIni);
+			GConfig->GetBool(TEXT("UnrealEd.PropertyFilters"), *ShowOnlyWhenTrueString, bMetaDataAllowVisible, GEditorPerProjectIni);
 		}
 
 		if (bMetaDataAllowVisible)
 		{
-			const bool bShowIfNonHiddenEditableProperty = (*It)->HasAnyPropertyFlags(CPF_Edit) && !bHidden;
-			const bool bShowIfDisableEditOnInstance = !(*It)->HasAnyPropertyFlags(CPF_DisableEditOnInstance) || bShouldShowDisableEditOnInstance;
-			if( bShouldShowHiddenProperties || (bShowIfNonHiddenEditableProperty && bShowIfDisableEditOnInstance) )
+			if( PropertyEditorHelpers::ShouldBeVisible(*this, *It) )
 			{
 				if (!CategoriesFromBlueprints.Contains(CategoryName) && !CategoriesFromProperties.Contains(CategoryName) )
 				{
@@ -523,31 +530,10 @@ void FObjectPropertyNode::InternalInitChildNodes( FName SinglePropertyName )
 				const FName CategoryName( *CurCategoryPathString );
 
 				// Check to see if we've already created a category at the specified path level
-				bool bFoundMatchingCategory = false;
-				{
-					for( int32 CurNodeIndex = 0; CurNodeIndex < ParentLevelNode->GetNumChildNodes(); ++CurNodeIndex )
-					{
-						TSharedPtr<FPropertyNode>& ChildNode = ParentLevelNode->GetChildNode( CurNodeIndex );
-						check( ChildNode.IsValid() );
-
-						// Is this a category node?
-						FCategoryPropertyNode* ChildCategoryNode = ChildNode->AsCategoryNode();
-						if( ChildCategoryNode != NULL )
-						{
-							// Does the name match?
-							if( ChildCategoryNode->GetCategoryName() == CategoryName )
-							{
-								// Descend by using the child node as the new parent
-								bFoundMatchingCategory = true;
-								ParentLevelNode = ChildNode;
-								break;
-							}
-						}
-					}
-				}
+				TSharedPtr<FPropertyNode> FoundCategory = FindChildCategory( ParentLevelNode, CategoryName );
 
 				// If we didn't find the category, then we'll need to create it now!
-				if( !bFoundMatchingCategory )
+				if( !FoundCategory.IsValid() )
 				{
 					// Create the category node and assign it to its parent node
 					TSharedPtr<FCategoryPropertyNode> NewCategoryNode( new FCategoryPropertyNode );
@@ -582,6 +568,10 @@ void FObjectPropertyNode::InternalInitChildNodes( FName SinglePropertyName )
 					// Descend into the newly created category by using this node as the new parent
 					ParentLevelNode = NewCategoryNode;
 				}
+				else
+				{
+					ParentLevelNode = FoundCategory;
+				}
 			}
 		}
 	}
@@ -590,11 +580,7 @@ void FObjectPropertyNode::InternalInitChildNodes( FName SinglePropertyName )
 		// Iterate over all fields, creating items.
 		for( TFieldIterator<UProperty> It(BaseClass.Get()); It; ++It )
 		{
-			static const FName Name_InlineEditConditionToggle("InlineEditConditionToggle");
-			const bool bOnlyShowAsInlineEditCondition = (*It)->HasMetaData(Name_InlineEditConditionToggle);
-			const bool bShowIfNonHiddenEditableProperty = (*It)->HasAnyPropertyFlags(CPF_Edit) && !FEditorCategoryUtils::IsCategoryHiddenFromClass(BaseClass.Get(), FObjectEditorUtils::GetCategory(*It));
-			const bool bShowIfDisableEditOnInstance = !(*It)->HasAnyPropertyFlags(CPF_DisableEditOnInstance) || bShouldShowDisableEditOnInstance;
-			if (bShouldShowHiddenProperties || (bShowIfNonHiddenEditableProperty && !bOnlyShowAsInlineEditCondition && bShowIfDisableEditOnInstance))
+			if (PropertyEditorHelpers::ShouldBeVisible(*this, *It))
 			{
 				UProperty* CurProp = *It;
 				if( SinglePropertyName == NAME_None || CurProp->GetFName() == SinglePropertyName )
diff --git a/Engine/Source/Editor/PropertyEditor/Private/PropertyEditorHelpers.cpp b/Engine/Source/Editor/PropertyEditor/Private/PropertyEditorHelpers.cpp
--- a/Engine/Source/Editor/PropertyEditor/Private/PropertyEditorHelpers.cpp
+++ b/Engine/Source/Editor/PropertyEditor/Private/PropertyEditorHelpers.cpp
@@ -378,6 +378,24 @@ ECheckBoxState SEditConditionWidget::OnGetEditConditionCheckState() const
 
 namespace PropertyEditorHelpers
 {
+	bool ShouldBeVisible(const FPropertyNode& InParentNode, const UProperty* Property)
+	{
+		const bool bShouldShowHiddenProperties = !!InParentNode.HasNodeFlags(EPropertyNodeFlags::ShouldShowHiddenProperties);
+		if (bShouldShowHiddenProperties)
+		{
+			return true;
+		}
+
+		const bool bShouldShowDisableEditOnInstance = !!InParentNode.HasNodeFlags(EPropertyNodeFlags::ShouldShowDisableEditOnInstance);
+
+		static const FName Name_InlineEditConditionToggle("InlineEditConditionToggle");
+		const bool bOnlyShowAsInlineEditCondition = Property->HasMetaData(Name_InlineEditConditionToggle);
+		const bool bShowIfEditableProperty = Property->HasAnyPropertyFlags(CPF_Edit);
+		const bool bShowIfDisableEditOnInstance = bShouldShowDisableEditOnInstance || !Property->HasAnyPropertyFlags(CPF_DisableEditOnInstance);
+
+		return bShowIfEditableProperty && !bOnlyShowAsInlineEditCondition && bShowIfDisableEditOnInstance;
+	}
+
 	bool IsBuiltInStructProperty( const UProperty* Property )
 	{
 		bool bIsBuiltIn = false;
diff --git a/Engine/Source/Editor/PropertyEditor/Private/PropertyEditorHelpers.h b/Engine/Source/Editor/PropertyEditor/Private/PropertyEditorHelpers.h
--- a/Engine/Source/Editor/PropertyEditor/Private/PropertyEditorHelpers.h
+++ b/Engine/Source/Editor/PropertyEditor/Private/PropertyEditorHelpers.h
@@ -116,6 +116,12 @@ private:
 
 namespace PropertyEditorHelpers
 {
+	/** 
+	 * Perform common checks from metadata to determine if a given Property
+	 * should be visible when displayed underneath the given property node.
+	 */
+	bool ShouldBeVisible(const FPropertyNode& InParentNode, const UProperty* Property);
+
 	/**
 	 * Returns whether or not a property is a built in struct property like a vector or color
 	 *
diff --git a/Engine/Source/Editor/PropertyEditor/Private/PropertyHandleImpl.cpp b/Engine/Source/Editor/PropertyEditor/Private/PropertyHandleImpl.cpp
--- a/Engine/Source/Editor/PropertyEditor/Private/PropertyHandleImpl.cpp
+++ b/Engine/Source/Editor/PropertyEditor/Private/PropertyHandleImpl.cpp
@@ -3067,17 +3067,8 @@ TArray<TSharedPtr<IPropertyHandle>> FPropertyHandleBase::AddChildStructure( TSha
 	for (TFieldIterator<UProperty> It(InStruct->GetStruct()); It; ++It)
 	{
 		UProperty* StructMember = *It;
-		if (!StructMember)
-		{
-			continue;
-		}
-
-		static const FName Name_InlineEditConditionToggle("InlineEditConditionToggle");
-		const bool bOnlyShowAsInlineEditCondition = StructMember->HasMetaData(Name_InlineEditConditionToggle);
-		const bool bShowIfEditableProperty = StructMember->HasAnyPropertyFlags(CPF_Edit);
-		const bool bShowIfDisableEditOnInstance = !StructMember->HasAnyPropertyFlags(CPF_DisableEditOnInstance) || bShouldShowDisableEditOnInstance;
 
-		if (bShouldShowHiddenProperties || (bShowIfEditableProperty && !bOnlyShowAsInlineEditCondition && bShowIfDisableEditOnInstance))
+		if (PropertyEditorHelpers::ShouldBeVisible(*StructPropertyNode.Get(), StructMember))
 		{
 			TSharedRef<FItemPropertyNode> NewItemNode(new FItemPropertyNode);
 
diff --git a/Engine/Source/Editor/PropertyEditor/Private/StructurePropertyNode.cpp b/Engine/Source/Editor/PropertyEditor/Private/StructurePropertyNode.cpp
--- a/Engine/Source/Editor/PropertyEditor/Private/StructurePropertyNode.cpp
+++ b/Engine/Source/Editor/PropertyEditor/Private/StructurePropertyNode.cpp
@@ -20,12 +20,7 @@ void FStructurePropertyNode::InitChildNodes()
 
 		if (StructMember)
 		{
-			static const FName Name_InlineEditConditionToggle("InlineEditConditionToggle");
-			const bool bOnlyShowAsInlineEditCondition = StructMember->HasMetaData(Name_InlineEditConditionToggle);
-			const bool bShowIfEditableProperty = StructMember->HasAnyPropertyFlags(CPF_Edit);
-			const bool bShowIfDisableEditOnInstance = !StructMember->HasAnyPropertyFlags(CPF_DisableEditOnInstance) || bShouldShowDisableEditOnInstance;
-
-			if (bShouldShowHiddenProperties || (bShowIfEditableProperty && !bOnlyShowAsInlineEditCondition && bShowIfDisableEditOnInstance))
+			if (PropertyEditorHelpers::ShouldBeVisible(*this, StructMember))
 			{
 				StructMembers.Add(StructMember);
 			}
