Added EditConditionHides meta property that determines visibility from the edit condition.

diff --git a/Engine/Source/Editor/PropertyEditor/Private/DetailPropertyRow.cpp b/Engine/Source/Editor/PropertyEditor/Private/DetailPropertyRow.cpp
--- a/Engine/Source/Editor/PropertyEditor/Private/DetailPropertyRow.cpp
+++ b/Engine/Source/Editor/PropertyEditor/Private/DetailPropertyRow.cpp
@@ -33,7 +33,7 @@ FDetailPropertyRow::FDetailPropertyRow(TSharedPtr<FPropertyNode> InPropertyNode,
 	if( InPropertyNode.IsValid() )
 	{
 		TSharedRef<FPropertyNode> PropertyNodeRef = PropertyNode.ToSharedRef();
-
+		
 		PropertyHandle = InParentCategory->GetParentLayoutImpl().GetPropertyHandle(PropertyNodeRef);
 
 		const TSharedRef<IPropertyUtilities> Utilities = InParentCategory->GetParentLayoutImpl().GetPropertyUtilities();
@@ -160,7 +160,8 @@ bool FDetailPropertyRow::ShowOnlyChildren() const
 
 bool FDetailPropertyRow::RequiresTick() const
 {
-	return PropertyVisibility.IsBound();
+	return PropertyVisibility.IsBound() || 
+		(PropertyEditor.IsValid() && PropertyEditor->IsOnlyVisibleWhenEditConditionMet());
 }
 
 FDetailWidgetRow& FDetailPropertyRow::CustomWidget( bool bShowChildren )
@@ -538,6 +539,16 @@ void FDetailPropertyRow::MakeExternalPropertyRowCustomization(const TArray<UObje
 	}
 }
 
+EVisibility FDetailPropertyRow::GetPropertyVisibility() const
+{
+	if (PropertyEditor->IsOnlyVisibleWhenEditConditionMet() && !PropertyEditor->IsEditConditionMet())
+	{
+		return EVisibility::Collapsed;
+	}
+
+	return PropertyVisibility.Get();
+}
+
 bool FDetailPropertyRow::HasEditCondition() const
 {
 	return ( PropertyEditor.IsValid() && PropertyEditor->HasEditCondition() ) || CustomEditCondition.IsValid();
diff --git a/Engine/Source/Editor/PropertyEditor/Private/DetailPropertyRow.h b/Engine/Source/Editor/PropertyEditor/Private/DetailPropertyRow.h
--- a/Engine/Source/Editor/PropertyEditor/Private/DetailPropertyRow.h
+++ b/Engine/Source/Editor/PropertyEditor/Private/DetailPropertyRow.h
@@ -100,7 +100,7 @@ public:
 	/** 
 	 * @return The visibility of this property
 	 */
-	EVisibility GetPropertyVisibility() const { return PropertyVisibility.Get(); }
+	EVisibility GetPropertyVisibility() const;
 
 	static void MakeExternalPropertyRowCustomization(TSharedPtr<FStructOnScope> StructData, FName PropertyName, TSharedRef<FDetailCategoryImpl> ParentCategory, struct FDetailLayoutCustomization& OutCustomization);
 	static void MakeExternalPropertyRowCustomization(const TArray<UObject*>& InObjects, FName PropertyName, TSharedRef<FDetailCategoryImpl> ParentCategory, struct FDetailLayoutCustomization& OutCustomization, TOptional<bool> bAllowChildrenOverride = TOptional<bool>(), TOptional<bool> bCreateCategoryNodesOverride = TOptional<bool>());
diff --git a/Engine/Source/Editor/PropertyEditor/Private/Presentation/PropertyEditor/PropertyEditor.cpp b/Engine/Source/Editor/PropertyEditor/Private/Presentation/PropertyEditor/PropertyEditor.cpp
--- a/Engine/Source/Editor/PropertyEditor/Private/Presentation/PropertyEditor/PropertyEditor.cpp
+++ b/Engine/Source/Editor/PropertyEditor/Private/Presentation/PropertyEditor/PropertyEditor.cpp
@@ -566,6 +566,17 @@ void FPropertyEditor::RequestRefresh()
 	PropertyUtilities->RequestRefresh();
 }
 
+bool FPropertyEditor::IsOnlyVisibleWhenEditConditionMet() const
+{
+	static const FName Name_EditConditionHides("EditConditionHides");
+	if (PropertyNode->GetProperty()->HasMetaData(Name_EditConditionHides))
+	{
+		return HasEditCondition();
+	}
+
+	return false;
+}
+
 bool FPropertyEditor::HasEditCondition() const 
 { 
 	return EditConditionExpression.IsValid();
diff --git a/Engine/Source/Editor/PropertyEditor/Private/Presentation/PropertyEditor/PropertyEditor.h b/Engine/Source/Editor/PropertyEditor/Private/Presentation/PropertyEditor/PropertyEditor.h
--- a/Engine/Source/Editor/PropertyEditor/Private/Presentation/PropertyEditor/PropertyEditor.h
+++ b/Engine/Source/Editor/PropertyEditor/Private/Presentation/PropertyEditor/PropertyEditor.h
@@ -52,6 +52,9 @@ public:
 	/**	@return Whether the condition has been met to allow editing of this property's value */
 	bool IsEditConditionMet() const;
 
+	/**	@return Whether this property derives its visibility from its edit condition */
+	bool IsOnlyVisibleWhenEditConditionMet() const;
+
 	/**	@return The tooltip for this property editor */
 	FText GetToolTipText() const;
 
diff --git a/Engine/Source/Editor/UnrealEd/Classes/Editor/PropertyEditorTestObject.h b/Engine/Source/Editor/UnrealEd/Classes/Editor/PropertyEditorTestObject.h
--- a/Engine/Source/Editor/UnrealEd/Classes/Editor/PropertyEditorTestObject.h
+++ b/Engine/Source/Editor/UnrealEd/Classes/Editor/PropertyEditorTestObject.h
@@ -485,6 +485,15 @@ class UPropertyEditorTestObject : public UObject
 	UPROPERTY(EditAnywhere, Category = EditCondition, meta = (EditCondition = "FloatEditCondition < 10"))
 	bool bEnabledWhenFloatLessThan10;
 
-	UPROPERTY(EditAnywhere, Category=InlineProperty, meta=(InlineCategoryProperty))
+	UPROPERTY(EditAnywhere, Category = OnlyInlineProperty, meta = (InlineCategoryProperty))
 	TEnumAsByte<EComponentMobility::Type> InlineProperty;
+
+	UPROPERTY(EditAnywhere, Category = EditConditionHides, meta = (InlineCategoryProperty))
+	TEnumAsByte<EComponentMobility::Type> PropertyThatHides;
+
+	UPROPERTY(EditAnywhere, Category = EditConditionHides, meta = (EditConditionHides, EditCondition = "PropertyThatHides == EComponentMobility::Static"))
+	bool bVisibleWhenStatic;
+
+	UPROPERTY(EditAnywhere, Category = EditConditionHides, meta = (EditConditionHides, EditCondition = "PropertyThatHides == EComponentMobility::Stationary"))
+	int32 VisibleWhenStationary;
 };
