Fixed crash occurring in cases where PropertyEditor for FDetailPropertyRow is not yet set.

diff --git a/Engine/Source/Editor/PropertyEditor/Private/DetailPropertyRow.cpp b/Engine/Source/Editor/PropertyEditor/Private/DetailPropertyRow.cpp
--- a/Engine/Source/Editor/PropertyEditor/Private/DetailPropertyRow.cpp
+++ b/Engine/Source/Editor/PropertyEditor/Private/DetailPropertyRow.cpp
@@ -541,7 +541,7 @@ void FDetailPropertyRow::MakeExternalPropertyRowCustomization(const TArray<UObje
 
 EVisibility FDetailPropertyRow::GetPropertyVisibility() const
 {
-	if (PropertyEditor->IsOnlyVisibleWhenEditConditionMet() && !PropertyEditor->IsEditConditionMet())
+	if (PropertyEditor.IsValid() && PropertyEditor->IsOnlyVisibleWhenEditConditionMet() && !PropertyEditor->IsEditConditionMet())
 	{
 		return EVisibility::Collapsed;
 	}
