SPropertyEditorAsset now correctly doesn't display thumbnails when an Array/Set/Map of SoftObjectPaths is tagged with "DisplayThumbnail=false".

This wasn't occurring elsewhere due to other places passing in the FPropertyEditor reference when constructing the widget, but FSoftObjectPathCustomization does not do so. However, it does pass in an IPropertyHandle so we can use that instead.

diff --git a/Engine/Source/Editor/PropertyEditor/Private/UserInterface/PropertyEditor/SPropertyEditorAsset.cpp b/Engine/Source/Editor/PropertyEditor/Private/UserInterface/PropertyEditor/SPropertyEditorAsset.cpp
--- a/Engine/Source/Editor/PropertyEditor/Private/UserInterface/PropertyEditor/SPropertyEditorAsset.cpp
+++ b/Engine/Source/Editor/PropertyEditor/Private/UserInterface/PropertyEditor/SPropertyEditorAsset.cpp
@@ -39,34 +39,39 @@ bool SPropertyEditorAsset::ShouldDisplayThumbnail( const FArguments& InArgs, con
 {
 	bool bDisplayThumbnail = InArgs._DisplayThumbnail && InArgs._ThumbnailPool.IsValid() && (!InObjectClass || !InObjectClass->IsChildOf(AActor::StaticClass()));
 	
-	if(PropertyEditor.IsValid())
+	if (InArgs._ThumbnailPool.IsValid())
 	{
-		// also check metadata for thumbnail & text display
-		if(InArgs._ThumbnailPool.IsValid())
+		const UProperty* PropertyToCheck = nullptr;
+		if (PropertyEditor.IsValid())
+		{
+			PropertyToCheck = PropertyEditor->GetProperty();
+		}
+		else if (PropertyHandle.IsValid())
 		{
-			const UProperty* ArrayParent = PropertyEditorHelpers::GetArrayParent( *PropertyEditor->GetPropertyNode() );
-			const UProperty* SetParent = PropertyEditorHelpers::GetSetParent( *PropertyEditor->GetPropertyNode() );
-			const UProperty* MapParent = PropertyEditorHelpers::GetMapParent( *PropertyEditor->GetPropertyNode() );
+			PropertyToCheck = PropertyHandle->GetProperty();
+		}
 
-			const UProperty* PropertyToCheck = PropertyEditor->GetProperty();
-			if( ArrayParent != nullptr )
+		if (PropertyToCheck != nullptr)
+		{
+			UArrayProperty* ArrayParent = Cast<UArrayProperty>(PropertyToCheck->GetOuter());
+			USetProperty* SetParent = Cast<USetProperty>(PropertyToCheck->GetOuter());
+			UMapProperty* MapParent = Cast<UMapProperty>(PropertyToCheck->GetOuter());
+			if (ArrayParent != nullptr)
 			{
-				// If the property is a child of an array property, the parent will have the display thumbnail metadata
 				PropertyToCheck = ArrayParent;
 			}
-			else if ( SetParent != nullptr )
+			else if (SetParent != nullptr)
 			{
 				PropertyToCheck = SetParent;
 			}
-			else if ( MapParent != nullptr )
+			else if (MapParent != nullptr)
 			{
 				PropertyToCheck = MapParent;
 			}
 
-			const FString& DisplayThumbnailString = PropertyToCheck->GetMetaData(TEXT("DisplayThumbnail"));
-			if (DisplayThumbnailString.Len() > 0)
+			if (PropertyToCheck->HasMetaData(TEXT("DisplayThumbnail")))
 			{
-				bDisplayThumbnail = DisplayThumbnailString == TEXT("true");
+				bDisplayThumbnail = PropertyToCheck->GetBoolMetaData(TEXT("DisplayThumbnail"));
 			}
 		}
 	}
