SPropertyEditorAsset now handles multiple actor types correctly when set with the AllowedClasses meta. It also now supports the ExactClass metadata.

As a result, FSoftObjectPath members now also support UInterfaces, because it no longer uses custom filtering logic.

diff --git a/Engine/Source/Editor/DetailCustomizations/Private/SoftObjectPathCustomization.cpp b/Engine/Source/Editor/DetailCustomizations/Private/SoftObjectPathCustomization.cpp
--- a/Engine/Source/Editor/DetailCustomizations/Private/SoftObjectPathCustomization.cpp
+++ b/Engine/Source/Editor/DetailCustomizations/Private/SoftObjectPathCustomization.cpp
@@ -10,72 +10,7 @@
 void FSoftObjectPathCustomization::CustomizeHeader( TSharedRef<IPropertyHandle> InStructPropertyHandle, FDetailWidgetRow& HeaderRow, IPropertyTypeCustomizationUtils& StructCustomizationUtils )
 {
 	StructPropertyHandle = InStructPropertyHandle;
-
-	const FString& AllowedClassFilterString = StructPropertyHandle->GetMetaData("AllowedClasses");
-	if (!AllowedClassFilterString.IsEmpty())
-	{
-		TArray<FString> AllowedClassFilterNames;
-		AllowedClassFilterString.ParseIntoArray(AllowedClassFilterNames, TEXT(","), true);
-
-		for(auto It = AllowedClassFilterNames.CreateConstIterator(); It; ++It)
-		{
-			const FString& ClassName = *It;
-
-			UClass* Class = FindObject<UClass>(ANY_PACKAGE, *ClassName);
-			if (!Class)
-			{
-				Class = LoadObject<UClass>(nullptr, *ClassName);
-			}
-			if (Class)
-			{
-				AllowedClassFilters.Add(Class);
-			}
-		}
-	}
-
-	const FString& DisallowedClassFilterString = StructPropertyHandle->GetMetaData("DisallowedClasses");
-	if (!DisallowedClassFilterString.IsEmpty())
-	{
-		TArray<FString> DisallowedClassFilterNames;
-		DisallowedClassFilterString.ParseIntoArray(DisallowedClassFilterNames, TEXT(","), true);
-
-		for (auto It = DisallowedClassFilterNames.CreateConstIterator(); It; ++It)
-		{
-			const FString& ClassName = *It;
-
-			UClass* Class = FindObject<UClass>(ANY_PACKAGE, *ClassName);
-			if (!Class)
-			{
-				Class = LoadObject<UClass>(nullptr, *ClassName);
-			}
-			if (Class)
-			{
-				DisallowedClassFilters.Add(Class);
-			}
-		}
-	}
-
-	bExactClass = StructPropertyHandle->GetBoolMetaData("ExactClass");
-
-	FOnShouldFilterAsset AssetFilter;
-	UClass* ClassFilter = UObject::StaticClass();
-	if (AllowedClassFilters.Num() == 1 && DisallowedClassFilters.Num() <= 0 && !bExactClass)
-	{
-		// If we only have one class to filter on, set it as the class type filter rather than use a filter callback
-		// We can only do this if we don't need an exact match, as the class filter also allows derived types
-		// The class filter is much faster than the callback as we're not performing two different sets of type tests
-		// (one against UObject, one against the actual type)
-		ClassFilter = AllowedClassFilters[0];
-	}
-	else if (AllowedClassFilters.Num() > 0 || DisallowedClassFilters.Num() > 0)
-	{
-		// Only bind the filter if we have classes that need filtering
-		AssetFilter.BindSP(this, &FSoftObjectPathCustomization::OnShouldFilterAsset);
-	}
-
-	// Can the field be cleared
-	const bool bAllowClear = !(StructPropertyHandle->GetMetaDataProperty()->PropertyFlags & CPF_NoClear);
-
+	
 	HeaderRow
 	.NameContent()
 	[
@@ -89,9 +24,6 @@ void FSoftObjectPathCustomization::CustomizeHeader( TSharedRef<IPropertyHandle>
 		SNew( SObjectPropertyEntryBox )
 		.PropertyHandle( InStructPropertyHandle )
 		.ThumbnailPool( StructCustomizationUtils.GetThumbnailPool() )
-		.AllowedClass( ClassFilter )
-		.OnShouldFilterAsset( AssetFilter )
-		.AllowClear( bAllowClear )
 	];
 
 	// This avoids making duplicate reset boxes
@@ -101,30 +33,3 @@ void FSoftObjectPathCustomization::CustomizeHeader( TSharedRef<IPropertyHandle>
 void FSoftObjectPathCustomization::CustomizeChildren(TSharedRef<IPropertyHandle> InStructPropertyHandle, IDetailChildrenBuilder& StructBuilder, IPropertyTypeCustomizationUtils& StructCustomizationUtils)
 {
 }
-
-bool FSoftObjectPathCustomization::OnShouldFilterAsset( const FAssetData& InAssetData ) const
-{
-	// Only bound if we have classes to filter on, so we don't need to test for an empty array here
-	UClass* const AssetClass = InAssetData.GetClass();
-
-	for (auto It = DisallowedClassFilters.CreateConstIterator(); It; ++It)
-	{
-		UClass* const DisallowClass = *It;
-		if (AssetClass->IsChildOf(DisallowClass))
-		{
-			return true;
-		}
-	}
-
-	for(auto It = AllowedClassFilters.CreateConstIterator(); It; ++It)
-	{
-		UClass* const FilterClass = *It;
-		const bool bMatchesFilter = (bExactClass) ? AssetClass == FilterClass : AssetClass->IsChildOf(FilterClass);
-		if (bMatchesFilter)
-		{
-			return false;
-		}
-	}
-
-	return true;
-}
diff --git a/Engine/Source/Editor/DetailCustomizations/Private/SoftObjectPathCustomization.h b/Engine/Source/Editor/DetailCustomizations/Private/SoftObjectPathCustomization.h
--- a/Engine/Source/Editor/DetailCustomizations/Private/SoftObjectPathCustomization.h
+++ b/Engine/Source/Editor/DetailCustomizations/Private/SoftObjectPathCustomization.h
@@ -22,16 +22,7 @@ public:
 	virtual void CustomizeHeader( TSharedRef<class IPropertyHandle> InStructPropertyHandle, class FDetailWidgetRow& HeaderRow, IPropertyTypeCustomizationUtils& StructCustomizationUtils ) override;
 	virtual void CustomizeChildren( TSharedRef<class IPropertyHandle> InStructPropertyHandle, class IDetailChildrenBuilder& StructBuilder, IPropertyTypeCustomizationUtils& StructCustomizationUtils ) override;
 
-private:
-	bool OnShouldFilterAsset( const struct FAssetData& InAssetData ) const;
-
 private:
 	/** Handle to the struct property being customized */
 	TSharedPtr<IPropertyHandle> StructPropertyHandle;
-	/** Classes that can be used with this property */
-	TArray<UClass*> AllowedClassFilters;
-	/** Classes that can NOT be used with this property */
-	TArray<UClass*> DisallowedClassFilters;
-	/** Whether the classes in the above array must be an exact match, or whether they can also be a derived type; default is false */
-	bool bExactClass;
 };
diff --git a/Engine/Source/Editor/PropertyEditor/Private/UserInterface/PropertyEditor/SPropertyEditorAsset.cpp b/Engine/Source/Editor/PropertyEditor/Private/UserInterface/PropertyEditor/SPropertyEditorAsset.cpp
--- a/Engine/Source/Editor/PropertyEditor/Private/UserInterface/PropertyEditor/SPropertyEditorAsset.cpp
+++ b/Engine/Source/Editor/PropertyEditor/Private/UserInterface/PropertyEditor/SPropertyEditorAsset.cpp
@@ -35,191 +35,224 @@
 DECLARE_DELEGATE( FOnCopy );
 DECLARE_DELEGATE( FOnPaste );
 
-bool SPropertyEditorAsset::ShouldDisplayThumbnail( const FArguments& InArgs, const UClass* InObjectClass ) const
+// Helper to retrieve the correct property that has the applicable metadata.
+static const UProperty* GetActualMetadataProperty(const UProperty* Property)
 {
-	bool bDisplayThumbnail = InArgs._DisplayThumbnail && InArgs._ThumbnailPool.IsValid() && (!InObjectClass || !InObjectClass->IsChildOf(AActor::StaticClass()));
-	
-	if (InArgs._ThumbnailPool.IsValid())
+	if (UProperty* OuterProperty = Cast<UProperty>(Property->GetOuter()))
 	{
-		const UProperty* PropertyToCheck = nullptr;
-		if (PropertyEditor.IsValid())
-		{
-			PropertyToCheck = PropertyEditor->GetProperty();
-		}
-		else if (PropertyHandle.IsValid())
+		if (OuterProperty->IsA<UArrayProperty>()
+			|| OuterProperty->IsA<USetProperty>()
+			|| OuterProperty->IsA<UMapProperty>())
 		{
-			PropertyToCheck = PropertyHandle->GetProperty();
+			return OuterProperty;
 		}
+	}
+
+	return Property;
+}
+
+// Helper to support both meta=(TagName) and meta=(TagName=true) syntaxes
+static bool GetTagOrBoolMetadata(const UProperty* Property, const TCHAR* TagName, bool bDefault)
+{
+	bool bResult = bDefault;
+
+	if (Property->HasMetaData(TagName))
+	{
+		bResult = true;
 
-		if (PropertyToCheck != nullptr)
+		const FString ValueString = Property->GetMetaData(TagName);
+		if (!ValueString.IsEmpty())
 		{
-			UArrayProperty* ArrayParent = Cast<UArrayProperty>(PropertyToCheck->GetOuter());
-			USetProperty* SetParent = Cast<USetProperty>(PropertyToCheck->GetOuter());
-			UMapProperty* MapParent = Cast<UMapProperty>(PropertyToCheck->GetOuter());
-			if (ArrayParent != nullptr)
+			if (ValueString == TEXT("true"))
 			{
-				PropertyToCheck = ArrayParent;
+				bResult = true;
 			}
-			else if (SetParent != nullptr)
+			else if (ValueString == TEXT("false"))
 			{
-				PropertyToCheck = SetParent;
-			}
-			else if (MapParent != nullptr)
-			{
-				PropertyToCheck = MapParent;
-			}
-
-			if (PropertyToCheck->HasMetaData(TEXT("DisplayThumbnail")))
-			{
-				bDisplayThumbnail = PropertyToCheck->GetBoolMetaData(TEXT("DisplayThumbnail"));
+				bResult = false;
 			}
 		}
 	}
 
-	return bDisplayThumbnail;
+	return bResult;
 }
 
-void SPropertyEditorAsset::Construct( const FArguments& InArgs, const TSharedPtr<FPropertyEditor>& InPropertyEditor )
+bool SPropertyEditorAsset::ShouldDisplayThumbnail(const FArguments& InArgs, const UClass* InObjectClass) const
 {
-	PropertyEditor = InPropertyEditor;
-	PropertyHandle = InArgs._PropertyHandle;
-	OnSetObject = InArgs._OnSetObject;
-	OnShouldFilterAsset = InArgs._OnShouldFilterAsset;
-	bool DisplayCompactedSize = InArgs._DisplayCompactSize;
+	if (!InArgs._DisplayThumbnail || !InArgs._ThumbnailPool.IsValid())
+	{
+		return false;
+	}
 
-	UProperty* Property = nullptr;
-	if(PropertyEditor.IsValid())
+	bool bShowThumbnail = InObjectClass == nullptr || !InObjectClass->IsChildOf(AActor::StaticClass());
+	
+	// also check metadata for thumbnail & text display
+	const UProperty* PropertyToCheck = nullptr;
+	if (PropertyEditor.IsValid())
 	{
-		Property = PropertyEditor->GetPropertyNode()->GetProperty();
-		
-		bAllowClear = !(Property->PropertyFlags & CPF_NoClear);
-		ObjectClass = GetObjectPropertyClass(Property);
-		bIsActor = ObjectClass->IsChildOf( AActor::StaticClass() );
+		PropertyToCheck = PropertyEditor->GetProperty();
 	}
-	else
+	else if (PropertyHandle.IsValid())
 	{
-		bAllowClear = InArgs._AllowClear;
-		ObjectPath = InArgs._ObjectPath;
-		ObjectClass = InArgs._Class;
-		bIsActor = ObjectClass->IsChildOf( AActor::StaticClass() );
+		PropertyToCheck = PropertyHandle->GetProperty();
+	}
 
-		if (PropertyHandle.IsValid() && PropertyHandle->IsValidHandle())
-		{
-			Property = PropertyHandle->GetProperty();
-		}
-		else
-		{
-			AllowedClassFilters.Add(ObjectClass);
-		}
+	if (PropertyToCheck != nullptr)
+	{
+		PropertyToCheck = GetActualMetadataProperty(PropertyToCheck);
+
+		return GetTagOrBoolMetadata(PropertyToCheck, TEXT("DisplayThumbnail"), bShowThumbnail);
+	}
+
+	return bShowThumbnail;
+}
+
+void SPropertyEditorAsset::InitializeClassFilters(const UProperty* Property)
+{
+	if (Property == nullptr)
+	{
+		AllowedClassFilters.Add(ObjectClass);
+		return;
 	}
 
 	// Account for the allowed classes specified in the property metadata
-	if (Property)
+	const UProperty* MetadataProperty = GetActualMetadataProperty(Property);
+
+	bExactClass = GetTagOrBoolMetadata(MetadataProperty, TEXT("ExactClass"), false);
+
+	const FString AllowedClassesFilterString = MetadataProperty->GetMetaData(TEXT("AllowedClasses"));
+	if (!AllowedClassesFilterString.IsEmpty())
 	{
-		const FString* AllowedClassesFilterString;
-		if (UArrayProperty* ArrayParent = Cast<UArrayProperty>(Property->GetOuter()))
-		{
-			AllowedClassesFilterString = &ArrayParent->GetMetaData("AllowedClasses");
-		}
-		else
-		{
-			AllowedClassesFilterString = &Property->GetMetaData("AllowedClasses");
-		}
+		TArray<FString> AllowedClassFilterNames;
+		AllowedClassesFilterString.ParseIntoArray(AllowedClassFilterNames, TEXT(","), true);
 
-		if (AllowedClassesFilterString->IsEmpty())
-		{
-			AllowedClassFilters.Add(ObjectClass);
-		}
-		else
+		for (FString& ClassName : AllowedClassFilterNames)
 		{
-			TArray<FString> AllowedClassFilterNames;
-			AllowedClassesFilterString->ParseIntoArray(AllowedClassFilterNames, TEXT(","), true);
+			// User can potentially list class names with leading or trailing whitespace
+			ClassName.TrimStartAndEndInline();
 
-			for (auto It = AllowedClassFilterNames.CreateIterator(); It; ++It)
+			UClass* Class = FindObject<UClass>(ANY_PACKAGE, *ClassName);
+			if (!Class)
 			{
-				FString& ClassName = *It;
-				// User can potentially list class names with leading or trailing whitespace
-				ClassName.TrimStartAndEndInline();
-
-				UClass* Class = FindObject<UClass>(ANY_PACKAGE, *ClassName);
-
-				if (!Class)
-				{
-					Class = LoadObject<UClass>(nullptr, *ClassName);
-				}
+				Class = LoadObject<UClass>(nullptr, *ClassName);
+			}
 
-				if (Class)
+			if (Class)
+			{
+				// If the class is an interface, expand it to be all classes in memory that implement the class.
+				if (Class->HasAnyClassFlags(CLASS_Interface))
 				{
-					// If the class is an interface, expand it to be all classes in memory that implement the class.
-					if (Class->HasAnyClassFlags(CLASS_Interface))
+					for (TObjectIterator<UClass> ClassIt; ClassIt; ++ClassIt)
 					{
-						for (TObjectIterator<UClass> ClassIt; ClassIt; ++ClassIt)
+						UClass* const ClassWithInterface = (*ClassIt);
+						if (ClassWithInterface->ImplementsInterface(Class))
 						{
-							UClass* const ClassWithInterface = (*ClassIt);
-							if (ClassWithInterface->ImplementsInterface(Class))
-							{
-								AllowedClassFilters.Add(ClassWithInterface);
-							}
+							AllowedClassFilters.Add(ClassWithInterface);
 						}
 					}
-					else
-					{
-						AllowedClassFilters.Add(Class);
-					}
+				}
+				else
+				{
+					AllowedClassFilters.Add(Class);
 				}
 			}
 		}
+	}
+	else
+	{
+		// always add the object class to the filters
+		AllowedClassFilters.Add(ObjectClass);
+	}
 
-		const FString* DisallowedClassesFilterString;
-		if (UArrayProperty* ArrayParent = Cast<UArrayProperty>(Property->GetOuter()))
-		{
-			DisallowedClassesFilterString = &ArrayParent->GetMetaData("DisallowedClasses");
-		}
-		else
-		{
-			DisallowedClassesFilterString = &Property->GetMetaData("DisallowedClasses");
-		}
+	const FString DisallowedClassesFilterString = MetadataProperty->GetMetaData(TEXT("DisallowedClasses"));
+	if (!DisallowedClassesFilterString.IsEmpty())
+	{
+		TArray<FString> DisallowedClassFilterNames;
+		DisallowedClassesFilterString.ParseIntoArray(DisallowedClassFilterNames, TEXT(","), true);
 
-		if (!DisallowedClassesFilterString->IsEmpty())
+		for (FString& ClassName : DisallowedClassFilterNames)
 		{
-			TArray<FString> DisallowedClassFilterNames;
-			DisallowedClassesFilterString->ParseIntoArray(DisallowedClassFilterNames, TEXT(","), true);
+			// User can potentially list class names with leading or trailing whitespace
+			ClassName.TrimStartAndEndInline();
 
-			for (auto It = DisallowedClassFilterNames.CreateIterator(); It; ++It)
+			UClass* Class = FindObject<UClass>(ANY_PACKAGE, *ClassName);
+			if (!Class)
 			{
-				FString& ClassName = *It;
-				// User can potentially list class names with leading or trailing whitespace
-				ClassName.TrimStartAndEndInline();
-
-				UClass* Class = FindObject<UClass>(ANY_PACKAGE, *ClassName);
-
-				if (!Class)
-				{
-					Class = LoadObject<UClass>(nullptr, *ClassName);
-				}
+				Class = LoadObject<UClass>(nullptr, *ClassName);
+			}
 
-				if (Class)
+			if (Class)
+			{
+				// If the class is an interface, expand it to be all classes in memory that implement the class.
+				if (Class->HasAnyClassFlags(CLASS_Interface))
 				{
-					// If the class is an interface, expand it to be all classes in memory that implement the class.
-					if (Class->HasAnyClassFlags(CLASS_Interface))
+					for (TObjectIterator<UClass> ClassIt; ClassIt; ++ClassIt)
 					{
-						for (TObjectIterator<UClass> ClassIt; ClassIt; ++ClassIt)
+						UClass* const ClassWithInterface = (*ClassIt);
+						if (ClassWithInterface->ImplementsInterface(Class))
 						{
-							UClass* const ClassWithInterface = (*ClassIt);
-							if (ClassWithInterface->ImplementsInterface(Class))
-							{
-								DisallowedClassFilters.Add(ClassWithInterface);
-							}
+							DisallowedClassFilters.Add(ClassWithInterface);
 						}
 					}
-					else
-					{
-						DisallowedClassFilters.Add(Class);
-					}
+				}
+				else
+				{
+					DisallowedClassFilters.Add(Class);
 				}
 			}
 		}
 	}
+}
+
+// Awful hack to deal with UClass::FindCommonBase taking an array of non-const classes...
+static TArray<UClass*> ConstCastClassArray(TArray<const UClass*>& Classes)
+{
+	TArray<UClass*> Result;
+	for (const UClass* Class : Classes)
+	{
+		Result.Add(const_cast<UClass*>(Class));
+	}
+
+	return Result;
+}
+
+void SPropertyEditorAsset::Construct(const FArguments& InArgs, const TSharedPtr<FPropertyEditor>& InPropertyEditor)
+{
+	PropertyEditor = InPropertyEditor;
+	PropertyHandle = InArgs._PropertyHandle;
+	OnSetObject = InArgs._OnSetObject;
+	OnShouldFilterAsset = InArgs._OnShouldFilterAsset;
+	ObjectPath = InArgs._ObjectPath;
+
+	UProperty* Property = nullptr;
+	if (PropertyEditor.IsValid())
+	{
+		Property = PropertyEditor->GetPropertyNode()->GetProperty();
+	}
+	else if (PropertyHandle.IsValid() && PropertyHandle->IsValidHandle())
+	{
+		Property = PropertyHandle->GetProperty();
+	}
+
+	ObjectClass = InArgs._Class != nullptr ? InArgs._Class : GetObjectPropertyClass(Property);
+	bAllowClear = InArgs._AllowClear.IsSet() ? InArgs._AllowClear.GetValue() : !(Property->PropertyFlags & CPF_NoClear);
+
+	InitializeClassFilters(Property);
+
+	ensure(AllowedClassFilters.Num() > 0);
+
+	// Make the ObjectClass more specific if we only have one class filter
+	// eg. if ObjectClass was set to Actor, but AllowedClasses="PointLight", we can limit it to PointLight immediately
+	if (AllowedClassFilters.Num() == 1 && DisallowedClassFilters.Num() == 0)
+	{
+		ObjectClass = const_cast<UClass*>(AllowedClassFilters[0]);
+	}
+	else
+	{
+		ObjectClass = UClass::FindCommonBase(ConstCastClassArray(AllowedClassFilters));
+	}
+
+	bIsActor = ObjectClass->IsChildOf(AActor::StaticClass());
 
 	if (InArgs._NewAssetFactories.IsSet())
 	{
@@ -247,17 +280,7 @@ void SPropertyEditorAsset::Construct( const FArguments& InArgs, const TSharedPtr
 
 	if (Property)
 	{
-		UProperty* PropToConsider = Property;
-		if (UProperty* OuterProperty = Cast<UProperty>(PropToConsider->GetOuter()))
-		{
-			if (OuterProperty->IsA<UArrayProperty>()
-				|| OuterProperty->IsA<USetProperty>()
-				|| OuterProperty->IsA<UMapProperty>())
-			{
-				PropToConsider = OuterProperty;
-			}
-		}
-
+		const UProperty* PropToConsider = GetActualMetadataProperty(Property);
 		if (PropToConsider->HasAnyPropertyFlags(CPF_EditConst | CPF_DisableEditOnTemplate))
 		{
 			// There are some cases where editing an Actor Property is not allowed, such as when it is contained within a struct or a CDO
@@ -289,6 +312,7 @@ void SPropertyEditorAsset::Construct( const FArguments& InArgs, const TSharedPtr
 			}
 		}
 	}
+
 	bool bOldEnableAttribute = IsEnabledAttribute.Get();
 	if (bOldEnableAttribute && !InArgs._EnableContentPicker)
 	{
@@ -382,7 +406,7 @@ void SPropertyEditorAsset::Construct( const FArguments& InArgs, const TSharedPtr
 			]
 		];
 
-		if(DisplayCompactedSize)
+		if(InArgs._DisplayCompactSize)
 		{
 			ValueContentBox->AddSlot()
 			[
@@ -690,32 +714,7 @@ void SPropertyEditorAsset::OnMenuOpenChanged(bool bOpen)
 
 bool SPropertyEditorAsset::IsFilteredActor( const AActor* const Actor ) const
 {
-	bool IsAllowed = Actor->IsA(ObjectClass) && !Actor->IsChildActor();
-
-	if (IsAllowed)
-	{
-		bool ClassFilterAllowed = false;
-		for (const UClass* Class : AllowedClassFilters)
-		{
-			if (Actor->GetClass()->IsChildOf(Class))
-			{
-				ClassFilterAllowed = true;
-				break;
-			}
-		}
-
-		for (const UClass* Class : DisallowedClassFilters)
-		{
-			if (Actor->GetClass()->IsChildOf(Class))
-			{
-				ClassFilterAllowed = false;
-				break;
-			}
-		}
-
-		IsAllowed = ClassFilterAllowed;
-	}
-
+	bool IsAllowed = Actor->IsA(ObjectClass) && !Actor->IsChildActor() && IsClassAllowed(Actor->GetClass());
 	return IsAllowed;
 }
 
@@ -1260,45 +1259,47 @@ bool SPropertyEditorAsset::CanEdit() const
 
 bool SPropertyEditorAsset::CanSetBasedOnCustomClasses( const FAssetData& InAssetData ) const
 {
-	bool bAllowedToSetBasedOnFilter = true;
-	if( InAssetData.IsValid() && AllowedClassFilters.Num() > 0 )
+	if (InAssetData.IsValid())
 	{
-		const int32 AllowedClassCount = AllowedClassFilters.Num();
-		const int32 DisallowedClassCount = DisallowedClassFilters.Num();
+		return IsClassAllowed(InAssetData.GetClass());
+	}
+
+	return true;
+}
 
-		if (AllowedClassCount > 0 || DisallowedClassCount > 0)
+bool SPropertyEditorAsset::IsClassAllowed(const UClass* InClass) const
+{
+	bool bClassAllowed = true;
+	if (AllowedClassFilters.Num() > 0)
+	{
+		bClassAllowed = false;
+		for (const UClass* AllowedClass : AllowedClassFilters)
 		{
-			UClass* AssetClass = InAssetData.GetClass();
-			if (AllowedClassCount > 0)
+			const bool bAllowedClassIsInterface = AllowedClass->HasAnyClassFlags(CLASS_Interface);
+			bClassAllowed = bExactClass ? InClass == AllowedClass :
+				InClass->IsChildOf(AllowedClass) || (bAllowedClassIsInterface && InClass->ImplementsInterface(AllowedClass));
+
+			if (bClassAllowed)
 			{
-				bAllowedToSetBasedOnFilter = false; 
-				for (const UClass* AllowedClass : AllowedClassFilters)
-				{
-					const bool bAllowedClassIsInterface = AllowedClass->HasAnyClassFlags(CLASS_Interface);
-					if (AssetClass->IsChildOf(AllowedClass) || (bAllowedClassIsInterface && AssetClass->ImplementsInterface(AllowedClass)))
-					{
-						bAllowedToSetBasedOnFilter = true;
-						break;
-					}
-				}
+				break;
 			}
+		}
+	}
 
-			if (DisallowedClassCount > 0 && bAllowedToSetBasedOnFilter)
+	if (DisallowedClassFilters.Num() > 0 && bClassAllowed)
+	{
+		for (const UClass* DisallowedClass : DisallowedClassFilters)
+		{
+			const bool bDisallowedClassIsInterface = DisallowedClass->HasAnyClassFlags(CLASS_Interface);
+			if (InClass->IsChildOf(DisallowedClass) || (bDisallowedClassIsInterface && InClass->ImplementsInterface(DisallowedClass)))
 			{
-				for (const UClass* DisallowedClass : DisallowedClassFilters)
-				{
-					const bool bDisallowedClassIsInterface = DisallowedClass->HasAnyClassFlags(CLASS_Interface);
-					if (AssetClass->IsChildOf(DisallowedClass) || (bDisallowedClassIsInterface && AssetClass->ImplementsInterface(DisallowedClass)))
-					{
-						bAllowedToSetBasedOnFilter = false;
-						break;
-					}
-				}
+				bClassAllowed = false;
+				break;
 			}
 		}
 	}
 
-	return bAllowedToSetBasedOnFilter;
+	return bClassAllowed;
 }
 
 bool SPropertyEditorAsset::CanSetBasedOnAssetReferenceFilter( const FAssetData& InAssetData, FText* OutOptionalFailureReason) const
diff --git a/Engine/Source/Editor/PropertyEditor/Private/UserInterface/PropertyEditor/SPropertyEditorAsset.h b/Engine/Source/Editor/PropertyEditor/Private/UserInterface/PropertyEditor/SPropertyEditorAsset.h
--- a/Engine/Source/Editor/PropertyEditor/Private/UserInterface/PropertyEditor/SPropertyEditorAsset.h
+++ b/Engine/Source/Editor/PropertyEditor/Private/UserInterface/PropertyEditor/SPropertyEditorAsset.h
@@ -46,22 +46,21 @@ public:
 	SLATE_BEGIN_ARGS( SPropertyEditorAsset )
 		: _AssetFont( FEditorStyle::GetFontStyle("PropertyEditor.AssetName.Font") ) 
 		, _ClassFont( FEditorStyle::GetFontStyle("PropertyEditor.AssetClass.Font") ) 
-		, _AllowClear(true)
 		, _DisplayThumbnail(true)
 		, _DisplayUseSelected(true)
 		, _DisplayBrowse(true)
 		, _EnableContentPicker(true)
 		, _DisplayCompactSize(false)
-		, _ThumbnailPool(NULL)
-		, _ThumbnailSize( FIntPoint(64, 64) )
+		, _ThumbnailPool(nullptr)
+		, _ThumbnailSize(FIntPoint(64, 64))
 		, _ObjectPath()
-		, _Class(NULL)
+		, _Class(nullptr)
 		, _CustomContentSlot()
 		, _ResetToDefaultSlot()
 		{}
 		SLATE_ATTRIBUTE( FSlateFontInfo, AssetFont )
 		SLATE_ATTRIBUTE( FSlateFontInfo, ClassFont )
-		SLATE_ARGUMENT( bool, AllowClear )
+		SLATE_ARGUMENT( TOptional<bool>, AllowClear )
 		SLATE_ARGUMENT( bool, DisplayThumbnail )
 		SLATE_ARGUMENT( bool, DisplayUseSelected )
 		SLATE_ARGUMENT( bool, DisplayBrowse )
@@ -318,6 +317,15 @@ private:
 	 * @return	The class of the property. Asserts if this tries to return null.
 	 */
 	static UClass* GetObjectPropertyClass(const UProperty* Property);
+
+	/**
+	 * Initialize the AllowedClasses and DisallowedClasses arrays based on the property metadata.
+	 */
+	void InitializeClassFilters(const UProperty* Property);
+
+	/** Check if the given class is allowed by the AllowedClasses and DisallowedClasses as defined in the property metadata. */
+	bool IsClassAllowed(const UClass* InClass) const;
+
 private:
 
 	/** Main combobutton */
@@ -353,6 +361,9 @@ private:
 	/** Whether the object we are editing is an Actor (i.e. requires a Scene Outliner to be displayed) */
 	bool bIsActor;
 
+	/** Whether the classes in the AllowedClassFilters list are exact or whether valid classes can be derived. */
+	bool bExactClass;
+
 	/** Delegate to call when our object value is set */
 	FOnSetObject OnSetObject;
 
diff --git a/Engine/Source/Editor/UnrealEd/Classes/Editor/PropertyEditorTestObject.h b/Engine/Source/Editor/UnrealEd/Classes/Editor/PropertyEditorTestObject.h
--- a/Engine/Source/Editor/UnrealEd/Classes/Editor/PropertyEditorTestObject.h
+++ b/Engine/Source/Editor/UnrealEd/Classes/Editor/PropertyEditorTestObject.h
@@ -9,6 +9,9 @@
 #include "UObject/SoftObjectPath.h"
 #include "Engine/BlendableInterface.h"
 #include "Engine/EngineTypes.h"
+
+#include "Curves/RichCurve.h"
+
 #include "PropertyEditorTestObject.generated.h"
 
 class AActor;
@@ -125,7 +128,7 @@ class UPropertyEditorTestObject : public UObject
 {
     GENERATED_BODY()
 
-	UPROPERTY(EditAnywhere, Category = BasicProperties)
+/*	UPROPERTY(EditAnywhere, Category = BasicProperties)
 	int8 Int8Property;
 
 	UPROPERTY(EditAnywhere, Category = BasicProperties)
@@ -300,7 +303,10 @@ class UPropertyEditorTestObject : public UObject
 	bool bEditConditionStructWithMultipleInstances2;
 
 	UPROPERTY(EditAnywhere, Category=StructTests, meta=(editcondition = "bEditConditionStructWithMultipleInstances2"))
-	FPropertyEditorTestBasicStruct StructWithMultipleInstances2;
+	FPropertyEditorTestBasicStruct StructWithMultipleInstances2; */
+
+	UPROPERTY(EditAnywhere, Category=StructTests)
+	FRichCurve RichCurve;
 
 	UPROPERTY(EditAnywhere, Category=StructTests)
 	FSoftObjectPath AssetReferenceCustomStruct;
@@ -308,7 +314,25 @@ class UPropertyEditorTestObject : public UObject
 	UPROPERTY(EditAnywhere, Category=StructTests, meta=(DisplayThumbnail = "true"))
 	FSoftObjectPath AssetReferenceCustomStructWithThumbnail;
 
-	UPROPERTY(EditAnywhere, Category=StructTests, meta=(InlineEditConditionToggle))
+	UPROPERTY(EditAnywhere, Category=StructTests, meta=(AllowedClasses="PointLight", ExactClass))
+	FSoftObjectPath ExactlyPointLightActorReference;
+
+	UPROPERTY(EditAnywhere, Category=StructTests, meta=(AllowedClasses="Light"))
+	FSoftObjectPath LightActorReference;
+
+	UPROPERTY(EditAnywhere, Category=StructTests, meta=(AllowedClasses="PointLight, SpotLight", ExactClass=true))
+	FSoftObjectPath ExactPointOrSpotLightActorReference;
+
+	UPROPERTY(EditAnywhere, Category=StructTests, meta=(AllowedClasses="Light ,StaticMeshActor", DisplayThumbnail))
+	FSoftObjectPath LightOrStaticMeshActorReference;
+
+	UPROPERTY(EditAnywhere, Category=StructTests, meta=(AllowedClasses="Actor", DisallowedClasses="Light"))
+	FSoftObjectPath NotLightActorReference;
+
+	UPROPERTY(EditAnywhere, Category=StructTests, meta=(AllowedClasses="Material,Texture"))
+	FSoftObjectPath MaterialOrTextureAssetReference;
+
+	/*UPROPERTY(EditAnywhere, Category=StructTests, meta=(InlineEditConditionToggle))
 	bool bEditCondition;
 
 	UPROPERTY(EditAnywhere, Category=AdvancedProperties, meta=(editcondition = "bEditCondition"))
@@ -495,5 +519,5 @@ class UPropertyEditorTestObject : public UObject
 	bool bVisibleWhenStatic;
 
 	UPROPERTY(EditAnywhere, Category = EditConditionHides, meta = (EditConditionHides, EditCondition = "PropertyThatHides == EComponentMobility::Stationary"))
-	int32 VisibleWhenStationary;
+	int32 VisibleWhenStationary;*/
 };
