You can now populate arrays in the details view by dragging onto them.

diff --git a/Engine/Source/Editor/PropertyEditor/Private/UserInterface/PropertyEditor/SPropertyEditorArray.h b/Engine/Source/Editor/PropertyEditor/Private/UserInterface/PropertyEditor/SPropertyEditorArray.h
--- a/Engine/Source/Editor/PropertyEditor/Private/UserInterface/PropertyEditor/SPropertyEditorArray.h
+++ b/Engine/Source/Editor/PropertyEditor/Private/UserInterface/PropertyEditor/SPropertyEditorArray.h
@@ -110,7 +110,7 @@ private:
 				for (FAssetData AssetData : DragDropOp->GetAssets())
 				{
 					// if the type matches
-					if (ObjectProperty->PropertyClass == AssetData.GetClass())
+					if (AssetData.GetClass()->IsChildOf(ObjectProperty->PropertyClass))
 					{
 						PropertyEditor->AddGivenItem(AssetData.ObjectPath.ToString());
 					}
@@ -157,8 +157,8 @@ private:
 			{
 				for (FAssetData AssetData : DragDropOp->GetAssets())
 				{
-					// if the type matches
-					if (ObjectProperty->PropertyClass != AssetData.GetClass())
+					// if the type does not match
+					if (!AssetData.GetClass()->IsChildOf(ObjectProperty->PropertyClass))
 					{
 						bHasOnlyValidElements = false;
 						break;
