Fixed static analysis warning.

diff --git a/Engine/Source/Editor/PropertyEditor/Private/UserInterface/PropertyEditor/SPropertyEditorAsset.cpp b/Engine/Source/Editor/PropertyEditor/Private/UserInterface/PropertyEditor/SPropertyEditorAsset.cpp
--- a/Engine/Source/Editor/PropertyEditor/Private/UserInterface/PropertyEditor/SPropertyEditorAsset.cpp
+++ b/Engine/Source/Editor/PropertyEditor/Private/UserInterface/PropertyEditor/SPropertyEditorAsset.cpp
@@ -235,7 +235,7 @@ void SPropertyEditorAsset::Construct(const FArguments& InArgs, const TSharedPtr<
 	}
 
 	ObjectClass = InArgs._Class != nullptr ? InArgs._Class : GetObjectPropertyClass(Property);
-	bAllowClear = InArgs._AllowClear.IsSet() ? InArgs._AllowClear.GetValue() : !(Property->PropertyFlags & CPF_NoClear);
+	bAllowClear = InArgs._AllowClear.IsSet() ? InArgs._AllowClear.GetValue() : (Property ? !(Property->PropertyFlags & CPF_NoClear) : true);
 
 	InitializeClassFilters(Property);
 
