Don't display an error about duplicate keys/elements in maps/sets when dragging the slider for a numeric key/value.

diff --git a/Engine/Source/Editor/PropertyEditor/Private/PropertyHandleImpl.cpp b/Engine/Source/Editor/PropertyEditor/Private/PropertyHandleImpl.cpp
--- a/Engine/Source/Editor/PropertyEditor/Private/PropertyHandleImpl.cpp
+++ b/Engine/Source/Editor/PropertyEditor/Private/PropertyHandleImpl.cpp
@@ -453,7 +453,8 @@ FPropertyAccess::Result FPropertyValueImpl::ImportText( const TArray<FObjectBase
 					{
 						FScriptSetHelper SetHelper(SetProp, ParentNode->GetValueBaseAddress((uint8*)Cur.Object));
 
-						if ( SetHelper.HasElement(Cur.BaseAddress, InValues[ObjectIndex]) )
+						if (SetHelper.HasElement(Cur.BaseAddress, InValues[ObjectIndex]) &&
+							(Flags & EPropertyValueSetFlags::InteractiveChange) == 0)
 						{
 							// Duplicate element in the set
 							ShowInvalidOperationError(LOCTEXT("DuplicateSetElement", "Duplicate elements are not allowed in Set properties."));
@@ -469,8 +470,8 @@ FPropertyAccess::Result FPropertyValueImpl::ImportText( const TArray<FObjectBase
 					if (bIsInContainer)
 					{
 						FScriptMapHelper MapHelper(MapProperty, ParentNode->GetValueBaseAddress((uint8*)Cur.Object));
-
-						if ( MapHelper.HasKey(Cur.BaseAddress, InValues[ObjectIndex]) )
+						if (MapHelper.HasKey(Cur.BaseAddress, InValues[ObjectIndex]) && 
+							(Flags & EPropertyValueSetFlags::InteractiveChange) == 0)
 						{
 							// Duplicate key in the map
 							ShowInvalidOperationError(LOCTEXT("DuplicateMapKey", "Duplicate keys are not allowed in Map properties."));
