Maps with struct values now correctly display their key in a separate property.

diff --git a/Engine/Source/Editor/PropertyEditor/Private/DetailPropertyRow.cpp b/Engine/Source/Editor/PropertyEditor/Private/DetailPropertyRow.cpp
--- a/Engine/Source/Editor/PropertyEditor/Private/DetailPropertyRow.cpp
+++ b/Engine/Source/Editor/PropertyEditor/Private/DetailPropertyRow.cpp
@@ -52,12 +52,8 @@ FDetailPropertyRow::FDetailPropertyRow(TSharedPtr<FPropertyNode> InPropertyNode,
 
 		if (PropertyNode->GetPropertyKeyNode().IsValid())
 		{
-			UStructProperty* KeyStructProp = Cast<UStructProperty>(PropertyNode->GetPropertyKeyNode()->GetProperty());
-
 			// Only struct and customized properties require their own nodes. Everything else just needs a property editor.
-			bool bNeedsKeyPropEditor = KeyStructProp == nullptr && !GetPropertyCustomization(PropertyNode->GetPropertyKeyNode().ToSharedRef(), InParentCategory).IsValid();
-
-			if ( bNeedsKeyPropEditor )
+			if (!NeedsKeyNode(PropertyNodeRef, InParentCategory))
 			{
 				MakePropertyEditor(PropertyNode->GetPropertyKeyNode().ToSharedRef(), Utilities, PropertyKeyEditor);
 			}
@@ -65,6 +61,16 @@ FDetailPropertyRow::FDetailPropertyRow(TSharedPtr<FPropertyNode> InPropertyNode,
 	}
 }
 
+bool FDetailPropertyRow::NeedsKeyNode(TSharedRef<FPropertyNode> InPropertyNode, TSharedRef<FDetailCategoryImpl> InParentCategory)
+{
+	UStructProperty* KeyStructProp = Cast<UStructProperty>(InPropertyNode->GetPropertyKeyNode()->GetProperty());
+	UStructProperty* ValueStructProp = Cast<UStructProperty>(InPropertyNode->GetProperty());
+
+	return KeyStructProp != nullptr || ValueStructProp != nullptr ||
+		GetPropertyCustomization(InPropertyNode->GetPropertyKeyNode().ToSharedRef(), InParentCategory).IsValid() ||
+		GetPropertyCustomization(InPropertyNode, InParentCategory).IsValid();
+}
+
 IDetailPropertyRow& FDetailPropertyRow::DisplayName( const FText& InDisplayName )
 {
 	if (PropertyNode.IsValid())
@@ -600,7 +606,7 @@ void FDetailPropertyRow::MakeNameOrKeyWidget( FDetailWidgetRow& Row, const TShar
 	EHorizontalAlignment HorizontalAlignment = HAlign_Fill;
 
 	// We will only use key widgets for non-struct keys
-	const bool bHasKeyNode =  PropertyKeyEditor.IsValid() && !PropertyHandle->HasMetaData(TEXT("ReadOnlyKeys"));
+	const bool bHasKeyNode = PropertyKeyEditor.IsValid() && !PropertyHandle->HasMetaData(TEXT("ReadOnlyKeys"));
 
 	if( !bHasKeyNode && InCustomRow.IsValid() )
 	{
diff --git a/Engine/Source/Editor/PropertyEditor/Private/DetailPropertyRow.h b/Engine/Source/Editor/PropertyEditor/Private/DetailPropertyRow.h
--- a/Engine/Source/Editor/PropertyEditor/Private/DetailPropertyRow.h
+++ b/Engine/Source/Editor/PropertyEditor/Private/DetailPropertyRow.h
@@ -162,6 +162,9 @@ private:
 	 */
 	static TSharedPtr<IPropertyTypeCustomization> GetPropertyCustomization( const TSharedRef<FPropertyNode>& InPropertyNode, const TSharedRef<FDetailCategoryImpl>& InParentCategory );
 
+	/** Does the given property node require a key node to be created, or is a key property editor sufficient? */
+	static bool NeedsKeyNode(TSharedRef<FPropertyNode> InPropertyNode, TSharedRef<FDetailCategoryImpl> InParentCategory);
+
 private:
 	/** User driven enabled state */
 	TAttribute<bool> CustomIsEnabledAttrib;
