Removed duplicate call of CreateClassFilter(). Looks like a mismerge in CL 7120528.

diff --git a/Engine/Source/Editor/PropertyEditor/Private/UserInterface/PropertyEditor/SPropertyEditorClass.cpp b/Engine/Source/Editor/PropertyEditor/Private/UserInterface/PropertyEditor/SPropertyEditorClass.cpp
--- a/Engine/Source/Editor/PropertyEditor/Private/UserInterface/PropertyEditor/SPropertyEditorClass.cpp
+++ b/Engine/Source/Editor/PropertyEditor/Private/UserInterface/PropertyEditor/SPropertyEditorClass.cpp
@@ -122,7 +122,6 @@ void SPropertyEditorClass::Construct(const FArguments& InArgs, const TSharedPtr<
 
 		auto FillClassFilters = [&](TArray<const UClass*> &ClassFilters, FName TagName)
 		{
-
 			const FString* ClassesFilterString = &Property->GetMetaData(TagName);
 
 			ClassFilters.Empty();
@@ -191,8 +190,6 @@ void SPropertyEditorClass::Construct(const FArguments& InArgs, const TSharedPtr<
 		OnSetClass = InArgs._OnSetClass;
 	}
 
-	CreateClassFilter();
-	
 	CreateClassFilter();
 
 	SAssignNew(ComboButton, SComboButton)
