Fixed crash when modifying the CDO of a child class that has added a map or set that does not exist in the parent.

This was because FScriptSetHelper and FScriptMapHelper's Num() function ends up accessing the value pointer of the property node, which can be null at this point. Moving the range check after the null check fixes the issue.

Reported through UDN.

diff --git a/Engine/Source/Editor/PropertyEditor/Private/PropertyNode.cpp b/Engine/Source/Editor/PropertyEditor/Private/PropertyNode.cpp
--- a/Engine/Source/Editor/PropertyEditor/Private/PropertyNode.cpp
+++ b/Engine/Source/Editor/PropertyEditor/Private/PropertyNode.cpp
@@ -1575,9 +1575,7 @@ bool FPropertyNode::GetDiffersFromDefaultForObject( FPropertyItemValueDataTracke
 		else if (OuterSetProperty != NULL)
 		{
 			FScriptSetHelper SetHelper(OuterSetProperty, ValueTracker.GetPropertyDefaultBaseAddress());
-
-			bool bIsValidIndex = ArrayIndex >= 0 && ArrayIndex < SetHelper.Num();
-			if ( ValueTracker.GetPropertyDefaultBaseAddress() != NULL && !bIsValidIndex)
+			if ( ValueTracker.GetPropertyDefaultBaseAddress() != NULL && (ArrayIndex < 0 || ArrayIndex >= SetHelper.Num()))
 			{
 				bDiffersFromDefaultForObject = true;
 			}
@@ -1585,9 +1583,7 @@ bool FPropertyNode::GetDiffersFromDefaultForObject( FPropertyItemValueDataTracke
 		else if (OuterMapProperty != NULL)
 		{
 			FScriptMapHelper MapHelper(OuterMapProperty, ValueTracker.GetPropertyDefaultBaseAddress());
-
-			bool bIsValidIndex = ArrayIndex >= 0 && ArrayIndex < MapHelper.Num();
-			if ( ValueTracker.GetPropertyDefaultBaseAddress() != NULL && !bIsValidIndex )
+			if ( ValueTracker.GetPropertyDefaultBaseAddress() != NULL && (ArrayIndex < 0 || ArrayIndex >= MapHelper.Num()))
 			{
 				bDiffersFromDefaultForObject = true;
 			}
