TSubclassOf widgets now support AllowedClasses in their meta specifiers.

For TSubclassOf, QualifiedObject was a UClass*, which caused the check to fail.

diff --git a/Engine/Source/Editor/PropertyEditor/Private/PropertyHandleImpl.cpp b/Engine/Source/Editor/PropertyEditor/Private/PropertyHandleImpl.cpp
--- a/Engine/Source/Editor/PropertyEditor/Private/PropertyHandleImpl.cpp
+++ b/Engine/Source/Editor/PropertyEditor/Private/PropertyHandleImpl.cpp
@@ -3778,7 +3778,7 @@ FPropertyAccess::Result FPropertyHandleObject::SetValueFromFormattedString(const
 					const bool bIsInterface = AllowedClass && AllowedClass->HasAnyClassFlags(CLASS_Interface);
 
 					// Check if the object is an allowed class type this property supports
-					if ((AllowedClass && QualifiedObject->IsA(AllowedClass)) || (bIsInterface && QualifiedObject->GetClass()->ImplementsInterface(AllowedClass)))
+					if ((AllowedClass && QualifiedClass->IsChildOf(AllowedClass)) || (bIsInterface && QualifiedObject->GetClass()->ImplementsInterface(AllowedClass)))
 					{
 						bSupportedObject = true;
 						break;
@@ -3801,7 +3801,7 @@ FPropertyAccess::Result FPropertyHandleObject::SetValueFromFormattedString(const
 					UClass* DisallowedClass = FindObject<UClass>(ANY_PACKAGE, *DisallowedClassName);
 					const bool bIsInterface = DisallowedClass && DisallowedClass->HasAnyClassFlags(CLASS_Interface);
 
-					if ((DisallowedClass && QualifiedObject->IsA(DisallowedClass)) || (bIsInterface && QualifiedObject->GetClass()->ImplementsInterface(DisallowedClass)))
+					if ((DisallowedClass && QualifiedClass->IsChildOf(DisallowedClass)) || (bIsInterface && QualifiedObject->GetClass()->ImplementsInterface(DisallowedClass)))
 					{
 						bSupportedObject = false;
 						break;
