PropertyEditor: Add `GetDesiredWidth` method that forwards whatever the inner property editor widget says.

diff --git a/Engine/Source/Editor/PropertyEditor/Private/PropertyCustomizationHelpers.cpp b/Engine/Source/Editor/PropertyEditor/Private/PropertyCustomizationHelpers.cpp
--- a/Engine/Source/Editor/PropertyEditor/Private/PropertyCustomizationHelpers.cpp
+++ b/Engine/Source/Editor/PropertyEditor/Private/PropertyCustomizationHelpers.cpp
@@ -539,6 +539,12 @@ void SObjectPropertyEntryBox::Construct( const FArguments& InArgs )
 	];
 }
 
+void SObjectPropertyEntryBox::GetDesiredWidth(float& OutMinDesiredWidth, float &OutMaxDesiredWidth)
+{
+	checkf(PropertyEditorAsset.IsValid(), TEXT("SObjectPropertyEntryBox hasn't been constructed yet."));
+	PropertyEditorAsset->GetDesiredWidth(OutMinDesiredWidth, OutMaxDesiredWidth);
+}
+
 FString SObjectPropertyEntryBox::OnGetObjectPath() const
 {
 	FString StringReference;
diff --git a/Engine/Source/Editor/PropertyEditor/Public/PropertyCustomizationHelpers.h b/Engine/Source/Editor/PropertyEditor/Public/PropertyCustomizationHelpers.h
--- a/Engine/Source/Editor/PropertyEditor/Public/PropertyCustomizationHelpers.h
+++ b/Engine/Source/Editor/PropertyEditor/Public/PropertyCustomizationHelpers.h
@@ -169,6 +169,8 @@ public:
 
 	PROPERTYEDITOR_API void Construct( const FArguments& InArgs );
 
+	PROPERTYEDITOR_API void GetDesiredWidth(float& OutMinDesiredWidth, float &OutMaxDesiredWidth);
+
 private:
 	/**
 	 * Delegate function called when an object is changed
