PropertyEditor: Also check for instance metadata when figuring out if a "reset" button should be added.

diff --git a/Engine/Source/Editor/PropertyEditor/Private/DetailPropertyRow.cpp b/Engine/Source/Editor/PropertyEditor/Private/DetailPropertyRow.cpp
--- a/Engine/Source/Editor/PropertyEditor/Private/DetailPropertyRow.cpp
+++ b/Engine/Source/Editor/PropertyEditor/Private/DetailPropertyRow.cpp
@@ -714,7 +714,7 @@ void FDetailPropertyRow::MakeValueWidget( FDetailWidgetRow& Row, const TSharedPt
 
 	TSharedPtr<SResetToDefaultPropertyEditor> ResetButton = nullptr;
 	TSharedPtr<SWidget> ResetWidget = nullptr;
-	if (!PropertyHandle->HasMetaData(TEXT("NoResetToDefault")))
+	if (!PropertyHandle->HasMetaData(TEXT("NoResetToDefault")) && !PropertyHandle->GetInstanceMetaData(TEXT("NoResetToDefault")))
 	{
 		if (PropertyHandle->IsResetToDefaultCustomized())
 		{
