Fixed filtering of classes when AllowedClasses or DisallowedClasses is specified for a UClass* property. Previously, unloaded blueprint classes were always displayed. In addition, the classes were being compared exactly instead of using IsChildOf.

diff --git a/Engine/Source/Editor/PropertyEditor/Private/UserInterface/PropertyEditor/SPropertyEditorClass.cpp b/Engine/Source/Editor/PropertyEditor/Private/UserInterface/PropertyEditor/SPropertyEditorClass.cpp
--- a/Engine/Source/Editor/PropertyEditor/Private/UserInterface/PropertyEditor/SPropertyEditorClass.cpp
+++ b/Engine/Source/Editor/PropertyEditor/Private/UserInterface/PropertyEditor/SPropertyEditorClass.cpp
@@ -34,30 +34,35 @@ public:
 
 	virtual bool IsClassAllowed(const FClassViewerInitializationOptions& InInitOptions, const UClass* InClass, TSharedRef< FClassViewerFilterFuncs > InFilterFuncs ) override
 	{
-		bool bMatchesFlags = !InClass->HasAnyClassFlags(CLASS_Hidden|CLASS_HideDropDown|CLASS_Deprecated) &&
-			(bAllowAbstract || !InClass->HasAnyClassFlags(CLASS_Abstract));
-
-		if(bMatchesFlags && InClass->IsChildOf(ClassPropertyMetaClass)
-			&& (!InterfaceThatMustBeImplemented || InClass->ImplementsInterface(InterfaceThatMustBeImplemented)))
-		{
-			if (DisallowedClassFilters.Find(InClass) == INDEX_NONE && (AllowedClassFilters.Num() == 0 || AllowedClassFilters.Find(InClass) != INDEX_NONE))
-			{
-				return true;
-			}
-		}
-
-		return false;
+		return IsClassAllowedHelper(InClass);
+	}
+	
+	virtual bool IsUnloadedClassAllowed(const FClassViewerInitializationOptions& InInitOptions, const TSharedRef< const IUnloadedBlueprintData > InBlueprint, TSharedRef< FClassViewerFilterFuncs > InFilterFuncs) override
+	{
+		return IsClassAllowedHelper(InBlueprint);
 	}
 
-	virtual bool IsUnloadedClassAllowed(const FClassViewerInitializationOptions& InInitOptions, const TSharedRef< const IUnloadedBlueprintData > InClass, TSharedRef< FClassViewerFilterFuncs > InFilterFuncs) override
+private:
+
+	template <typename TClass>
+	bool IsClassAllowedHelper(TClass InClass)
 	{
-		bool bMatchesFlags = !InClass->HasAnyClassFlags(CLASS_Hidden|CLASS_HideDropDown|CLASS_Deprecated) &&
+		bool bMatchesFlags = !InClass->HasAnyClassFlags(CLASS_Hidden | CLASS_HideDropDown | CLASS_Deprecated) &&
 			(bAllowAbstract || !InClass->HasAnyClassFlags(CLASS_Abstract));
 
-		if(bMatchesFlags && InClass->IsChildOf(ClassPropertyMetaClass)
+		if (bMatchesFlags && InClass->IsChildOf(ClassPropertyMetaClass)
 			&& (!InterfaceThatMustBeImplemented || InClass->ImplementsInterface(InterfaceThatMustBeImplemented)))
 		{
-			return true;
+			auto PredicateFn = [InClass](const UClass* Class)
+			{
+				return InClass->IsChildOf(Class);
+			};
+
+			if (DisallowedClassFilters.FindByPredicate(PredicateFn) == nullptr &&
+				(AllowedClassFilters.Num() == 0 || AllowedClassFilters.FindByPredicate(PredicateFn) != nullptr))
+			{
+				return true;
+			}
 		}
 
 		return false;
