EditConditions applied to containers containing structs now correctly apply to the structs' members as well.

This issue was occurring because the wrong parent node was being used to resolve the memory in which the value resided. FEditConditionContext now resolves the edit condition's parent correctly to the array's parent node, because we do not support nested containers at this time.

diff --git a/Engine/Source/Editor/PropertyEditor/Private/EditConditionContext.cpp b/Engine/Source/Editor/PropertyEditor/Private/EditConditionContext.cpp
--- a/Engine/Source/Editor/PropertyEditor/Private/EditConditionContext.cpp
+++ b/Engine/Source/Editor/PropertyEditor/Private/EditConditionContext.cpp
@@ -81,6 +81,22 @@ T* FindTypedField(const TWeakPtr<FPropertyNode>& PropertyNode, const FString& Pr
 	return nullptr;
 }
 
+FPropertyNode* GetEditConditionParentNode(const TSharedPtr<FPropertyNode>& PropertyNode)
+{
+	FPropertyNode* ParentNode = PropertyNode->GetParentNode();
+	const UObject* PropertyOuter = PropertyNode->GetProperty()->GetOuter();
+
+	if (Cast<UArrayProperty>(PropertyOuter) != nullptr ||
+		Cast<USetProperty>(PropertyOuter) != nullptr ||
+		Cast<UMapProperty>(PropertyOuter) != nullptr)
+	{
+		// in a dynamic container, parent is actually one level up
+		return ParentNode->GetParentNode();
+	}
+
+	return ParentNode;
+}
+
 TOptional<bool> FEditConditionContext::GetBoolValue(const FString& PropertyName) const
 {
 	const UBoolProperty* BoolProperty = FindTypedField<UBoolProperty>(PropertyNode, PropertyName);
@@ -90,11 +106,10 @@ TOptional<bool> FEditConditionContext::GetBoolValue(const FString& PropertyName)
 	}
 
 	TSharedPtr<FPropertyNode> PinnedNode = PropertyNode.Pin();
-	FPropertyNode* ParentNode = PinnedNode->GetParentNode();
-	
-	TOptional<bool> Result;
-
+	FPropertyNode* ParentNode = GetEditConditionParentNode(PinnedNode);
 	FComplexPropertyNode* ComplexParentNode = PinnedNode->FindComplexParent();
+
+	TOptional<bool> Result;
 	for (int32 Index = 0; Index < ComplexParentNode->GetInstancesNum(); ++Index)
 	{
 		uint8* BasePtr = ComplexParentNode->GetMemoryOfInstance(Index);
@@ -131,11 +146,10 @@ TOptional<double> FEditConditionContext::GetNumericValue(const FString& Property
 	}
 
 	TSharedPtr<FPropertyNode> PinnedNode = PropertyNode.Pin();
-	FPropertyNode* ParentNode = PinnedNode->GetParentNode();
+	FPropertyNode* ParentNode = GetEditConditionParentNode(PinnedNode);
+	FComplexPropertyNode* ComplexParentNode = PinnedNode->FindComplexParent();
 
 	TOptional<double> Result;
-
-	FComplexPropertyNode* ComplexParentNode = PinnedNode->FindComplexParent();
 	for (int32 Index = 0; Index < ComplexParentNode->GetInstancesNum(); ++Index)
 	{
 		uint8* BasePtr = ComplexParentNode->GetMemoryOfInstance(Index);
@@ -199,11 +213,10 @@ TOptional<FString> FEditConditionContext::GetEnumValue(const FString& PropertyNa
 	}
 	
 	TSharedPtr<FPropertyNode> PinnedNode = PropertyNode.Pin();
-	FPropertyNode* ParentNode = PinnedNode->GetParentNode();
-	
-	TOptional<int64> Result;
-
+	FPropertyNode* ParentNode = GetEditConditionParentNode(PinnedNode);
 	FComplexPropertyNode* ComplexParentNode = PinnedNode->FindComplexParent();
+
+	TOptional<int64> Result;
 	for (int32 Index = 0; Index < ComplexParentNode->GetInstancesNum(); ++Index)
 	{
 		uint8* BasePtr = ComplexParentNode->GetMemoryOfInstance(Index);
diff --git a/Engine/Source/Editor/UnrealEd/Classes/Editor/PropertyEditorTestObject.h b/Engine/Source/Editor/UnrealEd/Classes/Editor/PropertyEditorTestObject.h
--- a/Engine/Source/Editor/UnrealEd/Classes/Editor/PropertyEditorTestObject.h
+++ b/Engine/Source/Editor/UnrealEd/Classes/Editor/PropertyEditorTestObject.h
@@ -542,6 +542,15 @@ class UPropertyEditorTestObject : public UObject
 	UPROPERTY(EditAnywhere, Category = EditCondition, meta = (EditCondition = "FloatEditCondition < 10"))
 	bool bEnabledWhenFloatLessThan10;
 
+	UPROPERTY(EditAnywhere, Category = EditCondition)
+	bool bEditConditionForArrays;
+
+	UPROPERTY(EditAnywhere, Category = EditCondition, meta = (EditCondition = "bEditConditionForArrays"))
+	TArray<UTexture2D*> ArrayWithEditCondition;
+
+	UPROPERTY(EditAnywhere, Category = EditCondition, meta = (EditCondition = "bEditConditionForArrays"))
+	TArray<FPropertyEditorTestBasicStruct> ArrayOfStructsWithEditCondition;
+
 	UPROPERTY(EditAnywhere, Category = OnlyInlineProperty, meta = (InlineCategoryProperty))
 	TEnumAsByte<EComponentMobility::Type> InlineProperty;
 
