Property Editor: Fix null crash

diff --git a/Engine/Source/Editor/PropertyEditor/Private/Presentation/PropertyEditor/PropertyEditor.cpp b/Engine/Source/Editor/PropertyEditor/Private/Presentation/PropertyEditor/PropertyEditor.cpp
--- a/Engine/Source/Editor/PropertyEditor/Private/Presentation/PropertyEditor/PropertyEditor.cpp
+++ b/Engine/Source/Editor/PropertyEditor/Private/Presentation/PropertyEditor/PropertyEditor.cpp
@@ -572,7 +572,8 @@ void FPropertyEditor::RequestRefresh()
 bool FPropertyEditor::IsOnlyVisibleWhenEditConditionMet() const
 {
 	static const FName Name_EditConditionHides("EditConditionHides");
-	if (PropertyNode->GetProperty()->HasMetaData(Name_EditConditionHides))
+	UProperty* Property = PropertyNode->GetProperty();
+	if (Property && Property->HasMetaData(Name_EditConditionHides))
 	{
 		return HasEditCondition();
 	}
