Property Editor: Consolidated add external property parameters and added the ability to forcibly show external properties

diff --git a/Engine/Plugins/FX/Niagara/Source/NiagaraEditor/Private/Customizations/NiagaraComponentDetails.cpp b/Engine/Plugins/FX/Niagara/Source/NiagaraEditor/Private/Customizations/NiagaraComponentDetails.cpp
--- a/Engine/Plugins/FX/Niagara/Source/NiagaraEditor/Private/Customizations/NiagaraComponentDetails.cpp
+++ b/Engine/Plugins/FX/Niagara/Source/NiagaraEditor/Private/Customizations/NiagaraComponentDetails.cpp
@@ -145,9 +145,12 @@ public:
 				TArray<UObject*> Objects;
 				Objects.Add(DefaultValueObject);
 
-				TOptional<bool> bAllowChildrenOverride(true);
-				TOptional<bool> bCreateCategoryNodesOverride(false);
-				Row = ChildrenBuilder.AddExternalObjectProperty(Objects, NAME_None, Parameter.GetName(), bAllowChildrenOverride, bCreateCategoryNodesOverride); 
+				FAddPropertyParams Params = FAddPropertyParams()
+					.UniqueId(Parameter.GetName())
+					.AllowChildren(true)
+					.CreateCategoryNodes(false);
+
+				Row = ChildrenBuilder.AddExternalObjectProperty(Objects, NAME_None, Params); 
 
 				CustomValueWidget =
 					SNew(STextBlock)
@@ -164,9 +167,12 @@ public:
 
 				//How do I set this up so I can have this pick actors from the level?
 
-				TOptional<bool> bAllowChildrenOverride(true);
-				TOptional<bool> bCreateCategoryNodesOverride(false);
-				Row = ChildrenBuilder.AddExternalObjectProperty(Objects, NAME_None, Parameter.GetName(), bAllowChildrenOverride, bCreateCategoryNodesOverride);
+				FAddPropertyParams Params = FAddPropertyParams()
+					.UniqueId(Parameter.GetName())
+					.AllowChildren(true)
+					.CreateCategoryNodes(false);
+
+				Row = ChildrenBuilder.AddExternalObjectProperty(Objects, NAME_None, Params);
 
 				//How do I make this an object picker from the level editor?
 				//Neither of these seem to work well.
@@ -206,7 +212,8 @@ public:
 			else
 			{
 				TSharedPtr<FStructOnScope> StructOnScope = MakeShareable(new FStructOnScope(Parameter.GetType().GetStruct(), (uint8*)ParamStore.GetParameterData(Parameter)));
-				Row = ChildrenBuilder.AddExternalStructureProperty(StructOnScope.ToSharedRef(), NAME_None, Parameter.GetName());
+
+				Row = ChildrenBuilder.AddExternalStructureProperty(StructOnScope.ToSharedRef(), NAME_None, FAddPropertyParams().UniqueId(Parameter.GetName()));
 
 			}
 
diff --git a/Engine/Plugins/FX/Niagara/Source/NiagaraEditor/Private/Customizations/NiagaraMetaDataCustomNodeBuilder.cpp b/Engine/Plugins/FX/Niagara/Source/NiagaraEditor/Private/Customizations/NiagaraMetaDataCustomNodeBuilder.cpp
--- a/Engine/Plugins/FX/Niagara/Source/NiagaraEditor/Private/Customizations/NiagaraMetaDataCustomNodeBuilder.cpp
+++ b/Engine/Plugins/FX/Niagara/Source/NiagaraEditor/Private/Customizations/NiagaraMetaDataCustomNodeBuilder.cpp
@@ -97,7 +97,7 @@ void FNiagaraMetaDataCustomNodeBuilder::GenerateChildContent(IDetailChildrenBuil
 		const FNiagaraParameterHandle& ParameterHandle = MetaDataRowDataItem.Get<1>();
 		TSharedRef<FStructOnScope>& MetaDataContainerStruct = MetaDataRowDataItem.Get<2>();
 		
-		IDetailPropertyRow* MetaDataRow = ChildrenBuilder.AddExternalStructureProperty(MetaDataContainerStruct, GET_MEMBER_NAME_CHECKED(FNiagaraVariableMetaDataContainer, MetaData), ParameterVariable.GetName());
+		IDetailPropertyRow* MetaDataRow = ChildrenBuilder.AddExternalStructureProperty(MetaDataContainerStruct, GET_MEMBER_NAME_CHECKED(FNiagaraVariableMetaDataContainer, MetaData), FAddPropertyParams().UniqueId(ParameterVariable.GetName()));
 		MetaDataRow->DisplayName(FText::FromName(ParameterHandle.GetName()));
 		if (MetaDataRow->GetPropertyHandle().IsValid())
 		{
diff --git a/Engine/Plugins/FX/Niagara/Source/NiagaraEditor/Private/Customizations/NiagaraParameterCollectionCustomNodeBuilder.h b/Engine/Plugins/FX/Niagara/Source/NiagaraEditor/Private/Customizations/NiagaraParameterCollectionCustomNodeBuilder.h
--- a/Engine/Plugins/FX/Niagara/Source/NiagaraEditor/Private/Customizations/NiagaraParameterCollectionCustomNodeBuilder.h
+++ b/Engine/Plugins/FX/Niagara/Source/NiagaraEditor/Private/Customizations/NiagaraParameterCollectionCustomNodeBuilder.h
@@ -87,7 +87,7 @@ public:
 
 			if (Parameter->GetDefaultValueType() == INiagaraParameterViewModel::EDefaultValueType::Struct)
 			{
-				Row = ChildrenBuilder.AddExternalStructureProperty(Parameter->GetDefaultValueStruct(), NAME_None, Parameter->GetName());
+				Row = ChildrenBuilder.AddExternalStructureProperty(Parameter->GetDefaultValueStruct(), NAME_None, FAddPropertyParams().UniqueId(Parameter->GetName()));
 
 			}
 			else if (Parameter->GetDefaultValueType() == INiagaraParameterViewModel::EDefaultValueType::Object)
@@ -97,10 +97,12 @@ public:
 				TArray<UObject*> Objects;
 				Objects.Add(DefaultValueObject);
 				
-				TOptional<bool> bAllowChildrenOverride(true);
-				TOptional<bool> bCreateCategoryNodesOverride(false);
+				FAddPropertyParams Params = FAddPropertyParams()
+					.UniqueId(Parameter->GetName())
+					.AllowChildren(true)
+					.CreateCategoryNodes(false);
 
-				Row = ChildrenBuilder.AddExternalObjectProperty(Objects, NAME_None, Parameter->GetName(), bAllowChildrenOverride, bCreateCategoryNodesOverride);
+				Row = ChildrenBuilder.AddExternalObjectProperty(Objects, NAME_None, Params);
 				CustomValueWidget =
 					SNew(STextBlock)
 					.TextStyle(FNiagaraEditorStyle::Get(), "NiagaraEditor.ParameterText")
diff --git a/Engine/Source/Editor/PropertyEditor/Private/CustomChildBuilder.cpp b/Engine/Source/Editor/PropertyEditor/Private/CustomChildBuilder.cpp
--- a/Engine/Source/Editor/PropertyEditor/Private/CustomChildBuilder.cpp
+++ b/Engine/Source/Editor/PropertyEditor/Private/CustomChildBuilder.cpp
@@ -65,23 +65,23 @@ IDetailPropertyRow& FCustomChildrenBuilder::AddProperty( TSharedRef<IPropertyHan
 
 IDetailPropertyRow* FCustomChildrenBuilder::AddExternalStructure(TSharedRef<FStructOnScope> ChildStructure, FName UniqueIdName)
 {
-	return AddExternalStructureProperty(ChildStructure, NAME_None, UniqueIdName);
+	return AddExternalStructureProperty(ChildStructure, NAME_None, FAddPropertyParams().UniqueId(UniqueIdName));
 }
 
-IDetailPropertyRow* FCustomChildrenBuilder::AddExternalStructureProperty(TSharedRef<FStructOnScope> ChildStructure, FName PropertyName, FName UniqueIdName)
+IDetailPropertyRow* FCustomChildrenBuilder::AddExternalStructureProperty(TSharedRef<FStructOnScope> ChildStructure, FName PropertyName, const FAddPropertyParams& Params)
 {
 	FDetailLayoutCustomization NewCustomization;
 
 	TSharedRef<FDetailCategoryImpl> ParentCategoryRef = ParentCategory.Pin().ToSharedRef();
 
-	FDetailPropertyRow::MakeExternalPropertyRowCustomization(ChildStructure, PropertyName, ParentCategoryRef, NewCustomization);
+	FDetailPropertyRow::MakeExternalPropertyRowCustomization(ChildStructure, PropertyName, ParentCategoryRef, NewCustomization, Params);
 
 	TSharedPtr<FDetailPropertyRow> NewRow = NewCustomization.PropertyRow;
 
-	NewRow->SetCustomExpansionId(UniqueIdName);
-
 	if (NewRow.IsValid())
 	{
+		NewRow->SetCustomExpansionId(Params.GetUniqueId());
+
 		TSharedPtr<FPropertyNode> PropertyNode = NewRow->GetPropertyNode();
 		TSharedPtr<FComplexPropertyNode> RootNode = StaticCastSharedRef<FComplexPropertyNode>(PropertyNode->FindComplexParent()->AsShared());
 
@@ -93,7 +93,7 @@ IDetailPropertyRow* FCustomChildrenBuilder::AddExternalStructureProperty(TShared
 
 IDetailPropertyRow* FCustomChildrenBuilder::AddExternalObjects(const TArray<UObject *>& Objects, FName UniqueIdName)
 {
-	return AddExternalObjectProperty(Objects, NAME_None, UniqueIdName);
+	return AddExternalObjectProperty(Objects, NAME_None, FAddPropertyParams().UniqueId(UniqueIdName));
 }
 
 TArray<TSharedPtr<IPropertyHandle>> FCustomChildrenBuilder::AddAllExternalStructureProperties(TSharedRef<FStructOnScope> ChildStructure)
@@ -101,20 +101,20 @@ TArray<TSharedPtr<IPropertyHandle>> FCustomChildrenBuilder::AddAllExternalStruct
 	return ParentCategory.Pin()->AddAllExternalStructureProperties(ChildStructure);
 }
 
-IDetailPropertyRow* FCustomChildrenBuilder::AddExternalObjectProperty(const TArray<UObject*>& Objects, FName PropertyName, FName UniqueIdName, TOptional<bool> bAllowChildrenOverride, TOptional<bool> bCreateCategoryNodesOverride)
+IDetailPropertyRow* FCustomChildrenBuilder::AddExternalObjectProperty(const TArray<UObject*>& Objects, FName PropertyName, const FAddPropertyParams& Params)
 {
 	FDetailLayoutCustomization NewCustomization;
 
 	TSharedRef<FDetailCategoryImpl> ParentCategoryRef = ParentCategory.Pin().ToSharedRef();
 
-	FDetailPropertyRow::MakeExternalPropertyRowCustomization(Objects, PropertyName, ParentCategoryRef, NewCustomization, bAllowChildrenOverride, bCreateCategoryNodesOverride);
+	FDetailPropertyRow::MakeExternalPropertyRowCustomization(Objects, PropertyName, ParentCategoryRef, NewCustomization, Params);
 
 	TSharedPtr<FDetailPropertyRow> NewRow = NewCustomization.PropertyRow;
 
-	NewRow->SetCustomExpansionId(UniqueIdName);
-
 	if (NewRow.IsValid())
 	{
+		NewRow->SetCustomExpansionId(Params.GetUniqueId());
+
 		TSharedPtr<FPropertyNode> PropertyNode = NewRow->GetPropertyNode();
 		TSharedPtr<FObjectPropertyNode> RootNode = StaticCastSharedRef<FObjectPropertyNode>(PropertyNode->FindObjectItemParent()->AsShared());
 
diff --git a/Engine/Source/Editor/PropertyEditor/Private/CustomChildBuilder.h b/Engine/Source/Editor/PropertyEditor/Private/CustomChildBuilder.h
--- a/Engine/Source/Editor/PropertyEditor/Private/CustomChildBuilder.h
+++ b/Engine/Source/Editor/PropertyEditor/Private/CustomChildBuilder.h
@@ -25,9 +25,9 @@ public:
 	virtual FDetailWidgetRow& AddCustomRow(const FText& SearchString) override;
 	virtual IDetailPropertyRow& AddProperty(TSharedRef<IPropertyHandle> PropertyHandle) override;
 	virtual IDetailPropertyRow* AddExternalObjects(const TArray<UObject*>& Objects, FName UniqueIdName = NAME_None) override;
-	virtual IDetailPropertyRow* AddExternalObjectProperty(const TArray<UObject*>& Objects, FName PropertyName, FName UniqueIdName = NAME_None, TOptional<bool> bAllowChildrenOverride = TOptional<bool>(), TOptional<bool> bCreateCategoryNodesOverride = TOptional<bool>()) override;
+	virtual IDetailPropertyRow* AddExternalObjectProperty(const TArray<UObject*>& Objects, FName PropertyName, const FAddPropertyParams& Params) override;
 	virtual IDetailPropertyRow* AddExternalStructure(TSharedRef<FStructOnScope> ChildStructure, FName UniqueIdName = NAME_None) override;
-	virtual IDetailPropertyRow* AddExternalStructureProperty(TSharedRef<FStructOnScope> ChildStructure, FName PropertyName, FName UniqueIdName = NAME_None) override;
+	virtual IDetailPropertyRow* AddExternalStructureProperty(TSharedRef<FStructOnScope> ChildStructure, FName PropertyName, const FAddPropertyParams& Params) override;
 	virtual TArray<TSharedPtr<IPropertyHandle>> AddAllExternalStructureProperties(TSharedRef<FStructOnScope> ChildStructure) override;
 	virtual TSharedRef<SWidget> GenerateStructValueWidget(TSharedRef<IPropertyHandle> StructPropertyHandle) override;
 	virtual IDetailCategoryBuilder& GetParentCategory() const override;
diff --git a/Engine/Source/Editor/PropertyEditor/Private/DetailCategoryBuilderImpl.cpp b/Engine/Source/Editor/PropertyEditor/Private/DetailCategoryBuilderImpl.cpp
--- a/Engine/Source/Editor/PropertyEditor/Private/DetailCategoryBuilderImpl.cpp
+++ b/Engine/Source/Editor/PropertyEditor/Private/DetailCategoryBuilderImpl.cpp
@@ -13,6 +13,7 @@
 #include "StructurePropertyNode.h"
 #include "ItemPropertyNode.h"
 #include "IPropertyGenerationUtilities.h"
+#include "DetailBuilderTypes.h"
 
 namespace DetailLayoutConstants
 {
@@ -326,7 +327,7 @@ IDetailPropertyRow* FDetailCategoryImpl::AddExternalObjects(const TArray<UObject
 {
 	FDetailLayoutCustomization NewCustomization;
 
-	FDetailPropertyRow::MakeExternalPropertyRowCustomization(Objects, NAME_None, AsShared(), NewCustomization, true);
+	FDetailPropertyRow::MakeExternalPropertyRowCustomization(Objects, NAME_None, AsShared(), NewCustomization, FAddPropertyParams().AllowChildren(true));
 
 	TSharedPtr<FDetailPropertyRow> NewRow = NewCustomization.PropertyRow;
 
@@ -345,11 +346,11 @@ IDetailPropertyRow* FDetailCategoryImpl::AddExternalObjects(const TArray<UObject
 	return NewRow.Get();
 }
 
-IDetailPropertyRow* FDetailCategoryImpl::AddExternalObjectProperty(const TArray<UObject*>& Objects, FName PropertyName, EPropertyLocation::Type Location)
+IDetailPropertyRow* FDetailCategoryImpl::AddExternalObjectProperty(const TArray<UObject*>& Objects, FName PropertyName, EPropertyLocation::Type Location, const FAddPropertyParams& Params)
 {
 	FDetailLayoutCustomization NewCustomization;
 
-	FDetailPropertyRow::MakeExternalPropertyRowCustomization(Objects, PropertyName, AsShared(), NewCustomization);
+	FDetailPropertyRow::MakeExternalPropertyRowCustomization(Objects, PropertyName, AsShared(), NewCustomization, Params);
 
 	TSharedPtr<FDetailPropertyRow> NewRow = NewCustomization.PropertyRow;
 
@@ -374,11 +375,11 @@ IDetailPropertyRow* FDetailCategoryImpl::AddExternalStructure(TSharedPtr<FStruct
 	return AddExternalStructureProperty(StructData, NAME_None, Location);
 }
 
-IDetailPropertyRow* FDetailCategoryImpl::AddExternalStructureProperty(TSharedPtr<FStructOnScope> StructData, FName PropertyName, EPropertyLocation::Type Location/* = EPropertyLocation::Default*/)
+IDetailPropertyRow* FDetailCategoryImpl::AddExternalStructureProperty(TSharedPtr<FStructOnScope> StructData, FName PropertyName, EPropertyLocation::Type Location/* = EPropertyLocation::Default*/, const FAddPropertyParams& Params)
 {
 	FDetailLayoutCustomization NewCustomization;
 
-	FDetailPropertyRow::MakeExternalPropertyRowCustomization(StructData, PropertyName, AsShared(), NewCustomization);
+	FDetailPropertyRow::MakeExternalPropertyRowCustomization(StructData, PropertyName, AsShared(), NewCustomization, Params);
 
 	TSharedPtr<FDetailPropertyRow> NewRow = NewCustomization.PropertyRow;
 
diff --git a/Engine/Source/Editor/PropertyEditor/Private/DetailCategoryBuilderImpl.h b/Engine/Source/Editor/PropertyEditor/Private/DetailCategoryBuilderImpl.h
--- a/Engine/Source/Editor/PropertyEditor/Private/DetailCategoryBuilderImpl.h
+++ b/Engine/Source/Editor/PropertyEditor/Private/DetailCategoryBuilderImpl.h
@@ -163,9 +163,9 @@ public:
 	virtual IDetailPropertyRow& AddProperty(FName PropertyPath, UClass* ClassOuter = nullptr, FName InstanceName = NAME_None, EPropertyLocation::Type Location = EPropertyLocation::Default) override;
 	virtual IDetailPropertyRow& AddProperty(TSharedPtr<IPropertyHandle> PropertyHandle, EPropertyLocation::Type Location = EPropertyLocation::Default) override;
 	virtual IDetailPropertyRow* AddExternalObjects(const TArray<UObject*>& Objects, EPropertyLocation::Type Location = EPropertyLocation::Default) override;
-	virtual IDetailPropertyRow* AddExternalObjectProperty(const TArray<UObject*>& Objects, FName PropertyName, EPropertyLocation::Type Location = EPropertyLocation::Default) override;
+	virtual IDetailPropertyRow* AddExternalObjectProperty(const TArray<UObject*>& Objects, FName PropertyName, EPropertyLocation::Type Location = EPropertyLocation::Default, const FAddPropertyParams& Params = FAddPropertyParams()) override;
 	virtual IDetailPropertyRow* AddExternalStructure(TSharedPtr<FStructOnScope> StructData, EPropertyLocation::Type Location = EPropertyLocation::Default) override;
-	virtual IDetailPropertyRow* AddExternalStructureProperty(TSharedPtr<FStructOnScope> StructData, FName PropertyName, EPropertyLocation::Type Location = EPropertyLocation::Default) override;
+	virtual IDetailPropertyRow* AddExternalStructureProperty(TSharedPtr<FStructOnScope> StructData, FName PropertyName, EPropertyLocation::Type Location = EPropertyLocation::Default, const FAddPropertyParams& Params = FAddPropertyParams()) override;
 	virtual TArray<TSharedPtr<IPropertyHandle>> AddAllExternalStructureProperties(TSharedRef<FStructOnScope> StructData, EPropertyLocation::Type Location = EPropertyLocation::Default) override;
 	virtual IDetailLayoutBuilder& GetParentLayout() const override { return *DetailLayoutBuilder.Pin(); }
 	virtual FDetailWidgetRow& AddCustomRow(const FText& FilterString, bool bForAdvanced = false) override;
diff --git a/Engine/Source/Editor/PropertyEditor/Private/DetailPropertyRow.cpp b/Engine/Source/Editor/PropertyEditor/Private/DetailPropertyRow.cpp
--- a/Engine/Source/Editor/PropertyEditor/Private/DetailPropertyRow.cpp
+++ b/Engine/Source/Editor/PropertyEditor/Private/DetailPropertyRow.cpp
@@ -431,7 +431,7 @@ TSharedPtr<IPropertyTypeCustomization> FDetailPropertyRow::GetPropertyCustomizat
 	return CustomInterface;
 }
 
-void FDetailPropertyRow::MakeExternalPropertyRowCustomization(TSharedPtr<FStructOnScope> StructData, FName PropertyName, TSharedRef<FDetailCategoryImpl> ParentCategory, FDetailLayoutCustomization& OutCustomization)
+void FDetailPropertyRow::MakeExternalPropertyRowCustomization(TSharedPtr<FStructOnScope> StructData, FName PropertyName, TSharedRef<FDetailCategoryImpl> ParentCategory, FDetailLayoutCustomization& OutCustomization, const FAddPropertyParams& Parameters)
 {
 	TSharedRef<FStructurePropertyNode> RootPropertyNode = MakeShared<FStructurePropertyNode>();
 
@@ -443,10 +443,13 @@ void FDetailPropertyRow::MakeExternalPropertyRowCustomization(TSharedPtr<FStruct
 	InitParams.Property = nullptr;
 	InitParams.ArrayOffset = 0;
 	InitParams.ArrayIndex = INDEX_NONE;
-	InitParams.bForceHiddenPropertyVisibility = FPropertySettings::Get().ShowHiddenProperties();
+	InitParams.bForceHiddenPropertyVisibility = Parameters.ShouldForcePropertyVisible() || FPropertySettings::Get().ShowHiddenProperties();
 	InitParams.bCreateCategoryNodes = false;
 	InitParams.bAllowChildren = false;
 
+	Parameters.OverrideAllowChildren(InitParams.bAllowChildren);
+	Parameters.OverrideCreateCategoryNodes(InitParams.bCreateCategoryNodes);
+
 	RootPropertyNode->InitNode(InitParams);
 
 	ParentCategory->GetParentLayoutImpl().AddExternalRootPropertyNode(RootPropertyNode);
@@ -460,7 +463,7 @@ void FDetailPropertyRow::MakeExternalPropertyRowCustomization(TSharedPtr<FStruct
 			TSharedPtr< FPropertyNode > PropertyNode = RootPropertyNode->GetChildNode(ChildIdx);
 			if (UProperty* Property = PropertyNode->GetProperty())
 			{
-				if (PropertyName == NAME_None || Property->GetFName() == PropertyName)
+				if (Property->GetFName() == PropertyName)
 				{
 					OutCustomization.PropertyRow = MakeShareable(new FDetailPropertyRow(PropertyNode, ParentCategory, RootPropertyNode));
 					break;
@@ -485,7 +488,7 @@ void FDetailPropertyRow::MakeExternalPropertyRowCustomization(TSharedPtr<FStruct
 		ItemNodeInitParams.ArrayOffset = 0;
 		ItemNodeInitParams.ArrayIndex = INDEX_NONE;
 		ItemNodeInitParams.bAllowChildren = true;
-		ItemNodeInitParams.bForceHiddenPropertyVisibility = FPropertySettings::Get().ShowHiddenProperties();
+		ItemNodeInitParams.bForceHiddenPropertyVisibility = Parameters.ShouldForcePropertyVisible() || FPropertySettings::Get().ShowHiddenProperties();
 		ItemNodeInitParams.bCreateCategoryNodes = false;
 
 		ItemNode->InitNode(ItemNodeInitParams);
@@ -496,7 +499,7 @@ void FDetailPropertyRow::MakeExternalPropertyRowCustomization(TSharedPtr<FStruct
 	}
 }
 
-void FDetailPropertyRow::MakeExternalPropertyRowCustomization(const TArray<UObject*>& InObjects, FName PropertyName, TSharedRef<FDetailCategoryImpl> ParentCategory, struct FDetailLayoutCustomization& OutCustomization, TOptional<bool> bAllowChildrenOverride, TOptional<bool> bCreateCategoryNodesOverride)
+void FDetailPropertyRow::MakeExternalPropertyRowCustomization(const TArray<UObject*>& InObjects, FName PropertyName, TSharedRef<FDetailCategoryImpl> ParentCategory, struct FDetailLayoutCustomization& OutCustomization, const FAddPropertyParams& Parameters)
 {
 	TSharedRef<FObjectPropertyNode> RootPropertyNode = MakeShared<FObjectPropertyNode>();
 
@@ -511,17 +514,11 @@ void FDetailPropertyRow::MakeExternalPropertyRowCustomization(const TArray<UObje
 	InitParams.ArrayOffset = 0;
 	InitParams.ArrayIndex = INDEX_NONE;
 	InitParams.bAllowChildren = false;
-	InitParams.bForceHiddenPropertyVisibility = FPropertySettings::Get().ShowHiddenProperties();
+	InitParams.bForceHiddenPropertyVisibility = Parameters.ShouldForcePropertyVisible() || FPropertySettings::Get().ShowHiddenProperties();
 	InitParams.bCreateCategoryNodes = PropertyName == NAME_None;
 
-	if (bAllowChildrenOverride.IsSet())
-	{
-		InitParams.bAllowChildren = bAllowChildrenOverride.GetValue();
-	}
-	if (bCreateCategoryNodesOverride.IsSet())
-	{
-		InitParams.bCreateCategoryNodes = bCreateCategoryNodesOverride.GetValue();
-	}
+	Parameters.OverrideAllowChildren(InitParams.bAllowChildren);
+	Parameters.OverrideCreateCategoryNodes(InitParams.bCreateCategoryNodes);
 
 	RootPropertyNode->InitNode(InitParams);
 
diff --git a/Engine/Source/Editor/PropertyEditor/Private/DetailPropertyRow.h b/Engine/Source/Editor/PropertyEditor/Private/DetailPropertyRow.h
--- a/Engine/Source/Editor/PropertyEditor/Private/DetailPropertyRow.h
+++ b/Engine/Source/Editor/PropertyEditor/Private/DetailPropertyRow.h
@@ -15,6 +15,8 @@
 #include "IDetailPropertyRow.h"
 #include "PropertyEditorHelpers.h"
 
+struct FAddPropertyParams;
+
 class FCustomChildrenBuilder;
 class IDetailGroup;
 
@@ -102,8 +104,8 @@ public:
 	 */
 	EVisibility GetPropertyVisibility() const;
 
-	static void MakeExternalPropertyRowCustomization(TSharedPtr<FStructOnScope> StructData, FName PropertyName, TSharedRef<FDetailCategoryImpl> ParentCategory, struct FDetailLayoutCustomization& OutCustomization);
-	static void MakeExternalPropertyRowCustomization(const TArray<UObject*>& InObjects, FName PropertyName, TSharedRef<FDetailCategoryImpl> ParentCategory, struct FDetailLayoutCustomization& OutCustomization, TOptional<bool> bAllowChildrenOverride = TOptional<bool>(), TOptional<bool> bCreateCategoryNodesOverride = TOptional<bool>());
+	static void MakeExternalPropertyRowCustomization(TSharedPtr<FStructOnScope> StructData, FName PropertyName, TSharedRef<FDetailCategoryImpl> ParentCategory, struct FDetailLayoutCustomization& OutCustomization, const FAddPropertyParams& Parameters);
+	static void MakeExternalPropertyRowCustomization(const TArray<UObject*>& InObjects, FName PropertyName, TSharedRef<FDetailCategoryImpl> ParentCategory, struct FDetailLayoutCustomization& OutCustomization, const FAddPropertyParams& Parameters);
 
 private:
 	/**
diff --git a/Engine/Source/Editor/PropertyEditor/Public/DetailBuilderTypes.h b/Engine/Source/Editor/PropertyEditor/Public/DetailBuilderTypes.h
new file mode 100644
--- /dev/null
+++ b/Engine/Source/Editor/PropertyEditor/Public/DetailBuilderTypes.h
@@ -0,0 +1,115 @@
+// Copyright 1998-2019 Epic Games, Inc. All Rights Reserved.
+
+#pragma once
+
+#include "CoreTypes.h"
+#include "UObject/NameTypes.h"
+#include "Misc/Optional.h"
+
+/**
+ * Parameters required for specifying behavior when adding external properties from detail customizations
+ */
+struct FAddPropertyParams
+{
+	FAddPropertyParams()
+		: bForceShowProperty(false)
+	{}
+
+
+	/**
+	 * Forcibly show the property, even if it does not have CPF_Edit
+	 */
+	FAddPropertyParams& ForceShowProperty()
+	{
+		bForceShowProperty = 1;
+		return *this;
+	}
+
+
+	/**
+	 * Override whether the property node should allow children or not. If not overridden the default is implementation defined.
+	 */
+	FAddPropertyParams& AllowChildren(bool bAllowChildren)
+	{
+		bAllowChildrenOverride = (bAllowChildren ? 1 : 2);
+		return *this;
+	}
+
+
+	/**
+	 * Override whether the property node should create category nodes or not. If not overridden the default is implementation defined.
+	 */
+	FAddPropertyParams& CreateCategoryNodes(bool bCreateCategoryNodes)
+	{
+		bCreateCategoryNodesOverride = (bCreateCategoryNodes ? 1 : 2);
+		return *this;
+	}
+
+
+	/**
+	 * Set a unique name for this property, allowing it to correctly save expansion states and other persistent UI state
+	 */
+	FAddPropertyParams& UniqueId(FName InUniqueId)
+	{
+		UniqueIdName = InUniqueId;
+		return *this;
+	}
+
+
+	/**
+	 * Check whether to forcibly show the property, even if it does not have CPF_Edit
+	 */
+	bool ShouldForcePropertyVisible() const
+	{
+		return bForceShowProperty == 1;
+	}
+
+
+	/**
+	 * Conditionally overwrites the specified boolean with a value specifying whether to allow child properties or not.
+	 * Only overrides the value if the calling code specified an override
+	 */
+	void OverrideAllowChildren(bool& OutAllowChildren) const
+	{
+		if (bAllowChildrenOverride.IsSet())
+		{
+			OutAllowChildren = bAllowChildrenOverride.GetValue();
+		}
+	}
+
+
+	/**
+	 * Conditionally overwrites the specified boolean with a value specifying whether to create category nodes or not
+	 * Only overrides the value if the calling code specified an override
+	 */
+	void OverrideCreateCategoryNodes(bool& OutCreateCategoryNodes) const
+	{
+		if (bCreateCategoryNodesOverride.IsSet())
+		{
+			OutCreateCategoryNodes = bCreateCategoryNodesOverride.GetValue();
+		}
+	}
+
+
+	/**
+	 * Get this property's unique ID name
+	 */
+	FName GetUniqueId() const
+	{
+		return UniqueIdName;
+	}
+
+private:
+
+	/** When true, the property will be forcefully shown, even if it does not have CPF_Edit. When false the property will only be created if it has CPF_Edit. */
+	bool bForceShowProperty : 1;
+
+	/** Tristate override for allowing children - if value is true: Allow Children, false: Disallow Children, unset: no override */
+	TOptional<bool> bAllowChildrenOverride;
+
+	/** Tristate override for allowing children - if value is true: create category nodes, false: do not create category nodes, unset: no override */
+	TOptional<bool> bCreateCategoryNodesOverride;
+
+	/** Unique ID name that is used for saving persistent UI state such as expansion */
+	FName UniqueIdName;
+};
\ No newline at end of file
diff --git a/Engine/Source/Editor/PropertyEditor/Public/DetailCategoryBuilder.h b/Engine/Source/Editor/PropertyEditor/Public/DetailCategoryBuilder.h
--- a/Engine/Source/Editor/PropertyEditor/Public/DetailCategoryBuilder.h
+++ b/Engine/Source/Editor/PropertyEditor/Public/DetailCategoryBuilder.h
@@ -8,6 +8,7 @@
 #include "UObject/StructOnScope.h"
 #include "PropertyHandle.h"
 #include "IDetailCustomNodeBuilder.h"
+#include "DetailBuilderTypes.h"
 
 class FDetailWidgetRow;
 class IDetailGroup;
@@ -92,7 +93,7 @@ public:
 	 * @param Location		The location within the category where the property is shown
 	 * @return A property row for customizing the property or NULL if the property could not be found
 	 */
-	virtual IDetailPropertyRow* AddExternalObjectProperty(const TArray<UObject*>& Objects, FName PropertyName, EPropertyLocation::Type Location = EPropertyLocation::Default) = 0;
+	virtual IDetailPropertyRow* AddExternalObjectProperty(const TArray<UObject*>& Objects, FName PropertyName, EPropertyLocation::Type Location = EPropertyLocation::Default, const FAddPropertyParams& Params = FAddPropertyParams()) = 0;
 
 	/**
 	 * Adds an external standalone structure to this details panel
@@ -111,7 +112,7 @@ public:
 	 * @param Location			The location within the category where the property is shown
 	 * @return A property row for customizing the property or NULL if the property could not be found
 	 */
-	virtual IDetailPropertyRow* AddExternalStructureProperty(TSharedPtr<FStructOnScope> StructData, FName PropertyName, EPropertyLocation::Type Location = EPropertyLocation::Default) = 0;
+	virtual IDetailPropertyRow* AddExternalStructureProperty(TSharedPtr<FStructOnScope> StructData, FName PropertyName, EPropertyLocation::Type Location = EPropertyLocation::Default, const FAddPropertyParams& Params = FAddPropertyParams()) = 0;
 
 	/**
 	 * Adds all properties for the specified external structure to this category
diff --git a/Engine/Source/Editor/PropertyEditor/Public/IDetailChildrenBuilder.h b/Engine/Source/Editor/PropertyEditor/Public/IDetailChildrenBuilder.h
--- a/Engine/Source/Editor/PropertyEditor/Public/IDetailChildrenBuilder.h
+++ b/Engine/Source/Editor/PropertyEditor/Public/IDetailChildrenBuilder.h
@@ -3,6 +3,7 @@
 #pragma once
 
 #include "CoreMinimal.h"
+#include "DetailBuilderTypes.h"
 
 class SWidget;
 class IPropertyHandle;
@@ -71,24 +72,25 @@ public:
 	 * @param  bCreateCategoryNodesOverride Allows customization of how the new root node's category is displayed (or not). 
 	 * @return The header row generated for this set of objects by the details panel
 	 */
-	virtual class IDetailPropertyRow* AddExternalObjectProperty(const TArray<UObject*>& Objects, FName PropertyName, FName UniqueIdName = NAME_None, TOptional<bool> bAllowChildrenOverride = TOptional<bool>(), TOptional<bool> bCreateCategoryNodesOverride = TOptional<bool>()) = 0;
+	virtual IDetailPropertyRow* AddExternalObjectProperty(const TArray<UObject*>& Objects, FName PropertyName, const FAddPropertyParams& Params = FAddPropertyParams()) = 0;
+
 
 	/**
-	 * Adds a custom structure as a child
+	 * Adds a property from a custom structure as a child
 	 *
 	 * @param ChildStructure	The structure to add
+	 * @param PropertyName		Optional name of a property inside the Child structure to add.  If this is empty, the entire structure will be added
 	 * @param UniqueIdName		Optional identifier that uniquely identifies this structure among other structures of the same type.  If this is empty, saving and restoring expansion state of this structure may not work
 	 */
-	virtual class IDetailPropertyRow* AddExternalStructure(TSharedRef<FStructOnScope> ChildStructure, FName UniqueIdName = NAME_None) = 0;
+	virtual IDetailPropertyRow* AddExternalStructureProperty(TSharedRef<FStructOnScope> ChildStructure, FName PropertyName, const FAddPropertyParams& Params = FAddPropertyParams()) = 0;
 
 	/**
-	 * Adds a property from a custom structure as a child
+	 * Adds a custom structure as a child
 	 *
 	 * @param ChildStructure	The structure to add
-	 * @param PropertyName		Optional name of a property inside the Child structure to add.  If this is empty, the entire structure will be added
 	 * @param UniqueIdName		Optional identifier that uniquely identifies this structure among other structures of the same type.  If this is empty, saving and restoring expansion state of this structure may not work
 	 */
-	virtual class IDetailPropertyRow* AddExternalStructureProperty(TSharedRef<FStructOnScope> ChildStructure, FName PropertyName, FName UniqueIdName = NAME_None) = 0;
+	virtual class IDetailPropertyRow* AddExternalStructure(TSharedRef<FStructOnScope> ChildStructure, FName UniqueIdName = NAME_None) = 0;
 
 	/**
 	 * Adds all the properties of an external structure as a children
@@ -116,6 +118,26 @@ public:
 	*/
 	virtual IDetailGroup* GetParentGroup() const = 0;
 
+	UE_DEPRECATED(4.24, "Please use the overload that accepts an FAddPropertyParams structure")
+	IDetailPropertyRow* AddExternalObjectProperty(const TArray<UObject*>& Objects, FName PropertyName, FName UniqueIdName = NAME_None, TOptional<bool> bAllowChildrenOverride = TOptional<bool>(), TOptional<bool> bCreateCategoryNodesOverride = TOptional<bool>())
+	{
+		FAddPropertyParams Params;
+		Params.UniqueId(UniqueIdName);
+		if (bAllowChildrenOverride.IsSet())
+		{
+			Params.AllowChildren(bAllowChildrenOverride.GetValue());
+		}
+		if (bCreateCategoryNodesOverride.IsSet())
+		{
+			Params.CreateCategoryNodes(bCreateCategoryNodesOverride.GetValue());
+		}
+		return AddExternalObjectProperty(Objects, PropertyName, Params);
+	}
+	UE_DEPRECATED(4.24, "Please use the overload that accepts an FAddPropertyParams structure")
+	IDetailPropertyRow* AddExternalStructureProperty(TSharedRef<FStructOnScope> ChildStructure, FName PropertyName, FName UniqueIdName)
+	{
+		return AddExternalStructureProperty(ChildStructure, PropertyName, FAddPropertyParams().UniqueId(UniqueIdName));
+	}
 
 	UE_DEPRECATED(4.17, "AddChildCustomBuilder has been deprecated.  Use AddCustomBuilder instead")
 	IDetailChildrenBuilder& AddChildCustomBuilder(TSharedRef<IDetailCustomNodeBuilder> InCustomBuilder)
@@ -144,7 +166,7 @@ public:
 	UE_DEPRECATED(4.17, "AddChildStructure has been deprecated.  Use AddExternalStructureProperty instead ")
 	class IDetailPropertyRow* AddChildStructure(TSharedRef<FStructOnScope> ChildStructure, FName PropertyName, FName UniqueIdName)
 	{
-		return AddExternalStructureProperty(ChildStructure, PropertyName, UniqueIdName);
+		return AddExternalStructureProperty(ChildStructure, PropertyName, FAddPropertyParams().UniqueId(UniqueIdName));
 	}
 
 	UE_DEPRECATED(4.17, "This version of AddStructure has been deprecated.  Use AddAllExternalStructureProperties instead ")
