Fixed issue with optional bools getting assigned an invalid value.

diff --git a/Engine/Source/Editor/PropertyEditor/Public/DetailBuilderTypes.h b/Engine/Source/Editor/PropertyEditor/Public/DetailBuilderTypes.h
--- a/Engine/Source/Editor/PropertyEditor/Public/DetailBuilderTypes.h
+++ b/Engine/Source/Editor/PropertyEditor/Public/DetailBuilderTypes.h
@@ -21,7 +21,7 @@ struct FAddPropertyParams
 	 */
 	FAddPropertyParams& ForceShowProperty()
 	{
-		bForceShowProperty = 1;
+		bForceShowProperty = true;
 		return *this;
 	}
 
@@ -31,7 +31,7 @@ struct FAddPropertyParams
 	 */
 	FAddPropertyParams& AllowChildren(bool bAllowChildren)
 	{
-		bAllowChildrenOverride = (bAllowChildren ? 1 : 2);
+		bAllowChildrenOverride = bAllowChildren;
 		return *this;
 	}
 
@@ -41,7 +41,7 @@ struct FAddPropertyParams
 	 */
 	FAddPropertyParams& CreateCategoryNodes(bool bCreateCategoryNodes)
 	{
-		bCreateCategoryNodesOverride = (bCreateCategoryNodes ? 1 : 2);
+		bCreateCategoryNodesOverride = bCreateCategoryNodes;
 		return *this;
 	}
 
@@ -61,7 +61,7 @@ struct FAddPropertyParams
 	 */
 	bool ShouldForcePropertyVisible() const
 	{
-		return bForceShowProperty == 1;
+		return bForceShowProperty;
 	}
 
 
