Prevent crash when SPropertyEditorAsset can't resolve an AllowedClasses metadata string, because the specified class is invalid or in a module not referenced by the current project.

diff --git a/Engine/Source/Editor/PropertyEditor/Private/UserInterface/PropertyEditor/SPropertyEditorAsset.cpp b/Engine/Source/Editor/PropertyEditor/Private/UserInterface/PropertyEditor/SPropertyEditorAsset.cpp
--- a/Engine/Source/Editor/PropertyEditor/Private/UserInterface/PropertyEditor/SPropertyEditorAsset.cpp
+++ b/Engine/Source/Editor/PropertyEditor/Private/UserInterface/PropertyEditor/SPropertyEditorAsset.cpp
@@ -158,7 +158,8 @@ void SPropertyEditorAsset::InitializeClassFilters(const UProperty* Property)
 			}
 		}
 	}
-	else
+	
+	if (AllowedClassFilters.Num() == 0)
 	{
 		// always add the object class to the filters
 		AllowedClassFilters.Add(ObjectClass);
@@ -240,8 +241,6 @@ void SPropertyEditorAsset::Construct(const FArguments& InArgs, const TSharedPtr<
 
 	InitializeClassFilters(Property);
 
-	ensure(AllowedClassFilters.Num() > 0);
-
 	// Make the ObjectClass more specific if we only have one class filter
 	// eg. if ObjectClass was set to Actor, but AllowedClasses="PointLight", we can limit it to PointLight immediately
 	if (AllowedClassFilters.Num() == 1 && DisallowedClassFilters.Num() == 0)
