Added handling of UObject pointers to edit conditions to enable equality, inequality and nullptr checks of UObject pointers in the form of: MyProperty != nullptr, or MyProperty == OtherProperty.

diff --git a/Engine/Source/Editor/PropertyEditor/Private/EditConditionContext.cpp b/Engine/Source/Editor/PropertyEditor/Private/EditConditionContext.cpp
--- a/Engine/Source/Editor/PropertyEditor/Private/EditConditionContext.cpp
+++ b/Engine/Source/Editor/PropertyEditor/Private/EditConditionContext.cpp
@@ -267,6 +267,52 @@ TOptional<FString> FEditConditionContext::GetEnumValue(const FString& PropertyNa
 	return TOptional<FString>();
 }
 
+TOptional<UObject*> FEditConditionContext::GetPointerValue(const FString& PropertyName) const
+{
+	const UProperty* Property = FindTypedField<UProperty>(PropertyNode, PropertyName);
+	if (Property == nullptr)
+	{
+		return TOptional<UObject*>();
+	}
+
+	const UObjectPropertyBase* ObjectProperty = Cast<UObjectPropertyBase>(Property);
+	if (ObjectProperty == nullptr)
+	{
+		return TOptional<UObject*>();
+	}
+
+	TSharedPtr<FPropertyNode> PinnedNode = PropertyNode.Pin();
+	FPropertyNode* ParentNode = GetEditConditionParentNode(PinnedNode);
+	FComplexPropertyNode* ComplexParentNode = PinnedNode->FindComplexParent();
+
+	TOptional<UObject*> Result;
+	for (int32 Index = 0; Index < ComplexParentNode->GetInstancesNum(); ++Index)
+	{
+		uint8* BasePtr = ComplexParentNode->GetMemoryOfInstance(Index);
+		if (BasePtr == nullptr)
+		{
+			// deleted or not fully loaded yet
+			return TOptional<UObject*>();
+		}
+
+		uint8* ParentPtr = ParentNode->GetValueAddress(BasePtr, PinnedNode->HasNodeFlags(EPropertyNodeFlags::IsSparseClassData) != 0);
+		uint8* ValuePtr = ComplexParentNode->GetValuePtrOfInstance(Index, Property, ParentNode);
+
+		UObject* Value = ObjectProperty->GetObjectPropertyValue(ValuePtr);
+		if (!Result.IsSet())
+		{
+			Result = Value;
+		}
+		else if (Result.GetValue() != Value)
+		{
+			// all values aren't the same
+			return TOptional<UObject*>();
+		}
+	}
+
+	return Result;
+}
+
 TOptional<FString> FEditConditionContext::GetTypeName(const FString& PropertyName) const
 {
 	const UProperty* Property = FindTypedField<UProperty>(PropertyNode, PropertyName);
diff --git a/Engine/Source/Editor/PropertyEditor/Private/EditConditionContext.h b/Engine/Source/Editor/PropertyEditor/Private/EditConditionContext.h
--- a/Engine/Source/Editor/PropertyEditor/Private/EditConditionContext.h
+++ b/Engine/Source/Editor/PropertyEditor/Private/EditConditionContext.h
@@ -10,6 +10,7 @@ public:
 	virtual TOptional<bool> GetBoolValue(const FString& PropertyName) const = 0;
 	virtual TOptional<double> GetNumericValue(const FString& PropertyName) const = 0;
 	virtual TOptional<FString> GetEnumValue(const FString& PropertyName) const = 0;
+	virtual TOptional<UObject*> GetPointerValue(const FString& PropertyName) const = 0;
 	virtual TOptional<FString> GetTypeName(const FString& PropertyName) const = 0;
 };
 
@@ -27,6 +28,7 @@ public:
 	virtual TOptional<bool> GetBoolValue(const FString& PropertyName) const override;
 	virtual TOptional<double> GetNumericValue(const FString& PropertyName) const override;
 	virtual TOptional<FString> GetEnumValue(const FString& PropertyName) const override;
+	virtual TOptional<UObject*> GetPointerValue(const FString& PropertyName) const override;
 	virtual TOptional<FString> GetTypeName(const FString& PropertyName) const override;
 
 	/**
diff --git a/Engine/Source/Editor/PropertyEditor/Private/EditConditionParser.cpp b/Engine/Source/Editor/PropertyEditor/Private/EditConditionParser.cpp
--- a/Engine/Source/Editor/PropertyEditor/Private/EditConditionParser.cpp
+++ b/Engine/Source/Editor/PropertyEditor/Private/EditConditionParser.cpp
@@ -29,44 +29,27 @@ static const TCHAR PropertyBreakingChars[] = { '|', '=', '&', '>', '<', '!', '+'
 
 static TOptional<FExpressionError> ConsumeBool(FExpressionTokenConsumer& Consumer)
 {
-	int TrueChars = 0;
-	int FalseChars = 0;
-	TOptional<FStringToken> StringToken = Consumer.GetStream().ParseToken([&TrueChars, &FalseChars](TCHAR InC)
+	TOptional<FStringToken> TrueToken = Consumer.GetStream().ParseToken(TEXT("true"));
+	if (TrueToken.IsSet())
 	{
-		if (TEXT("true")[TrueChars] == InC)
-		{
-			++TrueChars;
-			if (TrueChars == 4)
-			{
-				return EParseState::StopAfter;
-			}
-			else
-			{
-				return EParseState::Continue;
-			}
-		}
-		else if (TEXT("false")[FalseChars] == InC)
-		{
-			++FalseChars;
-			if (FalseChars == 5)
-			{
-				return EParseState::StopAfter;
-			}
-			else
-			{
-				return EParseState::Continue;
-			}
-		}
-		return EParseState::Cancel;
-	});
+		Consumer.Add(TrueToken.GetValue(), true);
+	}
 
-	if (StringToken.IsSet())
+	TOptional<FStringToken> FalseToken = Consumer.GetStream().ParseToken(TEXT("false"));
+	if (FalseToken.IsSet())
 	{
-		bool bIsTrue = TrueChars == 4;
-		bool bIsFalse = FalseChars == 5;
-		ensure(bIsTrue || bIsFalse);
+		Consumer.Add(FalseToken.GetValue(), false);
+	}
 
-		Consumer.Add(StringToken.GetValue(), bIsTrue);
+	return TOptional<FExpressionError>();
+}
+
+static TOptional<FExpressionError> ConsumeNullPtr(FExpressionTokenConsumer& Consumer)
+{
+	TOptional<FStringToken> NullToken = Consumer.GetStream().ParseToken(TEXT("nullptr"));
+	if (NullToken.IsSet())
+	{
+		Consumer.Add(NullToken.GetValue(), EditConditionParserTokens::FNullPtrToken());
 	}
 
 	return TOptional<FExpressionError>();
@@ -203,8 +186,41 @@ FExpressionResult ApplyBinary(TOperand<T> A, TOperand<T> B, Function Apply)
 	return MakeValue(Apply(ValueA.GetValue(), ValueB.GetValue()));
 }
 
+static FExpressionResult ApplyPropertyIsNull(const EditConditionParserTokens::FPropertyToken& Property, const IEditConditionContext& Context, bool bNegate)
+{
+	TOptional<FString> TypeName = Context.GetTypeName(Property.PropertyName);
+	if (!TypeName.IsSet())
+	{
+		return MakeError(FText::Format(LOCTEXT("InvalidOperand", "EditCondition attempted to use an invalid operand \"{0}\"."), FText::FromString(Property.PropertyName)));
+	}
+
+	const FString& TypeNameValue = TypeName.GetValue();
+	
+	TOptional<UObject*> Ptr = Context.GetPointerValue(Property.PropertyName);
+	if (!Ptr.IsSet())
+	{
+		return MakeError(FText::Format(LOCTEXT("InvalidOperand", "EditCondition attempted to use an invalid operand \"{0}\"."), FText::FromString(Property.PropertyName)));
+	}
+
+	bool bIsNull = Ptr.GetValue() == nullptr;
+	if (bNegate)
+	{
+		bIsNull = !bIsNull;
+	}
+
+	return MakeValue(bIsNull);
+}
+
 static FExpressionResult ApplyPropertiesEqual(const EditConditionParserTokens::FPropertyToken& A, const EditConditionParserTokens::FPropertyToken& B, const IEditConditionContext& Context, bool bNegate)
 {
+	TOptional<UObject*> PtrA = Context.GetPointerValue(A.PropertyName);
+	TOptional<UObject*> PtrB = Context.GetPointerValue(B.PropertyName);
+	if (PtrA.IsSet() && PtrB.IsSet())
+	{
+		const bool bAreEqual = PtrA.GetValue() == PtrB.GetValue();
+		return MakeValue(bNegate ? !bAreEqual : bAreEqual);
+	}
+
 	TOptional<FString> TypeNameA = Context.GetTypeName(A.PropertyName);
 	TOptional<FString> TypeNameB = Context.GetTypeName(B.PropertyName);
 	if (!TypeNameA.IsSet())
@@ -596,6 +612,7 @@ FEditConditionParser::FEditConditionParser()
 	TokenDefinitions.DefineToken(&ExpressionParser::ConsumeSymbol<FMultiply>);
 	TokenDefinitions.DefineToken(&ExpressionParser::ConsumeSymbol<FDivide>);
 	TokenDefinitions.DefineToken(&ExpressionParser::ConsumeNumber);
+	TokenDefinitions.DefineToken(&ConsumeNullPtr);
 	TokenDefinitions.DefineToken(&ConsumeBool);
 	TokenDefinitions.DefineToken(&ConsumePropertyName);
 				
@@ -623,6 +640,16 @@ FEditConditionParser::FEditConditionParser()
 		return ApplyPropertiesEqual(A, B, *Context, true);
 	});
 
+	OperatorJumpTable.MapBinary<FEqual>([](const FPropertyToken& A, const FNullPtrToken& B, const IEditConditionContext* Context) -> FExpressionResult
+	{
+		return ApplyPropertyIsNull(A, *Context, false);
+	});
+
+	OperatorJumpTable.MapBinary<FNotEqual>([](const FPropertyToken& A, const FNullPtrToken& B, const IEditConditionContext* Context) -> FExpressionResult
+	{
+		return ApplyPropertyIsNull(A, *Context, true);
+	});
+
 	CreateBooleanOperators(OperatorJumpTable);
 	CreateNumberOperators<double>(OperatorJumpTable);
 	CreateEnumOperators(OperatorJumpTable);
diff --git a/Engine/Source/Editor/PropertyEditor/Private/EditConditionParser.h b/Engine/Source/Editor/PropertyEditor/Private/EditConditionParser.h
--- a/Engine/Source/Editor/PropertyEditor/Private/EditConditionParser.h
+++ b/Engine/Source/Editor/PropertyEditor/Private/EditConditionParser.h
@@ -44,6 +44,13 @@ namespace EditConditionParserTokens
 		FString Type;
 		FString Value;
 	};
+
+	struct FNullPtrToken
+	{
+		FNullPtrToken() = default;
+		FNullPtrToken(const FNullPtrToken&) = default;
+		FNullPtrToken& operator=(const FNullPtrToken&) = default;
+	};
 }
 
 #define DEFINE_EDIT_CONDITION_NODE(TYPE, ...) \
@@ -69,6 +76,7 @@ namespace EditConditionParserTokens
 	DEFINE_EXPRESSION_NODE_TYPE(bool, 0xCACBC715, 0x505A6B4A, 0x8808809F, 0x897AA5F6)
 	DEFINE_EXPRESSION_NODE_TYPE(EditConditionParserTokens::FPropertyToken, 0x9A3FAF6F, 0xB2E45E4D, 0xA80A70C6, 0x47A89BD7)
 	DEFINE_EXPRESSION_NODE_TYPE(EditConditionParserTokens::FEnumToken, 0xC9A35C24, 0x21FC904B, 0x9F1B2B6A, 0xDF6F4BC4)
+	DEFINE_EXPRESSION_NODE_TYPE(EditConditionParserTokens::FNullPtrToken, 0xC3E91735, 0xB3B90C4D, 0x99BA3DCA, 0xC6CC6231)
 
 #undef DEFINE_EDIT_CONDITION_NODE
 
diff --git a/Engine/Source/Editor/PropertyEditor/Private/EditConditionParserTests.cpp b/Engine/Source/Editor/PropertyEditor/Private/EditConditionParserTests.cpp
--- a/Engine/Source/Editor/PropertyEditor/Private/EditConditionParserTests.cpp
+++ b/Engine/Source/Editor/PropertyEditor/Private/EditConditionParserTests.cpp
@@ -52,6 +52,11 @@ struct FTestEditConditionContext : IEditConditionContext
 		return Result;
 	}
 
+	virtual TOptional<UObject*> GetPointerValue(const FString& PropertyName) const override
+	{
+		return TOptional<UObject*>();
+	}
+
 	virtual TOptional<FString> GetTypeName(const FString& PropertyName) const override
 	{
 		TOptional<FString> Result;
@@ -114,8 +119,7 @@ static bool CanParse(const FEditConditionParser& Parser, const FString& Expressi
 		}
 	 }
 
-	 return Parsed->Tokens.Num() == ExpectedTokens &&
-		 PropertyCount == ExpectedProperties;
+	 return Parsed->Tokens.Num() == ExpectedTokens && PropertyCount == ExpectedProperties;
 }
 
 IMPLEMENT_SIMPLE_AUTOMATION_TEST(FEditConditionParser_Parse, "EditConditionParser.Parse", EAutomationTestFlags::EditorContext | EAutomationTestFlags::EngineFilter)
@@ -148,6 +152,8 @@ bool FEditConditionParser_Parse::RunTest(const FString& Parameters)
 	bResult &= CanParse(Parser, TEXT("Enum != EType::Value || BoolProperty == bFoo"), 7, 3);
 	bResult &= CanParse(Parser, TEXT("Enum == EType::Value && Foo != 5"), 7, 2);
 	bResult &= CanParse(Parser, TEXT("Enum != EType::Value && Foo == Bar"), 7, 3);
+	bResult &= CanParse(Parser, TEXT("PointerProperty == nullptr"), 3, 1);
+	bResult &= CanParse(Parser, TEXT("PointerProperty != nullptr"), 3, 1);
 
 	return bResult;
 }
@@ -404,6 +410,85 @@ bool FEditConditionParser_EvaluateUObject::RunTest(const FString& Parameters)
 	return bAllResults;
 }
 
+IMPLEMENT_SIMPLE_AUTOMATION_TEST(FEditConditionParser_EvaluatePointers, "EditConditionParser.EvaluatePointers", EAutomationTestFlags::EditorContext | EAutomationTestFlags::EngineFilter)
+bool FEditConditionParser_EvaluatePointers::RunTest(const FString& Parameters)
+{
+	UEditConditionTestObject* TestObject = NewObject<UEditConditionTestObject>();
+	TestObject->AddToRoot();
+
+	TSharedPtr<FObjectPropertyNode> ObjectNode(new FObjectPropertyNode);
+	ObjectNode->AddObject(TestObject);
+
+	FPropertyNodeInitParams InitParams;
+	ObjectNode->InitNode(InitParams);
+
+	static const FName BoolPropertyName(TEXT("BoolProperty"));
+	TSharedPtr<FPropertyNode> PropertyNode = ObjectNode->FindChildPropertyNode(BoolPropertyName, true);
+	FEditConditionContext Context(*PropertyNode.Get());
+
+	bool bAllResults = true;
+
+	FEditConditionParser Parser;
+
+	{
+		TestObject->UObjectPtr = nullptr;
+
+		bAllResults &= CanEvaluate(Parser, Context, TEXT("UObjectPtr != nullptr"), false);
+		bAllResults &= CanEvaluate(Parser, Context, TEXT("UObjectPtr == nullptr"), true);
+
+		TestObject->UObjectPtr = TestObject;
+
+		bAllResults &= CanEvaluate(Parser, Context, TEXT("UObjectPtr != nullptr"), true);
+		bAllResults &= CanEvaluate(Parser, Context, TEXT("UObjectPtr == nullptr"), false);
+	}
+
+	{
+		TestObject->SoftClassPtr = nullptr;
+
+		bAllResults &= CanEvaluate(Parser, Context, TEXT("SoftClassPtr != nullptr"), false);
+		bAllResults &= CanEvaluate(Parser, Context, TEXT("SoftClassPtr == nullptr"), true);
+
+		TestObject->SoftClassPtr = TestObject;
+
+		bAllResults &= CanEvaluate(Parser, Context, TEXT("SoftClassPtr != nullptr"), true);
+		bAllResults &= CanEvaluate(Parser, Context, TEXT("SoftClassPtr == nullptr"), false);
+	}
+
+	{
+		TestObject->WeakObjectPtr = nullptr;
+
+		bAllResults &= CanEvaluate(Parser, Context, TEXT("WeakObjectPtr != nullptr"), false);
+		bAllResults &= CanEvaluate(Parser, Context, TEXT("WeakObjectPtr == nullptr"), true);
+
+		TestObject->WeakObjectPtr = TestObject;
+
+		bAllResults &= CanEvaluate(Parser, Context, TEXT("WeakObjectPtr != nullptr"), true);
+		bAllResults &= CanEvaluate(Parser, Context, TEXT("WeakObjectPtr == nullptr"), false);
+	}
+
+	{
+		// equality & inequality
+		bAllResults &= CanEvaluate(Parser, Context, TEXT("UObjectPtr == SoftClassPtr"), true);
+		bAllResults &= CanEvaluate(Parser, Context, TEXT("SoftClassPtr == UObjectPtr"), true);
+
+		bAllResults &= CanEvaluate(Parser, Context, TEXT("WeakObjectPtr != UObjectPtr"), false);
+		bAllResults &= CanEvaluate(Parser, Context, TEXT("UObjectPtr != WeakObjectPtr"), false);
+
+		bAllResults &= CanEvaluate(Parser, Context, TEXT("SoftClassPtr == WeakObjectPtr"), true);
+		bAllResults &= CanEvaluate(Parser, Context, TEXT("WeakObjectPtr == UObjectPtr"), true);
+
+		TestObject->UObjectPtr = nullptr;
+
+		bAllResults &= CanEvaluate(Parser, Context, TEXT("UObjectPtr != SoftClassPtr"), true);
+		bAllResults &= CanEvaluate(Parser, Context, TEXT("WeakObjectPtr != UObjectPtr"), true);
+		bAllResults &= CanEvaluate(Parser, Context, TEXT("WeakObjectPtr == UObjectPtr"), false);
+	}
+
+	TestObject->RemoveFromRoot();
+
+	return bAllResults;
+}
+
 IMPLEMENT_SIMPLE_AUTOMATION_TEST(FEditConditionParser_SingleBool, "EditConditionParser.SingleBool", EAutomationTestFlags::EditorContext | EAutomationTestFlags::EngineFilter)
 bool FEditConditionParser_SingleBool::RunTest(const FString& Parameters)
 {
@@ -434,6 +519,8 @@ bool FEditConditionParser_SingleBool::RunTest(const FString& Parameters)
 		TestNotNull(TEXT("Uint Bitfield"), Property);
 	}
 
+	TestObject->RemoveFromRoot();
+
 	return true;
 }
 
diff --git a/Engine/Source/Editor/PropertyEditor/Private/EditConditionParserTests.h b/Engine/Source/Editor/PropertyEditor/Private/EditConditionParserTests.h
--- a/Engine/Source/Editor/PropertyEditor/Private/EditConditionParserTests.h
+++ b/Engine/Source/Editor/PropertyEditor/Private/EditConditionParserTests.h
@@ -41,4 +41,13 @@ class UEditConditionTestObject : public UObject
 
 	UPROPERTY(EditAnywhere, Category=Test)
 	uint8 UintBitfieldProperty : 1;
+
+	UPROPERTY(EditAnywhere, Category=Test)
+	UObject* UObjectPtr;
+
+	UPROPERTY(EditAnywhere, Category=Test)
+	TSoftClassPtr<UObject> SoftClassPtr;
+
+	UPROPERTY(EditAnywhere, Category=Test)
+	TWeakObjectPtr<UObject> WeakObjectPtr;
 };
