EditCondition now supports bitflags that use enums. Eg. MyIntProperty & MyEnum::Flag (and negation by tacking on == false).

Flags must be enums with integral values. Powers of two are not checked for to allow for the use of combinations of flags, eg. Foo = 1, Bar = 2, FooAndBar = Foo | Bar.

diff --git a/Engine/Source/Editor/PropertyEditor/Private/EditConditionContext.cpp b/Engine/Source/Editor/PropertyEditor/Private/EditConditionContext.cpp
--- a/Engine/Source/Editor/PropertyEditor/Private/EditConditionContext.cpp
+++ b/Engine/Source/Editor/PropertyEditor/Private/EditConditionContext.cpp
@@ -97,6 +97,25 @@ FPropertyNode* GetEditConditionParentNode(const TSharedPtr<FPropertyNode>& Prope
 	return ParentNode;
 }
 
+static uint8* GetPropertyValuePtr(const UProperty* Property, const TSharedPtr<FPropertyNode>& PropertyNode, FPropertyNode* ParentNode, FComplexPropertyNode* ComplexParentNode, int32 Index)
+{
+	uint8* BasePtr = ComplexParentNode->GetMemoryOfInstance(Index);
+	if (BasePtr == nullptr)
+	{
+		return nullptr;
+	}
+
+	uint8* ParentPtr = ParentNode->GetValueAddress(BasePtr, PropertyNode->HasNodeFlags(EPropertyNodeFlags::IsSparseClassData) != 0);
+	uint8* ValuePtr = ComplexParentNode->GetValuePtrOfInstance(Index, Property, ParentNode);
+
+	// SPARSEDATA_TODO: remove the next three lines once we're sure the pointer math is correct
+	uint8* OldValuePtr = Property->ContainerPtrToValuePtr<uint8>(ParentPtr);
+	check(OldValuePtr == ValuePtr || PropertyNode->HasNodeFlags(EPropertyNodeFlags::IsSparseClassData));
+	ensure(ValuePtr != nullptr);
+
+	return ValuePtr;
+}
+
 TOptional<bool> FEditConditionContext::GetBoolValue(const FString& PropertyName) const
 {
 	const UBoolProperty* BoolProperty = FindTypedField<UBoolProperty>(PropertyNode, PropertyName);
@@ -112,22 +131,12 @@ TOptional<bool> FEditConditionContext::GetBoolValue(const FString& PropertyName)
 	TOptional<bool> Result;
 	for (int32 Index = 0; Index < ComplexParentNode->GetInstancesNum(); ++Index)
 	{
-		uint8* BasePtr = ComplexParentNode->GetMemoryOfInstance(Index);
-		if (BasePtr == nullptr)
+		uint8* ValuePtr = GetPropertyValuePtr(BoolProperty, PinnedNode, ParentNode, ComplexParentNode, Index);
+		if (ValuePtr == nullptr)
 		{
-			// deleted or not fully loaded yet
 			return TOptional<bool>();
 		}
 
-		uint8* ParentPtr = ParentNode->GetValueAddress(BasePtr, PinnedNode->HasNodeFlags(EPropertyNodeFlags::IsSparseClassData) != 0);
-
-		uint8* ValuePtr = ComplexParentNode->GetValuePtrOfInstance(Index, BoolProperty, ParentNode);
-
-		// SPARSEDATA_TODO: remove the next three lines once we're sure the pointer math is correct
-		uint8* OldValuePtr = BoolProperty->ContainerPtrToValuePtr<uint8>(ParentPtr);
-		check(OldValuePtr == ValuePtr || PinnedNode->HasNodeFlags(EPropertyNodeFlags::IsSparseClassData));
-		ensure(ValuePtr != nullptr);
-
 		bool bValue = BoolProperty->GetPropertyValue(ValuePtr);
 		if (!Result.IsSet())
 		{
@@ -143,6 +152,47 @@ TOptional<bool> FEditConditionContext::GetBoolValue(const FString& PropertyName)
 	return Result;
 }
 
+TOptional<int64> FEditConditionContext::GetIntegerValue(const FString& PropertyName) const
+{
+	const UNumericProperty* NumericProperty = FindTypedField<UNumericProperty>(PropertyNode, PropertyName);
+	if (NumericProperty == nullptr)
+	{
+		return TOptional<int64>();
+	}
+
+	if (!NumericProperty->IsInteger())
+	{
+		return TOptional<int64>();
+	}
+
+	TSharedPtr<FPropertyNode> PinnedNode = PropertyNode.Pin();
+	FPropertyNode* ParentNode = GetEditConditionParentNode(PinnedNode);
+	FComplexPropertyNode* ComplexParentNode = PinnedNode->FindComplexParent();
+
+	TOptional<int64> Result;
+	for (int32 Index = 0; Index < ComplexParentNode->GetInstancesNum(); ++Index)
+	{
+		uint8* ValuePtr = GetPropertyValuePtr(NumericProperty, PinnedNode, ParentNode, ComplexParentNode, Index);
+		if (ValuePtr == nullptr)
+		{
+			return TOptional<int64>();
+		}
+
+		int64 Value = NumericProperty->GetSignedIntPropertyValue(ValuePtr);
+		if (!Result.IsSet())
+		{
+			Result = Value;
+		}
+		else if (Result.GetValue() != Value)
+		{
+			// all values aren't the same...
+			return TOptional<int64>();
+		}
+	}
+
+	return Result;
+}
+
 TOptional<double> FEditConditionContext::GetNumericValue(const FString& PropertyName) const
 {
 	const UNumericProperty* NumericProperty = FindTypedField<UNumericProperty>(PropertyNode, PropertyName);
@@ -158,22 +208,12 @@ TOptional<double> FEditConditionContext::GetNumericValue(const FString& Property
 	TOptional<double> Result;
 	for (int32 Index = 0; Index < ComplexParentNode->GetInstancesNum(); ++Index)
 	{
-		uint8* BasePtr = ComplexParentNode->GetMemoryOfInstance(Index);
-		if (BasePtr == nullptr)
+		uint8* ValuePtr = GetPropertyValuePtr(NumericProperty, PinnedNode, ParentNode, ComplexParentNode, Index);
+		if (ValuePtr == nullptr)
 		{
-			// deleted or not fully loaded yet
 			return TOptional<double>();
 		}
 
-		uint8* ParentPtr = ParentNode->GetValueAddress(BasePtr, PinnedNode->HasNodeFlags(EPropertyNodeFlags::IsSparseClassData) != 0);
-
-		uint8* ValuePtr = ComplexParentNode->GetValuePtrOfInstance(Index, NumericProperty, ParentNode);
-
-		// SPARSEDATA_TODO: remove the next three lines once we're sure the pointer math is correct
-		uint8* OldValuePtr = NumericProperty->ContainerPtrToValuePtr<uint8>(ParentPtr);
-		check(OldValuePtr == ValuePtr || PinnedNode->HasNodeFlags(EPropertyNodeFlags::IsSparseClassData));
-		ensure(ValuePtr != nullptr);
-
 		double Value = 0;
 
 		if (NumericProperty->IsInteger())
@@ -223,7 +263,12 @@ TOptional<FString> FEditConditionContext::GetEnumValue(const FString& PropertyNa
 	{
 		return TOptional<FString>();
 	}
-	
+
+	if (!NumericProperty->IsInteger())
+	{
+		return TOptional<FString>();
+	}
+
 	TSharedPtr<FPropertyNode> PinnedNode = PropertyNode.Pin();
 	FPropertyNode* ParentNode = GetEditConditionParentNode(PinnedNode);
 	FComplexPropertyNode* ComplexParentNode = PinnedNode->FindComplexParent();
@@ -231,22 +276,12 @@ TOptional<FString> FEditConditionContext::GetEnumValue(const FString& PropertyNa
 	TOptional<int64> Result;
 	for (int32 Index = 0; Index < ComplexParentNode->GetInstancesNum(); ++Index)
 	{
-		uint8* BasePtr = ComplexParentNode->GetMemoryOfInstance(Index);
-		if (BasePtr == nullptr)
+		uint8* ValuePtr = GetPropertyValuePtr(Property, PinnedNode, ParentNode, ComplexParentNode, Index);
+		if (ValuePtr == nullptr)
 		{
-			// deleted or not fully loaded yet
 			return TOptional<FString>();
 		}
 
-		uint8* ParentPtr = ParentNode->GetValueAddress(BasePtr, PinnedNode->HasNodeFlags(EPropertyNodeFlags::IsSparseClassData) != 0);
-
-		uint8* ValuePtr = ComplexParentNode->GetValuePtrOfInstance(Index, Property, ParentNode);
-
-		// SPARSEDATA_TODO: remove the next three lines once we're sure the pointer math is correct
-		uint8* OldValuePtr = Property->ContainerPtrToValuePtr<uint8>(ParentPtr);
-		check(OldValuePtr == ValuePtr || PinnedNode->HasNodeFlags(EPropertyNodeFlags::IsSparseClassData));
-		ensure(ValuePtr != nullptr);
-
 		int64 Value = NumericProperty->GetSignedIntPropertyValue(ValuePtr);
 		if (!Result.IsSet())
 		{
@@ -259,12 +294,12 @@ TOptional<FString> FEditConditionContext::GetEnumValue(const FString& PropertyNa
 		}
 	}
 
-	if (Result.IsSet())
+	if (!Result.IsSet())
 	{
-		return EnumType->GetNameStringByValue(Result.GetValue());
+		return TOptional<FString>();
 	}
 
-	return TOptional<FString>();
+	return EnumType->GetNameStringByValue(Result.GetValue());
 }
 
 TOptional<UObject*> FEditConditionContext::GetPointerValue(const FString& PropertyName) const
@@ -288,16 +323,12 @@ TOptional<UObject*> FEditConditionContext::GetPointerValue(const FString& Proper
 	TOptional<UObject*> Result;
 	for (int32 Index = 0; Index < ComplexParentNode->GetInstancesNum(); ++Index)
 	{
-		uint8* BasePtr = ComplexParentNode->GetMemoryOfInstance(Index);
-		if (BasePtr == nullptr)
+		uint8* ValuePtr = GetPropertyValuePtr(Property, PinnedNode, ParentNode, ComplexParentNode, Index);
+		if (ValuePtr == nullptr)
 		{
-			// deleted or not fully loaded yet
 			return TOptional<UObject*>();
 		}
 
-		uint8* ParentPtr = ParentNode->GetValueAddress(BasePtr, PinnedNode->HasNodeFlags(EPropertyNodeFlags::IsSparseClassData) != 0);
-		uint8* ValuePtr = ComplexParentNode->GetValuePtrOfInstance(Index, Property, ParentNode);
-
 		UObject* Value = ObjectProperty->GetObjectPropertyValue(ValuePtr);
 		if (!Result.IsSet())
 		{
@@ -331,4 +362,21 @@ TOptional<FString> FEditConditionContext::GetTypeName(const FString& PropertyNam
 	}
 
 	return Property->GetCPPType();
-}
\ No newline at end of file
+}
+
+TOptional<int64> FEditConditionContext::GetIntegerValueOfEnum(const FString& EnumTypeName, const FString& MemberName) const
+{
+	UEnum* EnumType = FindObject<UEnum>((UObject*) ANY_PACKAGE, *EnumTypeName, true);
+	if (EnumType == nullptr)
+	{
+		return TOptional<int64>();
+	}
+
+	int64 EnumValue = EnumType->GetValueByName(FName(*MemberName));
+	if (EnumValue == INDEX_NONE)
+	{
+		return TOptional<int64>();
+	}
+
+	return EnumValue;
+}
diff --git a/Engine/Source/Editor/PropertyEditor/Private/EditConditionContext.h b/Engine/Source/Editor/PropertyEditor/Private/EditConditionContext.h
--- a/Engine/Source/Editor/PropertyEditor/Private/EditConditionContext.h
+++ b/Engine/Source/Editor/PropertyEditor/Private/EditConditionContext.h
@@ -8,10 +8,12 @@ class IEditConditionContext
 {
 public:
 	virtual TOptional<bool> GetBoolValue(const FString& PropertyName) const = 0;
+	virtual TOptional<int64> GetIntegerValue(const FString& PropertyName) const = 0;
 	virtual TOptional<double> GetNumericValue(const FString& PropertyName) const = 0;
 	virtual TOptional<FString> GetEnumValue(const FString& PropertyName) const = 0;
 	virtual TOptional<UObject*> GetPointerValue(const FString& PropertyName) const = 0;
 	virtual TOptional<FString> GetTypeName(const FString& PropertyName) const = 0;
+	virtual TOptional<int64> GetIntegerValueOfEnum(const FString& EnumType, const FString& EnumValue) const = 0;
 };
 
 class UProperty;
@@ -26,10 +28,12 @@ public:
 	virtual ~FEditConditionContext() {}
 
 	virtual TOptional<bool> GetBoolValue(const FString& PropertyName) const override;
+	virtual TOptional<int64> GetIntegerValue(const FString& PropertyName) const override;
 	virtual TOptional<double> GetNumericValue(const FString& PropertyName) const override;
 	virtual TOptional<FString> GetEnumValue(const FString& PropertyName) const override;
 	virtual TOptional<UObject*> GetPointerValue(const FString& PropertyName) const override;
 	virtual TOptional<FString> GetTypeName(const FString& PropertyName) const override;
+	virtual TOptional<int64> GetIntegerValueOfEnum(const FString& EnumType, const FString& EnumValue) const override;
 
 	/**
 	 * Fetch the single boolean property referenced.
@@ -39,4 +43,4 @@ public:
 
 private:
 	TWeakPtr<FPropertyNode> PropertyNode;
-};
\ No newline at end of file
+};
diff --git a/Engine/Source/Editor/PropertyEditor/Private/EditConditionParser.cpp b/Engine/Source/Editor/PropertyEditor/Private/EditConditionParser.cpp
--- a/Engine/Source/Editor/PropertyEditor/Private/EditConditionParser.cpp
+++ b/Engine/Source/Editor/PropertyEditor/Private/EditConditionParser.cpp
@@ -5,6 +5,7 @@
 
 #include "Math/BasicMathExpressionEvaluator.h"
 #include "Misc/ExpressionParser.h"
+#include "UObject/Class.h"
 
 #define LOCTEXT_NAMESPACE "EditConditionParser"
 
@@ -23,6 +24,7 @@ namespace EditConditionParserTokens
 	const TCHAR* const FSubtract::Moniker = TEXT("-");
 	const TCHAR* const FMultiply::Moniker = TEXT("*");
 	const TCHAR* const FDivide::Moniker = TEXT("/");
+	const TCHAR* const FBitwiseAnd::Moniker = TEXT("&");
 }
 
 static const TCHAR PropertyBreakingChars[] = { '|', '=', '&', '>', '<', '!', '+', '-', '*', '/', ' ', '\t' };
@@ -186,6 +188,23 @@ FExpressionResult ApplyBinary(TOperand<T> A, TOperand<T> B, Function Apply)
 	return MakeValue(Apply(ValueA.GetValue(), ValueB.GetValue()));
 }
 
+static FExpressionResult ApplyBitwiseAnd(const EditConditionParserTokens::FPropertyToken& Property, const EditConditionParserTokens::FEnumToken& Enum, const IEditConditionContext& Context)
+{
+	TOptional<int64> EnumValue = Context.GetIntegerValueOfEnum(Enum.Type, Enum.Value);
+	if (!EnumValue.IsSet())
+	{
+		return MakeError(FText::Format(LOCTEXT("InvalidEnumValue", "EditCondition attempted to use an invalid enum value \"{0}::{1}\"."), FText::FromString(Enum.Type), FText::FromString(Enum.Value)));
+	}
+
+	TOptional<int64> PropertyValue = Context.GetIntegerValue(Property.PropertyName);
+	if (!PropertyValue.IsSet())
+	{
+		return MakeError(FText::Format(LOCTEXT("InvalidOperand", "EditCondition attempted to use an invalid operand \"{0}\"."), FText::FromString(Property.PropertyName)));
+	}
+
+	return MakeValue((PropertyValue.Get(0) & EnumValue.Get(0)) != 0);
+}
+
 static FExpressionResult ApplyPropertyIsNull(const EditConditionParserTokens::FPropertyToken& Property, const IEditConditionContext& Context, bool bNegate)
 {
 	TOptional<FString> TypeName = Context.GetTypeName(Property.PropertyName);
@@ -611,6 +630,7 @@ FEditConditionParser::FEditConditionParser()
 	TokenDefinitions.DefineToken(&ExpressionParser::ConsumeSymbol<FSubtract>);
 	TokenDefinitions.DefineToken(&ExpressionParser::ConsumeSymbol<FMultiply>);
 	TokenDefinitions.DefineToken(&ExpressionParser::ConsumeSymbol<FDivide>);
+	TokenDefinitions.DefineToken(&ExpressionParser::ConsumeSymbol<FBitwiseAnd>);
 	TokenDefinitions.DefineToken(&ExpressionParser::ConsumeNumber);
 	TokenDefinitions.DefineToken(&ConsumeNullPtr);
 	TokenDefinitions.DefineToken(&ConsumeBool);
@@ -624,12 +644,14 @@ FEditConditionParser::FEditConditionParser()
 	ExpressionGrammar.DefineBinaryOperator<FLessEqual>(3);
 	ExpressionGrammar.DefineBinaryOperator<FGreater>(3);
 	ExpressionGrammar.DefineBinaryOperator<FGreaterEqual>(3);
+	ExpressionGrammar.DefineBinaryOperator<FBitwiseAnd>(2);
 	ExpressionGrammar.DefineBinaryOperator<FAdd>(2);
 	ExpressionGrammar.DefineBinaryOperator<FSubtract>(2);
 	ExpressionGrammar.DefineBinaryOperator<FMultiply>(1);
 	ExpressionGrammar.DefineBinaryOperator<FDivide>(1);
 	ExpressionGrammar.DefinePreUnaryOperator<FNot>();
 
+	// POINTER EQUALITY
 	OperatorJumpTable.MapBinary<FEqual>([](const FPropertyToken& A, const FPropertyToken& B, const IEditConditionContext* Context) -> FExpressionResult
 	{
 		return ApplyPropertiesEqual(A, B, *Context, false);
@@ -640,6 +662,7 @@ FEditConditionParser::FEditConditionParser()
 		return ApplyPropertiesEqual(A, B, *Context, true);
 	});
 
+	// POINTER NULL
 	OperatorJumpTable.MapBinary<FEqual>([](const FPropertyToken& A, const FNullPtrToken& B, const IEditConditionContext* Context) -> FExpressionResult
 	{
 		return ApplyPropertyIsNull(A, *Context, false);
@@ -650,6 +673,12 @@ FEditConditionParser::FEditConditionParser()
 		return ApplyPropertyIsNull(A, *Context, true);
 	});
 
+	// BITWISE AND
+	OperatorJumpTable.MapBinary<FBitwiseAnd>([](const FPropertyToken& A, const FEnumToken& B, const IEditConditionContext* Context) -> FExpressionResult
+	{
+		return ApplyBitwiseAnd(A, B, *Context);
+	});
+
 	CreateBooleanOperators(OperatorJumpTable);
 	CreateNumberOperators<double>(OperatorJumpTable);
 	CreateEnumOperators(OperatorJumpTable);
@@ -699,4 +728,4 @@ TSharedPtr<FEditConditionExpression> FEditConditionParser::Parse(const FString&
 	return TSharedPtr<FEditConditionExpression>();
 }
 
-#undef LOCTEXT_NAMESPACE
\ No newline at end of file
+#undef LOCTEXT_NAMESPACE
diff --git a/Engine/Source/Editor/PropertyEditor/Private/EditConditionParser.h b/Engine/Source/Editor/PropertyEditor/Private/EditConditionParser.h
--- a/Engine/Source/Editor/PropertyEditor/Private/EditConditionParser.h
+++ b/Engine/Source/Editor/PropertyEditor/Private/EditConditionParser.h
@@ -70,6 +70,7 @@ namespace EditConditionParserTokens
 	DEFINE_EDIT_CONDITION_NODE(FSubtract, 0xE80D53E9, 0x32F1E145, 0x8535B767, 0x0634DE9E)
 	DEFINE_EDIT_CONDITION_NODE(FMultiply, 0x0D4D37DB, 0x53C7EA41, 0x8297C200, 0xF5B28D27)
 	DEFINE_EDIT_CONDITION_NODE(FDivide, 0x48009789, 0x2EE84F40, 0x83467A99, 0x92617168)
+	DEFINE_EDIT_CONDITION_NODE(FBitwiseAnd, 0xF5335505, 0xACB75044, 0xAA468084, 0x1042ADE8)
 	DEFINE_EDIT_CONDITION_NODE(FSubExpressionStart, 0x78CE8411, 0x34076241, 0xA2BD0548, 0x54458A3A)
 	DEFINE_EDIT_CONDITION_NODE(FSubExpressionEnd, 0x5CB25466, 0x780CAE4D, 0x8A88E000, 0xD3695885)
 
diff --git a/Engine/Source/Editor/PropertyEditor/Private/EditConditionParserTests.cpp b/Engine/Source/Editor/PropertyEditor/Private/EditConditionParserTests.cpp
--- a/Engine/Source/Editor/PropertyEditor/Private/EditConditionParserTests.cpp
+++ b/Engine/Source/Editor/PropertyEditor/Private/EditConditionParserTests.cpp
@@ -15,41 +15,61 @@ UEditConditionTestObject::UEditConditionTestObject(const FObjectInitializer& Obj
 struct FTestEditConditionContext : IEditConditionContext
 {
 	TMap<FString, bool> BoolValues;
+	TMap<FString, int64> IntegerValues;
 	TMap<FString, double> DoubleValues;
 	TMap<FString, FString> EnumValues;
+
 	FString EnumTypeName;
+	TMap<FString, int64> EnumTypeValues;
 
 	FTestEditConditionContext(){}
 	virtual ~FTestEditConditionContext() {}
 
 	virtual TOptional<bool> GetBoolValue(const FString& PropertyName) const override
 	{
-		TOptional<bool> Result;
 		if (const bool* Value = BoolValues.Find(PropertyName))
 		{
-			Result = *Value;
+			return *Value;
 		}
-		return Result;
+		return TOptional<bool>();
+	}
+
+	virtual TOptional<int64> GetIntegerValue(const FString& PropertyName) const override
+	{
+		if (const int64* Value = IntegerValues.Find(PropertyName))
+		{
+			return *Value;
+		}
+		return TOptional<int64>();
 	}
 
 	virtual TOptional<double> GetNumericValue(const FString& PropertyName) const override
 	{
-		TOptional<double> Result;
 		if (const double* Value = DoubleValues.Find(PropertyName))
 		{
-			Result = *Value;
+			return *Value;
 		}
-		return Result;
+		return TOptional<double>();
 	}
 
 	virtual TOptional<FString> GetEnumValue(const FString& PropertyName) const override
 	{
-		TOptional<FString> Result;
 		if (const FString* Value = EnumValues.Find(PropertyName))
 		{
-			Result = *Value;
+			return *Value;
 		}
-		return Result;
+		return TOptional<FString>();
+	}
+
+	virtual TOptional<int64> GetIntegerValueOfEnum(const FString& TypeName, const FString& ValueName) const override
+	{
+		check(TypeName == EnumTypeName);
+
+		if (const int64* Value = EnumTypeValues.Find(ValueName))
+		{
+			return *Value;
+		}
+		return TOptional<int64>();
 	}
 
 	virtual TOptional<UObject*> GetPointerValue(const FString& PropertyName) const override
@@ -82,6 +102,11 @@ struct FTestEditConditionContext : IEditConditionContext
 		BoolValues.Add(PropertyName, Value);
 	}
 
+	void SetupInteger(const FString& PropertyName, int64 Value)
+	{
+		IntegerValues.Add(PropertyName, Value);
+	}
+
 	void SetupDouble(const FString& PropertyName, double Value)
 	{
 		DoubleValues.Add(PropertyName, Value);
@@ -96,6 +121,11 @@ struct FTestEditConditionContext : IEditConditionContext
 	{
 		EnumTypeName = EnumType;
 	}
+
+	void SetupEnumTypeValue(const FString& Name, int64 Value)
+	{
+		EnumTypeValues.Add(Name, Value);
+	}
 };
 
 static bool CanParse(const FEditConditionParser& Parser, const FString& Expression, int32 ExpectedTokens, int32 ExpectedProperties)
@@ -119,7 +149,13 @@ static bool CanParse(const FEditConditionParser& Parser, const FString& Expressi
 		}
 	 }
 
-	 return Parsed->Tokens.Num() == ExpectedTokens && PropertyCount == ExpectedProperties;
+	 if (Parsed->Tokens.Num() != ExpectedTokens || PropertyCount != ExpectedProperties)
+	 {
+		 ensureMsgf(false, TEXT("Parsing produced an invalid number of tokens or properties. Expression: %s"), *Expression);
+		 return false;
+	 }
+
+	 return true;
 }
 
 IMPLEMENT_SIMPLE_AUTOMATION_TEST(FEditConditionParser_Parse, "EditConditionParser.Parse", EAutomationTestFlags::EditorContext | EAutomationTestFlags::EngineFilter)
@@ -154,6 +190,8 @@ bool FEditConditionParser_Parse::RunTest(const FString& Parameters)
 	bResult &= CanParse(Parser, TEXT("Enum != EType::Value && Foo == Bar"), 7, 3);
 	bResult &= CanParse(Parser, TEXT("PointerProperty == nullptr"), 3, 1);
 	bResult &= CanParse(Parser, TEXT("PointerProperty != nullptr"), 3, 1);
+	bResult &= CanParse(Parser, TEXT("Flags & EType::Value"), 3, 1);
+	bResult &= CanParse(Parser, TEXT("Flags & EType::Value == false"), 5, 1);
 
 	return bResult;
 }
@@ -280,7 +318,6 @@ static bool RunNumericTests(const IEditConditionContext& Context)
 	return bResult;
 }
 
-
 IMPLEMENT_SIMPLE_AUTOMATION_TEST(FEditConditionParser_EvaluateDouble, "EditConditionParser.EvaluateDouble", EAutomationTestFlags::EditorContext | EAutomationTestFlags::EngineFilter)
 bool FEditConditionParser_EvaluateDouble::RunTest(const FString& Parameters)
 {
@@ -329,6 +366,72 @@ bool FEditConditionParser_EvaluateEnum::RunTest(const FString& Parameters)
 	return RunEnumTests(TestContext, EnumType, PropertyName);
 }
 
+IMPLEMENT_SIMPLE_AUTOMATION_TEST(FEditConditionParser_EvaluateBitFlags, "EditConditionParser.EvaluateBitFlags", EAutomationTestFlags::EditorContext | EAutomationTestFlags::EngineFilter)
+bool FEditConditionParser_EvaluateBitFlags::RunTest(const FString& Parameters)
+{
+	FTestEditConditionContext TestContext;
+
+	const FString EnumType = TEXT("TestEnum");
+	TestContext.SetupEnumType(EnumType);
+	TestContext.SetupEnumTypeValue(TEXT("Nil"), 0);
+	TestContext.SetupEnumTypeValue(TEXT("One"), 1 << 0);
+	TestContext.SetupEnumTypeValue(TEXT("Two"), 1 << 1);
+	TestContext.SetupEnumTypeValue(TEXT("Four"), 1 << 2);
+
+	const FString PropertyName = TEXT("FlagsProperty");
+
+	FEditConditionParser Parser;
+	bool bResult = true;
+
+	TestContext.SetupInteger(PropertyName, 0);
+	bResult &= CanEvaluate(Parser, TestContext, "FlagsProperty & TestEnum::Nil", false);
+	bResult &= CanEvaluate(Parser, TestContext, "FlagsProperty & TestEnum::One", false);
+	bResult &= CanEvaluate(Parser, TestContext, "FlagsProperty & TestEnum::Two", false);
+	bResult &= CanEvaluate(Parser, TestContext, "FlagsProperty & TestEnum::Four", false);
+	bResult &= CanEvaluate(Parser, TestContext, "FlagsProperty & TestEnum::One == false", true);
+	bResult &= CanEvaluate(Parser, TestContext, "FlagsProperty & TestEnum::Two == false", true);
+	bResult &= CanEvaluate(Parser, TestContext, "FlagsProperty & TestEnum::Four == false", true);
+
+	TestContext.SetupInteger(PropertyName, 1);
+	bResult &= CanEvaluate(Parser, TestContext, "FlagsProperty & TestEnum::Nil", false);
+	bResult &= CanEvaluate(Parser, TestContext, "FlagsProperty & TestEnum::One", true);
+	bResult &= CanEvaluate(Parser, TestContext, "FlagsProperty & TestEnum::Two", false);
+	bResult &= CanEvaluate(Parser, TestContext, "FlagsProperty & TestEnum::Four", false);
+
+	TestContext.SetupInteger(PropertyName, 2);
+	bResult &= CanEvaluate(Parser, TestContext, "FlagsProperty & TestEnum::Nil", false);
+	bResult &= CanEvaluate(Parser, TestContext, "FlagsProperty & TestEnum::One", false);
+	bResult &= CanEvaluate(Parser, TestContext, "FlagsProperty & TestEnum::Two", true);
+	bResult &= CanEvaluate(Parser, TestContext, "FlagsProperty & TestEnum::Four", false);
+
+	TestContext.SetupInteger(PropertyName, 3);
+	bResult &= CanEvaluate(Parser, TestContext, "FlagsProperty & TestEnum::Nil", false);
+	bResult &= CanEvaluate(Parser, TestContext, "FlagsProperty & TestEnum::One", true);
+	bResult &= CanEvaluate(Parser, TestContext, "FlagsProperty & TestEnum::Two", true);
+	bResult &= CanEvaluate(Parser, TestContext, "FlagsProperty & TestEnum::Four", false);
+
+	TestContext.SetupInteger(PropertyName, 5);
+	bResult &= CanEvaluate(Parser, TestContext, "FlagsProperty & TestEnum::Nil", false);
+	bResult &= CanEvaluate(Parser, TestContext, "FlagsProperty & TestEnum::One", true);
+	bResult &= CanEvaluate(Parser, TestContext, "FlagsProperty & TestEnum::Two", false);
+	bResult &= CanEvaluate(Parser, TestContext, "FlagsProperty & TestEnum::Four", true);
+
+	TestContext.SetupInteger(PropertyName, 7);
+	bResult &= CanEvaluate(Parser, TestContext, "FlagsProperty & TestEnum::Nil", false);
+	bResult &= CanEvaluate(Parser, TestContext, "FlagsProperty & TestEnum::One", true);
+	bResult &= CanEvaluate(Parser, TestContext, "FlagsProperty & TestEnum::Two", true);
+	bResult &= CanEvaluate(Parser, TestContext, "FlagsProperty & TestEnum::Four", true);
+
+	bResult &= CanEvaluate(Parser, TestContext, "FlagsProperty & TestEnum::One == false", false);
+	bResult &= CanEvaluate(Parser, TestContext, "FlagsProperty & TestEnum::Two == false", false);
+	bResult &= CanEvaluate(Parser, TestContext, "FlagsProperty & TestEnum::Four == false", false);
+	bResult &= CanEvaluate(Parser, TestContext, "FlagsProperty & TestEnum::One == true", true);
+	bResult &= CanEvaluate(Parser, TestContext, "FlagsProperty & TestEnum::Two == true", true);
+	bResult &= CanEvaluate(Parser, TestContext, "FlagsProperty & TestEnum::Four == true", true);
+
+	return bResult;
+}
+
 IMPLEMENT_SIMPLE_AUTOMATION_TEST(FEditConditionParser_EvaluateUObject, "EditConditionParser.EvaluateUObject", EAutomationTestFlags::EditorContext | EAutomationTestFlags::EngineFilter)
 bool FEditConditionParser_EvaluateUObject::RunTest(const FString& Parameters)
 {
@@ -524,4 +627,4 @@ bool FEditConditionParser_SingleBool::RunTest(const FString& Parameters)
 	return true;
 }
 
-#endif // WITH_DEV_AUTOMATION_TESTS
\ No newline at end of file
+#endif // WITH_DEV_AUTOMATION_TESTS
diff --git a/Engine/Source/Editor/UnrealEd/Classes/Editor/PropertyEditorTestObject.h b/Engine/Source/Editor/UnrealEd/Classes/Editor/PropertyEditorTestObject.h
--- a/Engine/Source/Editor/UnrealEd/Classes/Editor/PropertyEditorTestObject.h
+++ b/Engine/Source/Editor/UnrealEd/Classes/Editor/PropertyEditorTestObject.h
@@ -69,6 +69,15 @@ enum class EditColor : uint8
 	Cyan
 };
 
+UENUM()
+enum class ETestEnumFlags : uint8
+{
+	None = 0,
+	One = 1 << 0,
+	Two = 1 << 1,
+	Four = 1 << 2
+};
+
 USTRUCT()
 struct FPropertyEditTestTextStruct
 {
@@ -551,6 +560,15 @@ class UPropertyEditorTestObject : public UObject
 	UPROPERTY(EditAnywhere, Category = EditCondition, meta = (EditCondition = "bEditConditionForArrays"))
 	TArray<FPropertyEditorTestBasicStruct> ArrayOfStructsWithEditCondition;
 
+	UPROPERTY(EditAnywhere, Category = EditCondition)
+	int64 Flags;
+
+	UPROPERTY(EditAnywhere, Category = EditCondition, meta = (EditCondition = "Flags & ETestEnumFlags::Two || Flags & ETestEnumFlags::Four"))
+	bool bEnabledWhenFlagsHasTwoOrFour;
+
+	UPROPERTY(EditAnywhere, Category = EditCondition, meta = (EditCondition = "Flags & ETestEnumFlags::One == false"))
+	bool bDisabledWhenFlagsIsOdd;
+
 	UPROPERTY(EditAnywhere, Category = OnlyInlineProperty, meta = (InlineCategoryProperty))
 	TEnumAsByte<EComponentMobility::Type> InlineProperty;
 
