EditConditions no longer crash when their parent object gets deleted from under them, eg. when editing a blueprint CDO and then compiling.

diff --git a/Engine/Source/Editor/PropertyEditor/Private/EditConditionContext.cpp b/Engine/Source/Editor/PropertyEditor/Private/EditConditionContext.cpp
--- a/Engine/Source/Editor/PropertyEditor/Private/EditConditionContext.cpp
+++ b/Engine/Source/Editor/PropertyEditor/Private/EditConditionContext.cpp
@@ -81,7 +81,14 @@ T* FindTypedField(const TWeakPtr<FPropertyNode>& PropertyNode, const FString& Pr
 	return nullptr;
 }
 
-FPropertyNode* GetEditConditionParentNode(const TSharedPtr<FPropertyNode>& PropertyNode)
+/** 
+ * Get the parent to use as the context when evaluating the edit condition.
+ * For normal properties inside a UObject, this is the UObject. 
+ * For children of containers, this is the UObject the container is in. 
+ * Note: We do not support nested containers.
+ * The result can be nullptr in exceptional cases, eg. if the UI is getting rebuilt.
+ */
+static FPropertyNode* GetEditConditionParentNode(const TSharedPtr<FPropertyNode>& PropertyNode)
 {
 	FPropertyNode* ParentNode = PropertyNode->GetParentNode();
 	const UObject* PropertyOuter = PropertyNode->GetProperty()->GetOuter();
@@ -106,6 +113,11 @@ static uint8* GetPropertyValuePtr(const UProperty* Property, const TSharedPtr<FP
 	}
 
 	uint8* ParentPtr = ParentNode->GetValueAddress(BasePtr, PropertyNode->HasNodeFlags(EPropertyNodeFlags::IsSparseClassData) != 0);
+	if (ParentPtr == nullptr)
+	{
+		return nullptr;
+	}
+
 	uint8* ValuePtr = ComplexParentNode->GetValuePtrOfInstance(Index, Property, ParentNode);
 
 	// SPARSEDATA_TODO: remove the next three lines once we're sure the pointer math is correct
@@ -125,7 +137,17 @@ TOptional<bool> FEditConditionContext::GetBoolValue(const FString& PropertyName)
 	}
 
 	TSharedPtr<FPropertyNode> PinnedNode = PropertyNode.Pin();
+	if (!PinnedNode.IsValid())
+	{
+		return TOptional<bool>();
+	}
+
 	FPropertyNode* ParentNode = GetEditConditionParentNode(PinnedNode);
+	if (ParentNode == nullptr)
+	{
+		return TOptional<bool>();
+	}
+
 	FComplexPropertyNode* ComplexParentNode = PinnedNode->FindComplexParent();
 
 	TOptional<bool> Result;
@@ -166,7 +188,17 @@ TOptional<int64> FEditConditionContext::GetIntegerValue(const FString& PropertyN
 	}
 
 	TSharedPtr<FPropertyNode> PinnedNode = PropertyNode.Pin();
+	if (!PinnedNode.IsValid())
+	{
+		return TOptional<int64>();
+	}
+
 	FPropertyNode* ParentNode = GetEditConditionParentNode(PinnedNode);
+	if (ParentNode == nullptr)
+	{
+		return TOptional<int64>();
+	}
+
 	FComplexPropertyNode* ComplexParentNode = PinnedNode->FindComplexParent();
 
 	TOptional<int64> Result;
@@ -202,7 +234,17 @@ TOptional<double> FEditConditionContext::GetNumericValue(const FString& Property
 	}
 
 	TSharedPtr<FPropertyNode> PinnedNode = PropertyNode.Pin();
+	if (!PinnedNode.IsValid())
+	{
+		return TOptional<double>();
+	}
+
 	FPropertyNode* ParentNode = GetEditConditionParentNode(PinnedNode);
+	if (ParentNode == nullptr)
+	{
+		return TOptional<double>();
+	}
+
 	FComplexPropertyNode* ComplexParentNode = PinnedNode->FindComplexParent();
 
 	TOptional<double> Result;
@@ -270,7 +312,17 @@ TOptional<FString> FEditConditionContext::GetEnumValue(const FString& PropertyNa
 	}
 
 	TSharedPtr<FPropertyNode> PinnedNode = PropertyNode.Pin();
+	if (!PinnedNode.IsValid())
+	{
+		return TOptional<FString>();
+	}
+
 	FPropertyNode* ParentNode = GetEditConditionParentNode(PinnedNode);
+	if (ParentNode == nullptr)
+	{
+		return TOptional<FString>();
+	}
+
 	FComplexPropertyNode* ComplexParentNode = PinnedNode->FindComplexParent();
 
 	TOptional<int64> Result;
@@ -317,7 +369,17 @@ TOptional<UObject*> FEditConditionContext::GetPointerValue(const FString& Proper
 	}
 
 	TSharedPtr<FPropertyNode> PinnedNode = PropertyNode.Pin();
+	if (!PinnedNode.IsValid())
+	{
+		return TOptional<UObject*>();
+	}
+
 	FPropertyNode* ParentNode = GetEditConditionParentNode(PinnedNode);
+	if (ParentNode == nullptr)
+	{
+		return TOptional<UObject*>();
+	}
+
 	FComplexPropertyNode* ComplexParentNode = PinnedNode->FindComplexParent();
 
 	TOptional<UObject*> Result;
