Adding an entry to an array in a component being edited in a locked details view no longer asserts.

This was occurring due to the construction script being rerun and the old instance of the component being garbage collected. When the view is not locked this was handled by the selection being set, but with a locked view we were relying on SDetailsView::ReplaceObjects, which was failing to match the existing component to the newly created component.

diff --git a/Engine/Source/Editor/PropertyEditor/Private/SDetailsView.cpp b/Engine/Source/Editor/PropertyEditor/Private/SDetailsView.cpp
--- a/Engine/Source/Editor/PropertyEditor/Private/SDetailsView.cpp
+++ b/Engine/Source/Editor/PropertyEditor/Private/SDetailsView.cpp
@@ -707,46 +707,46 @@ void SDetailsView::SetObjectArrayPrivate(const TArray< TWeakObjectPtr< UObject >
 
 void SDetailsView::ReplaceObjects( const TMap<UObject*, UObject*>& OldToNewObjectMap )
 {
-	TArray< TWeakObjectPtr< UObject > > NewObjectList;
+	TArray<TWeakObjectPtr<UObject>> NewObjectList;
 	bool bObjectsReplaced = false;
 
-	TArray< FObjectPropertyNode* > ObjectNodes;
-	for(TSharedPtr<FComplexPropertyNode>& RootNode : RootPropertyNodes)
+	TArray<FObjectPropertyNode*> ObjectNodes;
+	for (TSharedPtr<FComplexPropertyNode>& RootNode : RootPropertyNodes)
 	{
-		PropertyEditorHelpers::CollectObjectNodes(RootNode, ObjectNodes );
+		PropertyEditorHelpers::CollectObjectNodes(RootNode, ObjectNodes);
 	}
 
-	for( int32 ObjectNodeIndex = 0; ObjectNodeIndex < ObjectNodes.Num(); ++ObjectNodeIndex )
+	for (FObjectPropertyNode* CurrentNode : ObjectNodes)
 	{
-		FObjectPropertyNode* CurrentNode = ObjectNodes[ObjectNodeIndex];
-
 		// Scan all objects and look for objects which need to be replaced
-		for ( TPropObjectIterator Itor( CurrentNode->ObjectIterator() ); Itor; ++Itor )
+		for (int32 Idx = 0; Idx < CurrentNode->GetNumObjects(); ++Idx)
 		{
-			UObject* Replacement = OldToNewObjectMap.FindRef( Itor->Get() );
-			if( Replacement )
+			// We could be replacing an object that has already been garbage collected, so look up the object using the raw pointer.
+			TWeakObjectPtr<UObject> Existing = CurrentNode->GetInstanceAsUObject(Idx);
+			UObject* Replacement = OldToNewObjectMap.FindRef(Existing.GetEvenIfUnreachable());
+
+			// Note: only root objects count for the new object list. 
+			// Sub-Objects (i.e components) count as needing to be replaced but they don't belong in the top level object list)
+			if (Replacement)
 			{
 				bObjectsReplaced = true;
-				if( CurrentNode->IsRootNode() )
+
+				if (CurrentNode->IsRootNode())
 				{
-					// Note: only root objects count for the new object list. Sub-Objects (i.e components count as needing to be replaced but they don't belong in the top level object list
-					NewObjectList.Add( Replacement );
+					NewObjectList.Add(Replacement);
 				}
 			}
-			else if( CurrentNode->IsRootNode() )
+			else if (CurrentNode->IsRootNode())
 			{
-				// Note: only root objects count for the new object list. Sub-Objects (i.e components count as needing to be replaced but they don't belong in the top level object list
-				NewObjectList.Add( Itor->Get() );
+				NewObjectList.Add(Existing.Get());
 			}
 		}
 	}
 
-
-	if( bObjectsReplaced )
+	if (bObjectsReplaced)
 	{
-		SetObjectArrayPrivate( NewObjectList );
+		SetObjectArrayPrivate(NewObjectList);
 	}
-
 }
 
 void SDetailsView::RemoveDeletedObjects( const TArray<UObject*>& DeletedObjects )
