Changing default multi-object details panel appearance

diff --git a/Engine/Source/Editor/PropertyEditor/Private/DetailMultiTopLevelObjectRootNode.cpp b/Engine/Source/Editor/PropertyEditor/Private/DetailMultiTopLevelObjectRootNode.cpp
--- a/Engine/Source/Editor/PropertyEditor/Private/DetailMultiTopLevelObjectRootNode.cpp
+++ b/Engine/Source/Editor/PropertyEditor/Private/DetailMultiTopLevelObjectRootNode.cpp
@@ -99,10 +99,23 @@ ENodeVisibility FDetailMultiTopLevelObjectRootNode::GetVisibility() const
 TSharedRef< ITableRow > FDetailMultiTopLevelObjectRootNode::GenerateWidgetForTableView(const TSharedRef<STableViewBase>& OwnerTable, const FDetailColumnSizeData& ColumnSizeData, bool bAllowFavoriteSystem)
 {
 	FDetailWidgetRow Row;
-	TSharedRef<SDetailMultiTopLevelObjectTableRow> DetailMultiTopLevelObjectTableRow = SNew(SDetailMultiTopLevelObjectTableRow, AsShared());
-	GenerateStandaloneWidget(Row, DetailMultiTopLevelObjectTableRow);
-	DetailMultiTopLevelObjectTableRow->ChildSlotConstruct(Row.NameWidget.Widget, OwnerTable);
-	return DetailMultiTopLevelObjectTableRow;
+	TSharedPtr<SDetailMultiTopLevelObjectTableRow> DetailMultiTopLevelObjectTableRow;
+	if (RootObjectCustomization.IsValid())
+	{
+		
+		DetailMultiTopLevelObjectTableRow = SNew(SDetailMultiTopLevelObjectTableRow, AsShared());
+		GenerateStandaloneWidget(Row, DetailMultiTopLevelObjectTableRow.ToSharedRef());
+		DetailMultiTopLevelObjectTableRow->ChildSlotConstruct(Row.NameWidget.Widget, OwnerTable);
+		bShouldShowOnlyChildren = false;
+	}
+	else
+	{
+		GenerateStandaloneWidget(Row);
+		DetailMultiTopLevelObjectTableRow = SNew(SDetailMultiTopLevelObjectTableRow, AsShared(), Row.NameWidget.Widget, OwnerTable);
+		bShouldShowOnlyChildren = true;
+	}
+	
+	return DetailMultiTopLevelObjectTableRow.ToSharedRef();
 }
 
 
@@ -192,5 +205,5 @@ void FDetailMultiTopLevelObjectRootNode::FilterNode( const FDetailFilter& InFilt
 
 bool FDetailMultiTopLevelObjectRootNode::ShouldShowOnlyChildren() const
 {
-	return RootObjectCustomization.IsValid() && RootObject.IsValid() ? !RootObjectCustomization.Pin()->ShouldDisplayHeader(RootObject.Get()) : false;
+	return RootObjectCustomization.IsValid() && RootObject.IsValid() ? !RootObjectCustomization.Pin()->ShouldDisplayHeader(RootObject.Get()) : bShouldShowOnlyChildren;
 }
diff --git a/Engine/Source/Editor/PropertyEditor/Private/DetailMultiTopLevelObjectRootNode.h b/Engine/Source/Editor/PropertyEditor/Private/DetailMultiTopLevelObjectRootNode.h
--- a/Engine/Source/Editor/PropertyEditor/Private/DetailMultiTopLevelObjectRootNode.h
+++ b/Engine/Source/Editor/PropertyEditor/Private/DetailMultiTopLevelObjectRootNode.h
@@ -80,4 +80,5 @@ private:
 	const TWeakObjectPtr<UObject> RootObject;
 	FName NodeName;
 	bool bShouldBeVisible;
+	bool bShouldShowOnlyChildren;
 };
