Inline edit conditions in the details panel can now be properties without any metadata whatsoever.

Added a warning about this, since my view is that this sort of implicit behaviour should not be encouraged.

diff --git a/Engine/Source/Editor/PropertyEditor/Private/Presentation/PropertyEditor/PropertyEditor.cpp b/Engine/Source/Editor/PropertyEditor/Private/Presentation/PropertyEditor/PropertyEditor.cpp
--- a/Engine/Source/Editor/PropertyEditor/Private/Presentation/PropertyEditor/PropertyEditor.cpp
+++ b/Engine/Source/Editor/PropertyEditor/Private/Presentation/PropertyEditor/PropertyEditor.cpp
@@ -619,18 +619,23 @@ bool FPropertyEditor::SupportsEditConditionToggle() const
 		const UBoolProperty* ConditionalProperty = EditConditionContext->GetSingleBoolProperty(EditConditionExpression);
 		if (ConditionalProperty != nullptr)
 		{
+			// Both of these conditions are technically bugs of omission.
+			// However, they were supported prior to 4.23, so just warn and allow them for now.
+
+			// The property is not marked as an inline edit condition toggle.
 			static const FName Name_InlineEditConditionToggle("InlineEditConditionToggle");
 			if (ConditionalProperty->HasMetaData(Name_InlineEditConditionToggle))
 			{
-				// If the edit condition property is not marked as editable, it's technically a bug.
-				// However, this was the behaviour prior to 4.23, so just warn and allow it for now.
-				if (!ConditionalProperty->HasAllPropertyFlags(CPF_Edit))
-				{
-					UE_LOG(LogPropertyEditor, Error, TEXT("Property being used as InlineEditConditionToggle is not marked as editable: Field \"%s\" in class \"%s\"."), *ConditionalProperty->GetNameCPP(), *Property->GetOwnerStruct()->GetName());
-				}
+				UE_LOG(LogPropertyEditor, Warning, TEXT("Property being used as inline edit condition is not marked as InlineEditConditionToggle. Field \"%s\" in class \"%s\"."), *ConditionalProperty->GetNameCPP(), *Property->GetOwnerStruct()->GetName());
+			}
 
-				return true;
+			// The property is not marked as editable.
+			if (!ConditionalProperty->HasAllPropertyFlags(CPF_Edit))
+			{
+				UE_LOG(LogPropertyEditor, Warning, TEXT("Property being used as inline edit condition is not marked as editable. Field \"%s\" in class \"%s\"."), *ConditionalProperty->GetNameCPP(), *Property->GetOwnerStruct()->GetName());
 			}
+
+			return true;
 		}
 	}
 
