Added missing TEXT() macros in unit test code.

#rnx

diff --git a/Engine/Source/Editor/PropertyEditor/Private/EditConditionParserTests.cpp b/Engine/Source/Editor/PropertyEditor/Private/EditConditionParserTests.cpp
--- a/Engine/Source/Editor/PropertyEditor/Private/EditConditionParserTests.cpp
+++ b/Engine/Source/Editor/PropertyEditor/Private/EditConditionParserTests.cpp
@@ -384,50 +384,50 @@ bool FEditConditionParser_EvaluateBitFlags::RunTest(const FString& Parameters)
 	bool bResult = true;
 
 	TestContext.SetupInteger(PropertyName, 0);
-	bResult &= CanEvaluate(Parser, TestContext, "FlagsProperty & TestEnum::Nil", false);
-	bResult &= CanEvaluate(Parser, TestContext, "FlagsProperty & TestEnum::One", false);
-	bResult &= CanEvaluate(Parser, TestContext, "FlagsProperty & TestEnum::Two", false);
-	bResult &= CanEvaluate(Parser, TestContext, "FlagsProperty & TestEnum::Four", false);
-	bResult &= CanEvaluate(Parser, TestContext, "FlagsProperty & TestEnum::One == false", true);
-	bResult &= CanEvaluate(Parser, TestContext, "FlagsProperty & TestEnum::Two == false", true);
-	bResult &= CanEvaluate(Parser, TestContext, "FlagsProperty & TestEnum::Four == false", true);
+	bResult &= CanEvaluate(Parser, TestContext, TEXT("FlagsProperty & TestEnum::Nil"), false);
+	bResult &= CanEvaluate(Parser, TestContext, TEXT("FlagsProperty & TestEnum::One"), false);
+	bResult &= CanEvaluate(Parser, TestContext, TEXT("FlagsProperty & TestEnum::Two"), false);
+	bResult &= CanEvaluate(Parser, TestContext, TEXT("FlagsProperty & TestEnum::Four"), false);
+	bResult &= CanEvaluate(Parser, TestContext, TEXT("FlagsProperty & TestEnum::One == false"), true);
+	bResult &= CanEvaluate(Parser, TestContext, TEXT("FlagsProperty & TestEnum::Two == false"), true);
+	bResult &= CanEvaluate(Parser, TestContext, TEXT("FlagsProperty & TestEnum::Four == false"), true);
 
 	TestContext.SetupInteger(PropertyName, 1);
-	bResult &= CanEvaluate(Parser, TestContext, "FlagsProperty & TestEnum::Nil", false);
-	bResult &= CanEvaluate(Parser, TestContext, "FlagsProperty & TestEnum::One", true);
-	bResult &= CanEvaluate(Parser, TestContext, "FlagsProperty & TestEnum::Two", false);
-	bResult &= CanEvaluate(Parser, TestContext, "FlagsProperty & TestEnum::Four", false);
+	bResult &= CanEvaluate(Parser, TestContext, TEXT("FlagsProperty & TestEnum::Nil"), false);
+	bResult &= CanEvaluate(Parser, TestContext, TEXT("FlagsProperty & TestEnum::One"), true);
+	bResult &= CanEvaluate(Parser, TestContext, TEXT("FlagsProperty & TestEnum::Two"), false);
+	bResult &= CanEvaluate(Parser, TestContext, TEXT("FlagsProperty & TestEnum::Four"), false);
 
 	TestContext.SetupInteger(PropertyName, 2);
-	bResult &= CanEvaluate(Parser, TestContext, "FlagsProperty & TestEnum::Nil", false);
-	bResult &= CanEvaluate(Parser, TestContext, "FlagsProperty & TestEnum::One", false);
-	bResult &= CanEvaluate(Parser, TestContext, "FlagsProperty & TestEnum::Two", true);
-	bResult &= CanEvaluate(Parser, TestContext, "FlagsProperty & TestEnum::Four", false);
+	bResult &= CanEvaluate(Parser, TestContext, TEXT("FlagsProperty & TestEnum::Nil"), false);
+	bResult &= CanEvaluate(Parser, TestContext, TEXT("FlagsProperty & TestEnum::One"), false);
+	bResult &= CanEvaluate(Parser, TestContext, TEXT("FlagsProperty & TestEnum::Two"), true);
+	bResult &= CanEvaluate(Parser, TestContext, TEXT("FlagsProperty & TestEnum::Four"), false);
 
 	TestContext.SetupInteger(PropertyName, 3);
-	bResult &= CanEvaluate(Parser, TestContext, "FlagsProperty & TestEnum::Nil", false);
-	bResult &= CanEvaluate(Parser, TestContext, "FlagsProperty & TestEnum::One", true);
-	bResult &= CanEvaluate(Parser, TestContext, "FlagsProperty & TestEnum::Two", true);
-	bResult &= CanEvaluate(Parser, TestContext, "FlagsProperty & TestEnum::Four", false);
+	bResult &= CanEvaluate(Parser, TestContext, TEXT("FlagsProperty & TestEnum::Nil"), false);
+	bResult &= CanEvaluate(Parser, TestContext, TEXT("FlagsProperty & TestEnum::One"), true);
+	bResult &= CanEvaluate(Parser, TestContext, TEXT("FlagsProperty & TestEnum::Two"), true);
+	bResult &= CanEvaluate(Parser, TestContext, TEXT("FlagsProperty & TestEnum::Four"), false);
 
 	TestContext.SetupInteger(PropertyName, 5);
-	bResult &= CanEvaluate(Parser, TestContext, "FlagsProperty & TestEnum::Nil", false);
-	bResult &= CanEvaluate(Parser, TestContext, "FlagsProperty & TestEnum::One", true);
-	bResult &= CanEvaluate(Parser, TestContext, "FlagsProperty & TestEnum::Two", false);
-	bResult &= CanEvaluate(Parser, TestContext, "FlagsProperty & TestEnum::Four", true);
+	bResult &= CanEvaluate(Parser, TestContext, TEXT("FlagsProperty & TestEnum::Nil"), false);
+	bResult &= CanEvaluate(Parser, TestContext, TEXT("FlagsProperty & TestEnum::One"), true);
+	bResult &= CanEvaluate(Parser, TestContext, TEXT("FlagsProperty & TestEnum::Two"), false);
+	bResult &= CanEvaluate(Parser, TestContext, TEXT("FlagsProperty & TestEnum::Four"), true);
 
 	TestContext.SetupInteger(PropertyName, 7);
-	bResult &= CanEvaluate(Parser, TestContext, "FlagsProperty & TestEnum::Nil", false);
-	bResult &= CanEvaluate(Parser, TestContext, "FlagsProperty & TestEnum::One", true);
-	bResult &= CanEvaluate(Parser, TestContext, "FlagsProperty & TestEnum::Two", true);
-	bResult &= CanEvaluate(Parser, TestContext, "FlagsProperty & TestEnum::Four", true);
-
-	bResult &= CanEvaluate(Parser, TestContext, "FlagsProperty & TestEnum::One == false", false);
-	bResult &= CanEvaluate(Parser, TestContext, "FlagsProperty & TestEnum::Two == false", false);
-	bResult &= CanEvaluate(Parser, TestContext, "FlagsProperty & TestEnum::Four == false", false);
-	bResult &= CanEvaluate(Parser, TestContext, "FlagsProperty & TestEnum::One == true", true);
-	bResult &= CanEvaluate(Parser, TestContext, "FlagsProperty & TestEnum::Two == true", true);
-	bResult &= CanEvaluate(Parser, TestContext, "FlagsProperty & TestEnum::Four == true", true);
+	bResult &= CanEvaluate(Parser, TestContext, TEXT("FlagsProperty & TestEnum::Nil"), false);
+	bResult &= CanEvaluate(Parser, TestContext, TEXT("FlagsProperty & TestEnum::One"), true);
+	bResult &= CanEvaluate(Parser, TestContext, TEXT("FlagsProperty & TestEnum::Two"), true);
+	bResult &= CanEvaluate(Parser, TestContext, TEXT("FlagsProperty & TestEnum::Four"), true);
+
+	bResult &= CanEvaluate(Parser, TestContext, TEXT("FlagsProperty & TestEnum::One == false"), false);
+	bResult &= CanEvaluate(Parser, TestContext, TEXT("FlagsProperty & TestEnum::Two == false"), false);
+	bResult &= CanEvaluate(Parser, TestContext, TEXT("FlagsProperty & TestEnum::Four == false"), false);
+	bResult &= CanEvaluate(Parser, TestContext, TEXT("FlagsProperty & TestEnum::One == true"), true);
+	bResult &= CanEvaluate(Parser, TestContext, TEXT("FlagsProperty & TestEnum::Two == true"), true);
+	bResult &= CanEvaluate(Parser, TestContext, TEXT("FlagsProperty & TestEnum::Four == true"), true);
 
 	return bResult;
 }
