Key Category localized meta-data against the source string rather than the path

This applies the de-duplication at the gather level rather than the PO export level, which)
  1) Removes redundant string copies in memory.
  2) Removes the need to gather and import to resolve new duplicates.

Note: This will make categories appear non-localizable until the next gather runs (due to the key change), at which point they'll fix themselves.

#rnx

diff --git a/Engine/Config/Localization/Category.ini b/Engine/Config/Localization/Category.ini
--- a/Engine/Config/Localization/Category.ini
+++ b/Engine/Config/Localization/Category.ini
@@ -25,7 +25,7 @@ ExcludePathFilters=*/NoRedist/*
 ExcludePathFilters=*/NotForLicensees/*
 InputKeys=Category
 OutputNamespaces=UObjectCategory
-OutputKeys="{FieldPath}"
+OutputKeys="{MetaDataValue}"
 ShouldGatherFromEditorOnlyData=true
 
 ;Write Manifest
@@ -41,7 +41,6 @@ bPurgeOldEmptyEntries=true
 [GatherTextStep3]
 CommandletClass=InternationalizationExport
 bImportLoc=true
-LocalizedTextCollapseMode=ELocalizedTextCollapseMode::IdenticalNamespaceAndSource
 
 ;Write Localized Text Resource
 [GatherTextStep4]
@@ -52,7 +51,6 @@ ResourceName=Category.locres
 [GatherTextStep5]
 CommandletClass=InternationalizationExport
 bExportLoc=true
-LocalizedTextCollapseMode=ELocalizedTextCollapseMode::IdenticalNamespaceAndSource
 
 ;Write Text Localization Report
 [GatherTextStep6]
diff --git a/Engine/Source/Editor/PropertyEditor/Private/DetailCategoryBuilderImpl.cpp b/Engine/Source/Editor/PropertyEditor/Private/DetailCategoryBuilderImpl.cpp
--- a/Engine/Source/Editor/PropertyEditor/Private/DetailCategoryBuilderImpl.cpp
+++ b/Engine/Source/Editor/PropertyEditor/Private/DetailCategoryBuilderImpl.cpp
@@ -613,7 +613,26 @@ void FDetailCategoryImpl::SetDisplayName(FName InCategoryName, const FText& Loca
 	}
 	else if (InCategoryName != NAME_None)
 	{
-		DisplayName = FText::AsCultureInvariant(FName::NameToDisplayString(InCategoryName.ToString(), false));
+		static const FTextKey CategoryLocalizationNamespace = TEXT("UObjectCategory");
+		static const FName CategoryMetaDataKey = TEXT("Category");
+
+		DisplayName = FText();
+
+		const FString NativeCategory = InCategoryName.ToString();
+		if (FText::FindText(CategoryLocalizationNamespace, NativeCategory, /*OUT*/DisplayName, &NativeCategory))
+		{
+			// Category names in English are typically gathered in their non-pretty form (eg "UserInterface" rather than "User Interface"), so skip 
+			// applying the localized variant if the text matches the raw category name, as in this case the pretty printer will do a better job
+			if (NativeCategory.Equals(DisplayName.ToString(), ESearchCase::CaseSensitive))
+			{
+				DisplayName = FText();
+			}
+		}
+		
+		if (DisplayName.IsEmpty())
+		{
+			DisplayName = FText::AsCultureInvariant(FName::NameToDisplayString(NativeCategory, false));
+		}
 	}
 	else
 	{
diff --git a/Engine/Source/Editor/PropertyEditor/Private/DetailLayoutHelpers.cpp b/Engine/Source/Editor/PropertyEditor/Private/DetailLayoutHelpers.cpp
--- a/Engine/Source/Editor/PropertyEditor/Private/DetailLayoutHelpers.cpp
+++ b/Engine/Source/Editor/PropertyEditor/Private/DetailLayoutHelpers.cpp
@@ -181,20 +181,6 @@ namespace DetailLayoutHelpers
 
 							// Add a property to the default category
 							FDetailCategoryImpl& CategoryImpl = DetailLayout.DefaultCategory(CategoryName);
-							{
-								FText CategoryDisplayName = CategoryImpl.GetDisplayName();
-								if (CategoryDisplayName.IsEmpty() || CategoryDisplayName.IsCultureInvariant())
-								{
-									CategoryDisplayName = FObjectEditorUtils::GetCategoryText(Property);
-
-									// Category names in English are typically gathered in their non-pretty form (eg "UserInterface" rather than "User Interface"), so skip 
-									// applying the localized variant if the text matches the raw category name, as in this case the pretty printer will do a better job
-									if (!CategoryName.ToString().Equals(CategoryDisplayName.ToString(), ESearchCase::CaseSensitive))
-									{
-										CategoryImpl.SetDisplayName(CategoryName, CategoryDisplayName);
-									}
-								}
-							}
 							CategoryImpl.AddPropertyNode(ChildNodePtr.ToSharedRef(), InstanceName);
 						}
 
diff --git a/Engine/Source/Runtime/Engine/Private/ObjectEditorUtils.cpp b/Engine/Source/Runtime/Engine/Private/ObjectEditorUtils.cpp
--- a/Engine/Source/Runtime/Engine/Private/ObjectEditorUtils.cpp
+++ b/Engine/Source/Runtime/Engine/Private/ObjectEditorUtils.cpp
@@ -12,11 +12,23 @@ namespace FObjectEditorUtils
 
 	FText GetCategoryText( const UField* InField )
 	{
-		static const FName NAME_Category(TEXT("Category"));
-		if (InField && InField->HasMetaData(NAME_Category))
+		static const FTextKey CategoryLocalizationNamespace = TEXT("UObjectCategory");
+		static const FName CategoryMetaDataKey = TEXT("Category");
+
+		if (InField)
 		{
-			return InField->GetMetaDataText(NAME_Category, TEXT("UObjectCategory"), InField->GetFullGroupName(false));
+			const FString& NativeCategory = InField->GetMetaData(CategoryMetaDataKey);
+			if (!NativeCategory.IsEmpty())
+			{
+				FText LocalizedCategory;
+				if (!FText::FindText(CategoryLocalizationNamespace, NativeCategory, /*OUT*/LocalizedCategory, &NativeCategory))
+				{
+					LocalizedCategory = FText::AsCultureInvariant(NativeCategory);
+				}
+				return LocalizedCategory;
+			}
 		}
+
 		return FText::GetEmpty();
 	}
 
