Read meta-data off the correct property when querying IStructureDetailsView

#rnx

diff --git a/Engine/Source/Editor/PropertyEditor/Private/PropertyEditorModule.cpp b/Engine/Source/Editor/PropertyEditor/Private/PropertyEditorModule.cpp
--- a/Engine/Source/Editor/PropertyEditor/Private/PropertyEditorModule.cpp
+++ b/Engine/Source/Editor/PropertyEditor/Private/PropertyEditorModule.cpp
@@ -732,6 +732,9 @@ TSharedRef<class IStructureDetailsView> FPropertyEditorModule::CreateStructureDe
 			PropertyToTest = SetProperty ? SetProperty->ElementProp : PropertyToTest;
 			PropertyToTest = MapProperty ? MapProperty->ValueProp : PropertyToTest;
 
+			// Meta-data should always be queried off the top-level property, as the inner (for container types) is generated by UHT
+			const UProperty* MetaDataProperty = &PropertyAndParent.Property;
+
 			if( InStructureDetailsViewArgs.bShowClasses && (PropertyToTest->IsA<UClassProperty>() || PropertyToTest->IsA<USoftClassProperty>()) )
 			{
 				return true;
@@ -756,7 +759,7 @@ TSharedRef<class IStructureDetailsView> FPropertyEditorModule::CreateStructureDe
 					// Not an "asset" property, but it may still be a property using an asset class type (such as a raw pointer)
 					if( ObjectProperty->PropertyClass )
 					{
-						if (ObjectProperty->HasMetaData(TEXT("AllowedClasses")))
+						if (ensure(MetaDataProperty) && MetaDataProperty->HasMetaData(TEXT("AllowedClasses")))
 						{
 							return true;
 						}
