Editor - You can now add metadata to FString and FName UPROPERTYs, that specify a function to call dynamically to get the options that shoud be available to the user.

meta=(GetOptions="FuncName")

..

UFUNCTION()
TArray<FString> FuncName() const; // Always return string array even if FName property.

diff --git a/Engine/Source/Editor/PropertyEditor/Private/PropertyHandleImpl.cpp b/Engine/Source/Editor/PropertyEditor/Private/PropertyHandleImpl.cpp
--- a/Engine/Source/Editor/PropertyEditor/Private/PropertyHandleImpl.cpp
+++ b/Engine/Source/Editor/PropertyEditor/Private/PropertyHandleImpl.cpp
@@ -25,8 +25,8 @@
 #include "UObject/FieldPathProperty.h"
 #include "IDetailPropertyRow.h"
 #include "ObjectEditorUtils.h"
-#include "PropertyEditorHelpers.h"
 #include "SResetToDefaultPropertyEditor.h"
+#include "PropertyPathHelpers.h"
 
 #define LOCTEXT_NAMESPACE "PropertyHandleImplementation"
 
@@ -2907,7 +2907,6 @@ FPropertyAccess::Result FPropertyHandleBase::GetPerObjectValue( const int32 Obje
 	return Result;
 }
 
-
 bool FPropertyHandleBase::GeneratePossibleValues(TArray< TSharedPtr<FString> >& OutOptionStrings, TArray< FText >& OutToolTips, TArray<bool>& OutRestrictedItems)
 {
 	FProperty* Property = GetProperty();
@@ -2984,6 +2983,34 @@ bool FPropertyHandleBase::GeneratePossibleValues(TArray< TSharedPtr<FString> >&
 			}
 		}
 	}
+	else if ((Property->IsA(FStrProperty::StaticClass()) || Property->IsA(FNameProperty::StaticClass())) && Property->HasMetaData(TEXT("GetOptions")))
+	{
+		const FString& GetOptionsFunction = Property->GetMetaData(TEXT("GetOptions"));
+		if (!GetOptionsFunction.IsEmpty())
+		{
+			TArray<UObject*> OutObjects;
+			GetOuterObjects(OutObjects);
+
+			if (OutObjects.Num() > 0)
+			{
+				FString GetOptionsFunctionName = Property->GetMetaData(TEXT("GetOptions"));
+
+				UObject* Target = OutObjects[0];
+
+				TArray<FString> Options;
+				FCachedPropertyPath Path(GetOptionsFunctionName);
+				const bool bSuccess = PropertyPathHelpers::GetPropertyValue(Target, Path, Options);
+				if (ensure(bSuccess))
+				{
+					for (const FString& Option : Options)
+					{
+						OutOptionStrings.Add(MakeShared<FString>(Option));
+						OutToolTips.Add(FText());
+					}
+				}
+			}
+		}
+	}
 	else if( Property->IsA(FClassProperty::StaticClass()) || Property->IsA(FSoftClassProperty::StaticClass()) )		
 	{
 		UClass* MetaClass = Property->IsA(FClassProperty::StaticClass()) 
diff --git a/Engine/Source/Editor/PropertyEditor/Private/UserInterface/PropertyEditor/SPropertyComboBox.cpp b/Engine/Source/Editor/PropertyEditor/Private/UserInterface/PropertyEditor/SPropertyComboBox.cpp
--- a/Engine/Source/Editor/PropertyEditor/Private/UserInterface/PropertyEditor/SPropertyComboBox.cpp
+++ b/Engine/Source/Editor/PropertyEditor/Private/UserInterface/PropertyEditor/SPropertyComboBox.cpp
@@ -20,7 +20,15 @@ void SPropertyComboBox::Construct( const FArguments& InArgs )
 	{
 		if(*ComboItemList[ItemIndex].Get() == VisibleText)
 		{
-			SetToolTip(RichToolTips[ItemIndex]);
+			if (RichToolTips.IsValidIndex(ItemIndex))
+			{
+				SetToolTip(RichToolTips[ItemIndex]);
+			}
+			else
+			{
+				SetToolTip(nullptr);
+			}
+
 			InitiallySelectedItem = ComboItemList[ItemIndex];
 			break;
 		}
diff --git a/Engine/Source/Editor/PropertyEditor/Private/UserInterface/PropertyEditor/SPropertyEditorCombo.cpp b/Engine/Source/Editor/PropertyEditor/Private/UserInterface/PropertyEditor/SPropertyEditorCombo.cpp
--- a/Engine/Source/Editor/PropertyEditor/Private/UserInterface/PropertyEditor/SPropertyEditorCombo.cpp
+++ b/Engine/Source/Editor/PropertyEditor/Private/UserInterface/PropertyEditor/SPropertyEditorCombo.cpp
@@ -50,6 +50,7 @@ bool SPropertyEditorCombo::Supports( const TSharedRef< class FPropertyEditor >&
 	if(	((Property->IsA(FByteProperty::StaticClass()) && CastField<const FByteProperty>(Property)->Enum)
 		||	Property->IsA(FEnumProperty::StaticClass())
 		||	(Property->IsA(FStrProperty::StaticClass()) && Property->HasMetaData(TEXT("Enum")))
+		|| ((Property->IsA(FStrProperty::StaticClass()) || Property->IsA(FNameProperty::StaticClass())) && Property->HasMetaData(TEXT("GetOptions")))
 		)
 		&&	( ( ArrayIndex == -1 && Property->ArrayDim == 1 ) || ( ArrayIndex > -1 && Property->ArrayDim > 0 ) ) )
 	{
diff --git a/Engine/Source/Editor/PropertyEditor/PropertyEditor.Build.cs b/Engine/Source/Editor/PropertyEditor/PropertyEditor.Build.cs
--- a/Engine/Source/Editor/PropertyEditor/PropertyEditor.Build.cs
+++ b/Engine/Source/Editor/PropertyEditor/PropertyEditor.Build.cs
@@ -63,6 +63,7 @@ public class PropertyEditor : ModuleRules
 				"ConfigEditor",
                 "SceneOutliner",
 				"DesktopPlatform",
+				"PropertyPath",
             }
         );
 
diff --git a/Engine/Source/Runtime/CoreUObject/Public/UObject/ObjectMacros.h b/Engine/Source/Runtime/CoreUObject/Public/UObject/ObjectMacros.h
--- a/Engine/Source/Runtime/CoreUObject/Public/UObject/ObjectMacros.h
+++ b/Engine/Source/Runtime/CoreUObject/Public/UObject/ObjectMacros.h
@@ -1237,6 +1237,12 @@ namespace UM
 
 		/// [PropertyMetadata] Used for SoftObjectPtr/SoftObjectPath properties to specify a reference should not be tracked. This reference will not be automatically cooked or saved into the asset registry for redirector/delete fixup.
 		Untracked,
+
+		/// [PropertyMetadata] Causes FString and FName properties to have a limited set of options generated dynamically, e.g. meta=(GetOptions="FuncName")
+		///
+		/// UFUNCTION()
+		/// TArray<FString> FuncName() const; // Always return string array even if FName property.
+		GetOptions,
 	};
 
 	// Metadata usable in UPROPERTY for customizing the behavior of Persona and UMG
