Fixing FPropertyEditorModule::RegisterStructOnScopeProperty using FProperty constuctor meant for native class properties

diff --git a/Engine/Source/Editor/PropertyEditor/Private/PropertyEditorModule.cpp b/Engine/Source/Editor/PropertyEditor/Private/PropertyEditorModule.cpp
--- a/Engine/Source/Editor/PropertyEditor/Private/PropertyEditorModule.cpp
+++ b/Engine/Source/Editor/PropertyEditor/Private/PropertyEditorModule.cpp
@@ -378,7 +378,7 @@ FStructProperty* FPropertyEditorModule::RegisterStructOnScopeProperty(TSharedRef
 	if(!StructProperty)
 	{
 		UScriptStruct* InnerStruct = Cast<UScriptStruct>(const_cast<UStruct*>(StructOnScope->GetStruct()));
-		StructProperty = new FStructProperty(InnerStruct, *MakeUniqueObjectName(InnerStruct, UField::StaticClass(), InnerStruct->GetFName()).ToString(), RF_NoFlags, 0, CPF_None, InnerStruct);
+		StructProperty = new FStructProperty(InnerStruct, *MakeUniqueObjectName(InnerStruct, UField::StaticClass(), InnerStruct->GetFName()).ToString(), RF_NoFlags);
 		StructProperty->Struct = InnerStruct;
 		StructProperty->ElementSize = StructOnScope->GetStruct()->GetStructureSize();
 
