PropertyEditor - Fix an issue where where struct properties weren't being handled well due to their children being empty due to the FProperty conversion.  Fixes struct properties appearing as "Multiple Values" in the details panel.

diff --git a/Engine/Source/Editor/PropertyEditor/Private/StructurePropertyNode.h b/Engine/Source/Editor/PropertyEditor/Private/StructurePropertyNode.h
--- a/Engine/Source/Editor/PropertyEditor/Private/StructurePropertyNode.h
+++ b/Engine/Source/Editor/PropertyEditor/Private/StructurePropertyNode.h
@@ -50,7 +50,7 @@ public:
 		}
 
 		UStruct* OwnerStruct = InItemProperty->GetOwnerStruct();
-		if (!OwnerStruct || OwnerStruct->Children == nullptr)
+		if (!OwnerStruct || (OwnerStruct->Children == nullptr && OwnerStruct->ChildProperties == nullptr))
 		{
 			// Verify that the property is not part of an invalid trash class
 			return false;
