Fix SObjectPropertyEntryBox::OnSetObject which wasn't handling unloaded Assets properly. Now uses the PropertyHandle::SetValue(FAssetData) code path.

diff --git a/Engine/Source/Editor/PropertyEditor/Private/PropertyCustomizationHelpers.cpp b/Engine/Source/Editor/PropertyEditor/Private/PropertyCustomizationHelpers.cpp
--- a/Engine/Source/Editor/PropertyEditor/Private/PropertyCustomizationHelpers.cpp
+++ b/Engine/Source/Editor/PropertyEditor/Private/PropertyCustomizationHelpers.cpp
@@ -568,13 +568,7 @@ void SObjectPropertyEntryBox::OnSetObject(const FAssetData& AssetData)
 	{
 		if (!OnShouldSetAsset.IsBound() || OnShouldSetAsset.Execute(AssetData))
 		{
-			FString ObjectPathName = TEXT("None");
-			if (AssetData.IsValid())
-			{
-				ObjectPathName = AssetData.ObjectPath.ToString();
-			}
-
-			PropertyHandle->SetValueFromFormattedString(ObjectPathName);
+			PropertyHandle->SetValue(AssetData);
 		}
 	}
 	OnObjectChanged.ExecuteIfBound(AssetData);
