Fixing compiling GameplayEffect Blueprint resulting in its Modifier's Attribute property to be reset (FProperty fallout)

diff --git a/Engine/Source/Editor/PropertyEditor/Private/PropertyEditorHelpers.cpp b/Engine/Source/Editor/PropertyEditor/Private/PropertyEditorHelpers.cpp
--- a/Engine/Source/Editor/PropertyEditor/Private/PropertyEditorHelpers.cpp
+++ b/Engine/Source/Editor/PropertyEditor/Private/PropertyEditorHelpers.cpp
@@ -634,6 +634,10 @@ namespace PropertyEditorHelpers
 		{
 			PropertyHandle = MakeShareable( new FPropertyHandleMap( PropertyNode, NotifyHook, PropertyUtilities ) );
 		}
+		else if (FPropertyHandleFieldPath::Supports(PropertyNode))
+		{
+			PropertyHandle = MakeShareable(new FPropertyHandleFieldPath(PropertyNode, NotifyHook, PropertyUtilities));
+		}
 		else
 		{
 			// Untyped or doesn't support getting the property directly but the property is still valid(probably struct property)
diff --git a/Engine/Source/Editor/PropertyEditor/Private/PropertyHandleImpl.cpp b/Engine/Source/Editor/PropertyEditor/Private/PropertyHandleImpl.cpp
--- a/Engine/Source/Editor/PropertyEditor/Private/PropertyHandleImpl.cpp
+++ b/Engine/Source/Editor/PropertyEditor/Private/PropertyHandleImpl.cpp
@@ -3298,6 +3298,7 @@ IMPLEMENT_PROPERTY_VALUE( FPropertyHandleArray )
 IMPLEMENT_PROPERTY_VALUE( FPropertyHandleText )
 IMPLEMENT_PROPERTY_VALUE( FPropertyHandleSet )
 IMPLEMENT_PROPERTY_VALUE( FPropertyHandleMap )
+IMPLEMENT_PROPERTY_VALUE( FPropertyHandleFieldPath )
 
 // int32 
 bool FPropertyHandleInt::Supports( TSharedRef<FPropertyNode> PropertyNode )
@@ -3816,45 +3817,6 @@ FPropertyAccess::Result FPropertyHandleObject::SetValue( const UObject* const& N
 	return FPropertyAccess::Fail;
 }
 
-FPropertyAccess::Result FPropertyHandleObject::GetValue(FProperty*& OutValue) const
-{
-	return FPropertyHandleObject::GetValue((const FProperty*&)OutValue);
-}
-
-FPropertyAccess::Result FPropertyHandleObject::GetValue(const FProperty*& OutValue) const
-{
-	void* PropValue = nullptr;
-	FPropertyAccess::Result Res = Implementation->GetValueData(PropValue);
-
-	if (Res == FPropertyAccess::Success)
-	{
-		FProperty* Property = GetProperty();
-		check(Property->IsA(FFieldPathProperty::StaticClass()));
-		check(false); // @todo FProp: verify this
-		OutValue = reinterpret_cast<const FProperty*>(PropValue);
-	}
-
-	return Res;
-}
-
-FPropertyAccess::Result FPropertyHandleObject::SetValue(FProperty* const& NewValue, EPropertyValueSetFlags::Type Flags)
-{
-	return FPropertyHandleObject::SetValue((const FProperty*)NewValue);
-}
-
-FPropertyAccess::Result FPropertyHandleObject::SetValue(const FProperty* const& NewValue, EPropertyValueSetFlags::Type Flags)
-{
-	const TSharedPtr<FPropertyNode>& PropertyNode = Implementation->GetPropertyNode();
-
-	if (!PropertyNode->HasNodeFlags(EPropertyNodeFlags::EditInlineNew))
-	{
-		FString ObjectPathName = NewValue ? NewValue->GetPathName() : TEXT("None");
-		return SetValueFromFormattedString(ObjectPathName, Flags);
-	}
-
-	return FPropertyAccess::Fail;
-}
-
 FPropertyAccess::Result FPropertyHandleObject::GetValue(FAssetData& OutValue) const
 {
 	UObject* ObjectValue = nullptr;
@@ -4858,4 +4820,58 @@ bool FPropertyHandleMap::IsEditable() const
 	return Implementation->HasValidPropertyNode() && !Implementation->IsEditConst() && Implementation->IsPropertyTypeOf(FMapProperty::StaticClass());
 }
 
+// FieldPath
+
+bool FPropertyHandleFieldPath::Supports(TSharedRef<FPropertyNode> PropertyNode)
+{
+	FProperty* Property = PropertyNode->GetProperty();
+
+	if (Property == nullptr)
+	{
+		return false;
+	}
+
+	// The value is a field path
+	return Property->IsA<FFieldPathProperty>();
+}
+
+FPropertyAccess::Result FPropertyHandleFieldPath::GetValue(FProperty*& OutValue) const
+{
+	return FPropertyHandleFieldPath::GetValue((const FProperty*&)OutValue);
+}
+
+FPropertyAccess::Result FPropertyHandleFieldPath::GetValue(const FProperty*& OutValue) const
+{
+	void* PropValue = nullptr;
+	FPropertyAccess::Result Res = Implementation->GetValueData(PropValue);
+
+	if (Res == FPropertyAccess::Success)
+	{
+		FProperty* Property = GetProperty();
+		check(Property->IsA(FFieldPathProperty::StaticClass()));
+		const TFieldPath<FProperty>* FieldPathValue = (const TFieldPath<FProperty>*)(PropValue);
+		OutValue = FieldPathValue->Get();
+	}
+
+	return Res;
+}
+
+FPropertyAccess::Result FPropertyHandleFieldPath::SetValue(FProperty* const& NewValue, EPropertyValueSetFlags::Type Flags)
+{
+	return FPropertyHandleFieldPath::SetValue((const FProperty*)NewValue);
+}
+
+FPropertyAccess::Result FPropertyHandleFieldPath::SetValue(const FProperty* const& NewValue, EPropertyValueSetFlags::Type Flags)
+{
+	const TSharedPtr<FPropertyNode>& PropertyNode = Implementation->GetPropertyNode();
+
+	if (!PropertyNode->HasNodeFlags(EPropertyNodeFlags::EditInlineNew))
+	{
+		FString ObjectPathName = NewValue ? NewValue->GetPathName() : TEXT("None");
+		return SetValueFromFormattedString(ObjectPathName, Flags);
+	}
+
+	return FPropertyAccess::Fail;
+}
+
 #undef LOCTEXT_NAMESPACE
diff --git a/Engine/Source/Editor/PropertyEditor/Private/PropertyHandleImpl.h b/Engine/Source/Editor/PropertyEditor/Private/PropertyHandleImpl.h
--- a/Engine/Source/Editor/PropertyEditor/Private/PropertyHandleImpl.h
+++ b/Engine/Source/Editor/PropertyEditor/Private/PropertyHandleImpl.h
@@ -602,10 +602,6 @@ public:
 	virtual FPropertyAccess::Result SetValue( const UObject* const& InValue, EPropertyValueSetFlags::Type Flags = EPropertyValueSetFlags::DefaultFlags ) override;
 	virtual FPropertyAccess::Result GetValue( FAssetData& OutValue ) const override;
 	virtual FPropertyAccess::Result SetValue( const FAssetData& InValue, EPropertyValueSetFlags::Type Flags = EPropertyValueSetFlags::DefaultFlags ) override;
-	virtual FPropertyAccess::Result GetValue(FProperty*& OutValue) const override;
-	virtual FPropertyAccess::Result SetValue(FProperty* const& InValue, EPropertyValueSetFlags::Type Flags = EPropertyValueSetFlags::DefaultFlags) override;
-	virtual FPropertyAccess::Result GetValue(const FProperty*& OutValue) const override;
-	virtual FPropertyAccess::Result SetValue(const FProperty* const& InValue, EPropertyValueSetFlags::Type Flags = EPropertyValueSetFlags::DefaultFlags) override;
 	virtual FPropertyAccess::Result SetValueFromFormattedString(const FString& InValue, EPropertyValueSetFlags::Type Flags = EPropertyValueSetFlags::DefaultFlags) override;
 	virtual FPropertyAccess::Result SetObjectValueFromSelection() override;
 };
@@ -724,3 +720,14 @@ public:
 	virtual FString GetDocumentationExcerptName() override { return FString("Maps"); }
 	virtual bool IsEditable() const override;
 };
+
+class FPropertyHandleFieldPath : public FPropertyHandleBase
+{
+public:
+	FPropertyHandleFieldPath(TSharedRef<FPropertyNode> PropertyNode, FNotifyHook* NotifyHook, TSharedPtr<IPropertyUtilities> PropertyUtilities);
+	static bool Supports(TSharedRef<FPropertyNode> PropertyNode);
+	virtual FPropertyAccess::Result GetValue(FProperty*& OutValue) const override;
+	virtual FPropertyAccess::Result SetValue(FProperty* const& InValue, EPropertyValueSetFlags::Type Flags = EPropertyValueSetFlags::DefaultFlags) override;
+	virtual FPropertyAccess::Result GetValue(const FProperty*& OutValue) const override;
+	virtual FPropertyAccess::Result SetValue(const FProperty* const& InValue, EPropertyValueSetFlags::Type Flags = EPropertyValueSetFlags::DefaultFlags) override;
+};
