Edit conditions no longer duplicate themselves when shared between multiple properties. Instead, edit conditions are only shown inline when either of these are true:
1. The edit condition property is not marked as editable.
2. The edit condition property is marked as editable, but also has meta=InlineEditConditionToggle.

There's a warning for having a property not be editable and still marked as meta=InlineEditConditionToggle, because the meta tag is redundant.

diff --git a/Engine/Source/Editor/PropertyEditor/Private/Presentation/PropertyEditor/PropertyEditor.cpp b/Engine/Source/Editor/PropertyEditor/Private/Presentation/PropertyEditor/PropertyEditor.cpp
--- a/Engine/Source/Editor/PropertyEditor/Private/Presentation/PropertyEditor/PropertyEditor.cpp
+++ b/Engine/Source/Editor/PropertyEditor/Private/Presentation/PropertyEditor/PropertyEditor.cpp
@@ -614,28 +614,36 @@ bool FPropertyEditor::SupportsEditConditionToggle() const
 	UProperty* Property = PropertyNode->GetProperty();
 
 	static const FName Name_HideEditConditionToggle("HideEditConditionToggle");
-	if (!Property->HasMetaData(Name_HideEditConditionToggle) && EditConditionExpression.IsValid())
+	if (EditConditionExpression.IsValid() && !Property->HasMetaData(Name_HideEditConditionToggle))
 	{
 		const UBoolProperty* ConditionalProperty = EditConditionContext->GetSingleBoolProperty(EditConditionExpression);
 		if (ConditionalProperty != nullptr)
 		{
-			// Both of these conditions are technically bugs of omission.
-			// However, they were supported prior to 4.23, so just warn and allow them for now.
+			// There are 2 valid states for inline edit conditions:
+			// 1. The property is marked as editable and has InlineEditConditionToggle set. 
+			// 2. The property is not marked as editable and does not have InlineEditConditionToggle set.
+			// In both cases, the original property will be hidden and only show up as a toggle.
 
-			// The property is not marked as an inline edit condition toggle.
 			static const FName Name_InlineEditConditionToggle("InlineEditConditionToggle");
-			if (ConditionalProperty->HasMetaData(Name_InlineEditConditionToggle))
+			const bool bIsInlineEditCondition = ConditionalProperty->HasMetaData(Name_InlineEditConditionToggle);
+			const bool bIsEditable = ConditionalProperty->HasAllPropertyFlags(CPF_Edit);
+
+			if (bIsInlineEditCondition == bIsEditable)
 			{
-				UE_LOG(LogPropertyEditor, Warning, TEXT("Property being used as inline edit condition is not marked as InlineEditConditionToggle. Field \"%s\" in class \"%s\"."), *ConditionalProperty->GetNameCPP(), *Property->GetOwnerStruct()->GetName());
+				return true;
 			}
 
-			// The property is not marked as editable.
-			if (!ConditionalProperty->HasAllPropertyFlags(CPF_Edit))
+			if (bIsInlineEditCondition && !bIsEditable)
 			{
-				UE_LOG(LogPropertyEditor, Warning, TEXT("Property being used as inline edit condition is not marked as editable. Field \"%s\" in class \"%s\"."), *ConditionalProperty->GetNameCPP(), *Property->GetOwnerStruct()->GetName());
+				UE_LOG(LogPropertyEditor, Warning, TEXT("Property being used as inline edit condition is not editable, but has redundant InlineEditConditionToggle flag. Field \"%s\" in class \"%s\"."), *ConditionalProperty->GetNameCPP(), *Property->GetOwnerStruct()->GetName());
+				return true;
 			}
 
-			return true;
+			// The property is already shown, and not marked as inline edit condition.
+			if (!bIsInlineEditCondition && bIsEditable)
+			{
+				return false;
+			}
 		}
 	}
 
diff --git a/Engine/Source/Editor/PropertyEditor/Private/PropertyEditorHelpers.cpp b/Engine/Source/Editor/PropertyEditor/Private/PropertyEditorHelpers.cpp
--- a/Engine/Source/Editor/PropertyEditor/Private/PropertyEditorHelpers.cpp
+++ b/Engine/Source/Editor/PropertyEditor/Private/PropertyEditorHelpers.cpp
@@ -334,7 +334,6 @@ TSharedRef<SWidget> SPropertyValueWidget::ConstructPropertyEditorWidget( TShared
 			.Font( FontStyle );
 
 		BasePropertyEditorWidget->GetDesiredWidth( MinDesiredWidth, MaxDesiredWidth );
-
 	}
 
 	return PropertyWidget.ToSharedRef();
@@ -345,43 +344,43 @@ void SEditConditionWidget::Construct( const FArguments& Args, TSharedPtr<FProper
 	PropertyEditor = InPropertyEditor;
 	CustomEditCondition = Args._CustomEditCondition;
 
-	SetVisibility( HasEditCondition() ? EVisibility::Visible : EVisibility::Collapsed );
+	SetVisibility(HasEditConditionToggle() ? EVisibility::Visible : EVisibility::Collapsed);
 
 	ChildSlot
 	[
 		// Some properties become irrelevant depending on the value of other properties.
 		// We prevent the user from editing those properties by disabling their widgets.
 		// This is a shortcut for toggling the property that disables us.
-		SNew( SCheckBox )
-			.OnCheckStateChanged( this, &SEditConditionWidget::OnEditConditionCheckChanged )
-			.IsChecked( this, &SEditConditionWidget::OnGetEditConditionCheckState )
+		SNew(SCheckBox)
+		.OnCheckStateChanged(this, &SEditConditionWidget::OnEditConditionCheckChanged)
+		.IsChecked(this, &SEditConditionWidget::OnGetEditConditionCheckState)
 	];
 }
 
-bool SEditConditionWidget::HasEditCondition() const
-{	
-	return
-			( PropertyEditor.IsValid() && PropertyEditor->HasEditCondition() && PropertyEditor->SupportsEditConditionToggle() )
-		||	( CustomEditCondition.OnEditConditionValueChanged.IsBound() );
+bool SEditConditionWidget::HasEditConditionToggle() const
+{
+	return (PropertyEditor.IsValid() && PropertyEditor->HasEditCondition() && PropertyEditor->SupportsEditConditionToggle())
+		|| (CustomEditCondition.OnEditConditionValueChanged.IsBound());
 }
 
 void SEditConditionWidget::OnEditConditionCheckChanged( ECheckBoxState CheckState )
 {
 	FScopedTransaction EditConditionChangedTransaction(FText::Format(LOCTEXT("UpdatedEditConditionFmt", "{0} Edit Condition Changed"), PropertyEditor->GetDisplayName()));
 
-	if( PropertyEditor.IsValid() && PropertyEditor->HasEditCondition() && PropertyEditor->SupportsEditConditionToggle() )
+	if (PropertyEditor.IsValid() && PropertyEditor->HasEditCondition() && PropertyEditor->SupportsEditConditionToggle())
 	{
 		PropertyEditor->ToggleEditConditionState();
 	}
 	else
 	{
-		CustomEditCondition.OnEditConditionValueChanged.ExecuteIfBound( CheckState == ECheckBoxState::Checked );
+		CustomEditCondition.OnEditConditionValueChanged.ExecuteIfBound(CheckState == ECheckBoxState::Checked);
 	}
 }
 
 ECheckBoxState SEditConditionWidget::OnGetEditConditionCheckState() const
 {
-	bool bEditConditionMet = ( PropertyEditor.IsValid() && PropertyEditor->HasEditCondition() && PropertyEditor->IsEditConditionMet() ) || CustomEditCondition.EditConditionValue.Get();
+	bool bEditConditionMet = (PropertyEditor.IsValid() && PropertyEditor->HasEditCondition() && PropertyEditor->IsEditConditionMet())
+		|| CustomEditCondition.EditConditionValue.Get();
 	return bEditConditionMet ? ECheckBoxState::Checked : ECheckBoxState::Unchecked;
 }
 
diff --git a/Engine/Source/Editor/PropertyEditor/Private/PropertyEditorHelpers.h b/Engine/Source/Editor/PropertyEditor/Private/PropertyEditorHelpers.h
--- a/Engine/Source/Editor/PropertyEditor/Private/PropertyEditorHelpers.h
+++ b/Engine/Source/Editor/PropertyEditor/Private/PropertyEditorHelpers.h
@@ -107,8 +107,9 @@ class SEditConditionWidget : public SCompoundWidget
 
 private:
 	void OnEditConditionCheckChanged( ECheckBoxState CheckState );
-	bool HasEditCondition() const;
+	bool HasEditConditionToggle() const;
 	ECheckBoxState OnGetEditConditionCheckState() const;
+
 private:
 	TSharedPtr<FPropertyEditor> PropertyEditor;
 	FCustomEditCondition CustomEditCondition;
diff --git a/Engine/Source/Editor/UnrealEd/Classes/Editor/PropertyEditorTestObject.h b/Engine/Source/Editor/UnrealEd/Classes/Editor/PropertyEditorTestObject.h
--- a/Engine/Source/Editor/UnrealEd/Classes/Editor/PropertyEditorTestObject.h
+++ b/Engine/Source/Editor/UnrealEd/Classes/Editor/PropertyEditorTestObject.h
@@ -348,7 +348,7 @@ class UPropertyEditorTestObject : public UObject
 	UPROPERTY(EditAnywhere, Category = StructTests, meta = (InlineEditConditionToggle))
 	bool bEditConditionStructWithMultipleInstances2;
 
-	UPROPERTY(EditAnywhere, Category=StructTests, meta=(editcondition = "bEditConditionStructWithMultipleInstances2"))
+	UPROPERTY(EditAnywhere, Category=StructTests, meta=(EditCondition = "bEditConditionStructWithMultipleInstances2"))
 	FPropertyEditorTestBasicStruct StructWithMultipleInstances2;
 
 	UPROPERTY(EditAnywhere, Category=StructTests)
@@ -574,12 +574,12 @@ class UPropertyEditorTestObject : public UObject
 	TArray<FPropertyEditorTestBasicStruct> ArrayOfStructsWithEditCondition;
 
 	UPROPERTY(EditAnywhere, Category = EditCondition)
-	int64 Flags;
+	int64 EditConditionFlags;
 
-	UPROPERTY(EditAnywhere, Category = EditCondition, meta = (EditCondition = "Flags & ETestEnumFlags::Two || Flags & ETestEnumFlags::Four"))
+	UPROPERTY(EditAnywhere, Category = EditCondition, meta = (EditCondition = "EditConditionFlags & ETestEnumFlags::Two || EditConditionFlags & ETestEnumFlags::Four"))
 	bool bEnabledWhenFlagsHasTwoOrFour;
 
-	UPROPERTY(EditAnywhere, Category = EditCondition, meta = (EditCondition = "Flags & ETestEnumFlags::One == false"))
+	UPROPERTY(EditAnywhere, Category = EditCondition, meta = (EditCondition = "EditConditionFlags & ETestEnumFlags::One == false"))
 	bool bDisabledWhenFlagsIsOdd;
 
 	UPROPERTY(EditAnywhere, Category = OnlyInlineProperty, meta = (InlineCategoryProperty))
@@ -596,4 +596,31 @@ class UPropertyEditorTestObject : public UObject
 
 	UPROPERTY(EditAnywhere, Category = DateTime)
 	FDateTime DateTime;
+
+	UPROPERTY()
+	bool bInlineEditConditionWithoutMetaToggle;
+
+	UPROPERTY(EditAnywhere, Category = "Inline Edit Conditions", meta = (EditCondition = "bInlineEditConditionWithoutMetaToggle"))
+	float InlineEditConditionWithoutMeta;
+
+	UPROPERTY(EditAnywhere, Category = "Inline Edit Conditions", meta = (InlineEditConditionToggle))
+	bool bInlineEditConditionWithMetaToggle;
+
+	UPROPERTY(EditAnywhere, Category = "Inline Edit Conditions", meta = (EditCondition = "bInlineEditConditionWithMetaToggle"))
+	float InlineEditConditionWithMeta;
+
+	UPROPERTY(meta = (InlineEditConditionToggle))
+	bool bInlineEditConditionNotEditable;
+
+	UPROPERTY(EditAnywhere, Category = "Inline Edit Conditions", meta = (EditCondition = "bInlineEditConditionNotEditable"))
+	float HasNonEditableInlineCondition;
+
+	UPROPERTY(EditAnywhere, Category = "Inline Edit Conditions")
+	bool bSharedEditCondition;
+
+	UPROPERTY(EditAnywhere, Category = "Inline Edit Conditions", meta = (EditCondition = "bSharedEditCondition"))
+	float UsesSharedEditCondition1;
+
+	UPROPERTY(EditAnywhere, Category = "Inline Edit Conditions", meta = (EditCondition = "bSharedEditCondition"))
+	float UsesSharedEditCondition2;
 };
