Removed dead code

#rnx

diff --git a/Engine/Source/Editor/PropertyEditor/Private/DetailCategoryBuilderImpl.cpp b/Engine/Source/Editor/PropertyEditor/Private/DetailCategoryBuilderImpl.cpp
--- a/Engine/Source/Editor/PropertyEditor/Private/DetailCategoryBuilderImpl.cpp
+++ b/Engine/Source/Editor/PropertyEditor/Private/DetailCategoryBuilderImpl.cpp
@@ -911,8 +911,6 @@ void FDetailCategoryImpl::GenerateChildrenForLayouts()
 		}
 
 		const bool bShowSplitter = bLastItemHasMultipleColumns;
-		// if we are forcing advanced mode disable the dropdown
-		const bool bIsEnabled = !bForceAdvanced;
 		AdvancedDropdownNodeBottom = MakeShareable(new FAdvancedDropdownNode(*this, ShowAdvanced, IsEnabled, AdvancedChildNodes.Num() > 0, SimpleChildNodes.Num() == 0, bShowSplitter));
 	}
 }
