Fixing build breakage.

diff --git a/Engine/Source/Editor/PropertyEditor/Private/PropertyHandleImpl.cpp b/Engine/Source/Editor/PropertyEditor/Private/PropertyHandleImpl.cpp
--- a/Engine/Source/Editor/PropertyEditor/Private/PropertyHandleImpl.cpp
+++ b/Engine/Source/Editor/PropertyEditor/Private/PropertyHandleImpl.cpp
@@ -42,7 +42,7 @@ void PropertyToTextHelper(FString& OutString, FPropertyNode* InPropertyNode, UPr
 
 void PropertyToTextHelper(FString& OutString, FPropertyNode* InPropertyNode, UProperty* Property, const FObjectBaseAddress& ObjectAddress, EPropertyPortFlags PortFlags)
 {
-	bool bIsSparseProperty = InPropertyNode->HasNodeFlags(EPropertyNodeFlags::IsSparseClassData);
+	bool bIsSparseProperty = !!InPropertyNode->HasNodeFlags(EPropertyNodeFlags::IsSparseClassData);
 	bool bIsInContainer = false;
 	UProperty* Outer = Cast<UProperty>(Property->GetOuter());
 	if (bIsSparseProperty)
