Fixing Slate performance regression caused by FProperty changes

diff --git a/Engine/Source/Editor/PropertyEditor/Private/PropertyNode.cpp b/Engine/Source/Editor/PropertyEditor/Private/PropertyNode.cpp
--- a/Engine/Source/Editor/PropertyEditor/Private/PropertyNode.cpp
+++ b/Engine/Source/Editor/PropertyEditor/Private/PropertyNode.cpp
@@ -426,7 +426,7 @@ EPropertyDataValidationResult FPropertyNode::EnsureDataIsValid()
 			FProperty* MyProperty = Property.Get();
 			UStruct* OwnerStruct = MyProperty->GetOwnerStruct();
 
-			if (!OwnerStruct || OwnerStruct->Children == nullptr)
+			if (!OwnerStruct || OwnerStruct->ChildProperties == nullptr)
 			{
 				//verify that the property is not part of an invalid trash class, treat it as an invalid object if it is which will cause a refresh
 				return EPropertyDataValidationResult::ObjectInvalid;
