Fix merge compile error, UProperty -> FProperty

#rnx

diff --git a/Engine/Source/Editor/PropertyEditor/Private/UserInterface/PropertyEditor/SPropertyEditorAsset.cpp b/Engine/Source/Editor/PropertyEditor/Private/UserInterface/PropertyEditor/SPropertyEditorAsset.cpp
--- a/Engine/Source/Editor/PropertyEditor/Private/UserInterface/PropertyEditor/SPropertyEditorAsset.cpp
+++ b/Engine/Source/Editor/PropertyEditor/Private/UserInterface/PropertyEditor/SPropertyEditorAsset.cpp
@@ -205,14 +205,14 @@ void SPropertyEditorAsset::InitializeClassFilters(const FProperty* Property)
 	}
 }
 
-void SPropertyEditorAsset::InitializeAssetDataTags(const UProperty* Property)
+void SPropertyEditorAsset::InitializeAssetDataTags(const FProperty* Property)
 {
 	if (Property == nullptr)
 	{
 		return;
 	}
 
-	const UProperty* MetadataProperty = GetActualMetadataProperty(Property);
+	const FProperty* MetadataProperty = GetActualMetadataProperty(Property);
 	const FString DisallowedAssetDataTagsFilterString = MetadataProperty->GetMetaData(TEXT("DisallowedAssetDataTags"));
 	if (!DisallowedAssetDataTagsFilterString.IsEmpty())
 	{
diff --git a/Engine/Source/Editor/PropertyEditor/Private/UserInterface/PropertyEditor/SPropertyEditorAsset.h b/Engine/Source/Editor/PropertyEditor/Private/UserInterface/PropertyEditor/SPropertyEditorAsset.h
--- a/Engine/Source/Editor/PropertyEditor/Private/UserInterface/PropertyEditor/SPropertyEditorAsset.h
+++ b/Engine/Source/Editor/PropertyEditor/Private/UserInterface/PropertyEditor/SPropertyEditorAsset.h
@@ -330,7 +330,7 @@ private:
 	/**
 	 * Initialize DisallowedAssetDataTags based on the property metadata.
 	 */
-	void InitializeAssetDataTags(const UProperty* Property);
+	void InitializeAssetDataTags(const FProperty* Property);
 
 	/** @return Returns true if the asset is relevant for this property*/
 	bool IsAssetAllowed(const FAssetData& InAssetData);
