Log instead of ensuring when a class has been force deleted so that a class type is now null.

diff --git a/Engine/Source/Editor/PropertyEditor/Private/UserInterface/PropertyEditor/SPropertyEditorAsset.cpp b/Engine/Source/Editor/PropertyEditor/Private/UserInterface/PropertyEditor/SPropertyEditorAsset.cpp
--- a/Engine/Source/Editor/PropertyEditor/Private/UserInterface/PropertyEditor/SPropertyEditorAsset.cpp
+++ b/Engine/Source/Editor/PropertyEditor/Private/UserInterface/PropertyEditor/SPropertyEditorAsset.cpp
@@ -1415,14 +1415,24 @@ UClass* SPropertyEditorAsset::GetObjectPropertyClass(const FProperty* Property)
 	if (CastField<const FObjectPropertyBase>(Property) != nullptr)
 	{
 		Class = CastField<const FObjectPropertyBase>(Property)->PropertyClass;
+		if (Class == nullptr)
+		{
+			UE_LOG(LogPropertyNode, Warning, TEXT("Object Property (%s) has a null class, falling back to UObject"), *Property->GetFullName());
+			Class = UObject::StaticClass();
+		}
 	}
 	else if (CastField<const FInterfaceProperty>(Property) != nullptr)
 	{
 		Class = CastField<const FInterfaceProperty>(Property)->InterfaceClass;
+		if (Class == nullptr)
+		{
+			UE_LOG(LogPropertyNode, Warning, TEXT("Interface Property (%s) has a null class, falling back to UObject"), *Property->GetFullName());
+			Class = UObject::StaticClass();
+		}
 	}
-
-	if (!ensureMsgf(Class != nullptr, TEXT("Property (%s) is not an object or interface class"), Property ? *Property->GetFullName() : TEXT("null")))
+	else
 	{
+		ensureMsgf(Class != nullptr, TEXT("Property (%s) is not an object or interface class"), Property ? *Property->GetFullName() : TEXT("null"));
 		Class = UObject::StaticClass();
 	}
 	return Class;
