Property Editor: Purge killed objects before gathering for changes. This can happen because stale objects still emit PostEditChangeProperty

diff --git a/Engine/Source/Editor/PropertyEditor/Private/PropertyChangeListener.cpp b/Engine/Source/Editor/PropertyEditor/Private/PropertyChangeListener.cpp
--- a/Engine/Source/Editor/PropertyEditor/Private/PropertyChangeListener.cpp
+++ b/Engine/Source/Editor/PropertyEditor/Private/PropertyChangeListener.cpp
@@ -234,6 +234,8 @@ bool FPropertyChangeListener::ScanForChanges( bool bRecacheNewValues )
 			TArray<UObject*> ObjectsThatChanged;
 			if (ObjectNode)
 			{
+				ObjectNode->PurgeKilledObjects();
+
 				for (auto It = ObjectNode->ObjectConstIterator(); It; ++It)
 				{
 					UObject* Obj = It->Get();
@@ -273,6 +275,8 @@ void FPropertyChangeListener::TriggerAllPropertiesChangedDelegate()
 		TArray<UObject*> ObjectsThatChanged;
 		if (ObjectNode)
 		{
+			ObjectNode->PurgeKilledObjects();
+
 			for (auto It = ObjectNode->ObjectConstIterator(); It; ++It)
 			{
 				UObject* Obj = It->Get();
