Added filtering to details panels to filter the set of incoming objects passed to a details panel
Changed the way details panels handle multiple root objects.  The filter can now specify this without setting a specific flag on a details panel at creation time.

diff --git a/Engine/Source/Developer/SettingsEditor/Private/Widgets/SSettingsSectionHeader.cpp b/Engine/Source/Developer/SettingsEditor/Private/Widgets/SSettingsSectionHeader.cpp
--- a/Engine/Source/Developer/SettingsEditor/Private/Widgets/SSettingsSectionHeader.cpp
+++ b/Engine/Source/Developer/SettingsEditor/Private/Widgets/SSettingsSectionHeader.cpp
@@ -10,7 +10,6 @@
 #include "Framework/Application/SlateApplication.h"
 #include "Widgets/Input/SButton.h"
 #include "Widgets/Text/STextBlock.h"
-#include "Widgets/Views/STableRow.h"
 #include "EditorStyleSet.h"
 #include "ISourceControlModule.h"
 #include "SSettingsEditorCheckoutNotice.h"
@@ -19,148 +18,127 @@
 
 #define LOCTEXT_NAMESPACE "SSettingsEditor"
 
-#define S_SETTINGS_SECTION_HEADER_CONSTRUCT_CHILD_SLOT_INIT \
-	ChildSlot \
-	.Padding(FMargin(0.0f, 8.0f, 0.0f, 5.0f)) \
-	[ \
-		SNew(SVerticalBox) \
-		+ SVerticalBox::Slot() \
-		.AutoHeight() \
-		[ \
-			SNew(SHorizontalBox) \
-			+ SHorizontalBox::Slot() \
-			.FillWidth(1.0f) \
-			[ \
-				SNew(SVerticalBox) \
-				+ SVerticalBox::Slot() \
-				.AutoHeight() \
-				[
-
-#define S_SETTINGS_SECTION_HEADER_CONSTRUCT_CHILD_SLOT_END \
-				] \
-				+ SVerticalBox::Slot() \
-				.AutoHeight() \
-				.Padding(0.0f, 8.0f, 0.0f, 0.0f) \
-				[ \
-					/* category description */ \
-					SNew(STextBlock) \
-					.ColorAndOpacity(FSlateColor::UseSubduedForeground()) \
-					.Text(GetSettingsBoxDescriptionText()) \
-				] \
-			] \
-			+ SHorizontalBox::Slot() \
-			.AutoWidth() \
-			.HAlign(HAlign_Right) \
-			.VAlign(VAlign_Bottom) \
-			.Padding(16.0f, 0.0f, 0.0f, 0.0f) \
-			[ \
-				SNew(SHorizontalBox) \
-				.Visibility(this, &SSettingsSectionHeader::HandleButtonRowVisibility) \
-				+SHorizontalBox::Slot() \
-				[ \
-					/* set as default button */ \
-					SNew(SButton) \
-					.Visibility(this, &SSettingsSectionHeader::HandleSetAsDefaultButtonVisibility) \
-					.IsEnabled(this, &SSettingsSectionHeader::HandleSetAsDefaultButtonEnabled) \
-					.OnClicked(this, &SSettingsSectionHeader::HandleSetAsDefaultButtonClicked) \
-					.Text(LOCTEXT("SaveDefaultsButtonText", "Set as Default")) \
-					.ToolTipText(LOCTEXT("SaveDefaultsButtonTooltip", "Save the values below as the new default settings")) \
-				] \
-				+ SHorizontalBox::Slot() \
-				.AutoWidth() \
-				.Padding(8.0f, 0.0f, 0.0f, 0.0f) \
-				[ \
-					/* export button */ \
-					SNew(SButton) \
-					.IsEnabled(this, &SSettingsSectionHeader::HandleExportButtonEnabled) \
-					.OnClicked(this, &SSettingsSectionHeader::HandleExportButtonClicked) \
-					.Text(LOCTEXT("ExportButtonText", "Export...")) \
-					.ToolTipText(LOCTEXT("ExportButtonTooltip", "Export these settings to a file on your computer")) \
-				] \
-				+ SHorizontalBox::Slot() \
-				.AutoWidth() \
-				.Padding(8.0f, 0.0f, 0.0f, 0.0f) \
-				[ \
-					/* import button */ \
-					SNew(SButton) \
-					.IsEnabled(this, &SSettingsSectionHeader::HandleImportButtonEnabled) \
-					.OnClicked(this, &SSettingsSectionHeader::HandleImportButtonClicked) \
-					.Text(LOCTEXT("ImportButtonText", "Import...")) \
-					.ToolTipText(LOCTEXT("ImportButtonTooltip", "Import these settings from a file on your computer")) \
-				] \
-				+ SHorizontalBox::Slot() \
-				.AutoWidth() \
-				.Padding(8.0f, 0.0f, 0.0f, 0.0f) \
-				[ \
-					/* reset defaults button */ \
-					SNew(SButton) \
-					.Visibility(this, &SSettingsSectionHeader::HandleSetAsDefaultButtonVisibility) \
-					.IsEnabled(this, &SSettingsSectionHeader::HandleResetToDefaultsButtonEnabled) \
-					.OnClicked(this, &SSettingsSectionHeader::HandleResetDefaultsButtonClicked) \
-					.Text(LOCTEXT("ResetDefaultsButtonText", "Reset to Defaults")) \
-					.ToolTipText(LOCTEXT("ResetDefaultsButtonTooltip", "Reset the settings below to their default values")) \
-				] \
-			] \
-		] \
-		+ SVerticalBox::Slot() \
-		.AutoHeight() \
-		.Padding(0.0f, 8.0f, 0.0f, 0.0f) \
-		[ \
-			FileWatcherWidget.ToSharedRef() \
-		] \
-	];
-
-void SSettingsSectionHeader::CommonConstruct(const UObject* InSettingsObject, ISettingsEditorModelPtr InModel, TSharedPtr<IDetailsView> InDetailsView)
+void SSettingsSectionHeader::Construct(const FArguments& InArgs, const UObject* InSettingsObject, ISettingsEditorModelPtr InModel, TSharedPtr<IDetailsView> InDetailsView, const TSharedPtr<ITableRow>& InTableRow)
 {
 	Model = InModel;
 	SettingsObject = MakeWeakObjectPtr(const_cast<UObject*>(InSettingsObject));
 	SettingsSection = Model->GetSectionFromSectionObject(InSettingsObject);
 	DetailsView = InDetailsView;
+	TableRow = InTableRow;
 
 	Model->OnSelectionChanged().AddSP(this, &SSettingsSectionHeader::OnSettingsSelectionChanged);
 
 	// Create the watcher widget for the default config file (checks file status / SCC state)
 	FileWatcherWidget =
 		SNew(SSettingsEditorCheckoutNotice)
-		.Visibility(this, &SSettingsSectionHeader::HandleCheckoutNoticeVisibility)
+		.Visibility(this, &SSettingsSectionHeader::GetCheckoutNoticeVisibility)
 		.OnFileProbablyModifiedExternally(this, &SSettingsSectionHeader::HandleCheckoutNoticeFileProbablyModifiedExternally)
 		.ConfigFilePath(this, &SSettingsSectionHeader::HandleCheckoutNoticeConfigFilePath);
-}
-
-void SSettingsSectionHeader::Construct(const FArguments& InArgs, const UObject* InSettingsObject, ISettingsEditorModelPtr InModel, TSharedPtr<IDetailsView> InDetailsView)
-{
-	CommonConstruct(InSettingsObject, InModel, InDetailsView);
 
-	S_SETTINGS_SECTION_HEADER_CONSTRUCT_CHILD_SLOT_INIT
-					// category title
-					SNew(STextBlock)
-					.Font(FEditorStyle::GetFontStyle("SettingsEditor.CatgoryAndSectionFont"))
-					.Text(GetSettingsBoxTitleText())
-	S_SETTINGS_SECTION_HEADER_CONSTRUCT_CHILD_SLOT_END
-}
-
-void SSettingsSectionHeader::Construct(const FArguments& InArgs, const UObject* InSettingsObject, ISettingsEditorModelPtr InModel, TSharedPtr<IDetailsView> InDetailsView, const TSharedRef<ITableRow>& InTableRow)
-{
-	CommonConstruct(InSettingsObject, InModel, InDetailsView);
 
-	S_SETTINGS_SECTION_HEADER_CONSTRUCT_CHILD_SLOT_INIT
+	ChildSlot
+	.Padding(FMargin(0.0f, 8.0f, 0.0f, 5.0f))
+	[
+		SNew(SVerticalBox)
+		+ SVerticalBox::Slot()
+		.AutoHeight()
+		[
+			SNew(SHorizontalBox)
+			+ SHorizontalBox::Slot()
+			.FillWidth(1.0f)
+			[
+				SNew(SVerticalBox)
+				+ SVerticalBox::Slot()
+				.AutoHeight()
+				[
 					SNew(SHorizontalBox)
 					+ SHorizontalBox::Slot()
 					.VAlign(VAlign_Center)
-					.HAlign(HAlign_Left)
+					.Padding(2.0f, 2.0f, 2.0f, 2.0f)
 					.AutoWidth()
 					[
 						SNew(SExpanderArrow, InTableRow)
-						.Visibility(true ? EVisibility::Visible : EVisibility::Collapsed)
 					]
-					+ SHorizontalBox::Slot()
+					+ SHorizontalBox::Slot()// category title
 					[
-					// category title
-					SNew(STextBlock)
-					.Font(FEditorStyle::GetFontStyle("SettingsEditor.CatgoryAndSectionFont"))
-					.Text(GetSettingsBoxTitleText())
+						SNew(STextBlock)
+						.Font(FEditorStyle::GetFontStyle("SettingsEditor.CatgoryAndSectionFont"))
+						.Text(GetSettingsBoxTitleText())
 					]
-	S_SETTINGS_SECTION_HEADER_CONSTRUCT_CHILD_SLOT_END
+				]
+				+ SVerticalBox::Slot()
+				.AutoHeight()
+				.Padding(16.0f, 8.0f, 0.0f, 0.0f)
+				[
+					// category description
+					SNew(STextBlock)
+					.ColorAndOpacity(FSlateColor::UseSubduedForeground())
+					.Text(GetSettingsBoxDescriptionText())
+					.Visibility(this, &SSettingsSectionHeader::GetCategoryDescriptionVisibility)
+				]
+			]
+			+ SHorizontalBox::Slot()
+			.AutoWidth()
+			.HAlign(HAlign_Right)
+			.VAlign(VAlign_Bottom)
+			.Padding(16.0f, 0.0f, 0.0f, 0.0f)
+			[
+				SNew(SHorizontalBox)
+				.Visibility(this, &SSettingsSectionHeader::GetButtonRowVisibility)
+				+SHorizontalBox::Slot()
+				[
+					// set as default button
+					SNew(SButton)
+					.Visibility(this, &SSettingsSectionHeader::HandleSetAsDefaultButtonVisibility)
+					.IsEnabled(this, &SSettingsSectionHeader::HandleSetAsDefaultButtonEnabled)
+					.OnClicked(this, &SSettingsSectionHeader::HandleSetAsDefaultButtonClicked)
+					.Text(LOCTEXT("SaveDefaultsButtonText", "Set as Default"))
+					.ToolTipText(LOCTEXT("SaveDefaultsButtonTooltip", "Save the values below as the new default settings"))
+				]
+				+ SHorizontalBox::Slot()
+				.AutoWidth()
+				.Padding(8.0f, 0.0f, 0.0f, 0.0f)
+				[
+					// export button
+					SNew(SButton)
+					.IsEnabled(this, &SSettingsSectionHeader::HandleExportButtonEnabled)
+					.OnClicked(this, &SSettingsSectionHeader::HandleExportButtonClicked)
+					.Text(LOCTEXT("ExportButtonText", "Export..."))
+					.ToolTipText(LOCTEXT("ExportButtonTooltip", "Export these settings to a file on your computer"))
+				]
+				+ SHorizontalBox::Slot()
+				.AutoWidth()
+				.Padding(8.0f, 0.0f, 0.0f, 0.0f)
+				[
+					// import button
+					SNew(SButton)
+					.IsEnabled(this, &SSettingsSectionHeader::HandleImportButtonEnabled)
+					.OnClicked(this, &SSettingsSectionHeader::HandleImportButtonClicked)
+					.Text(LOCTEXT("ImportButtonText", "Import..."))
+					.ToolTipText(LOCTEXT("ImportButtonTooltip", "Import these settings from a file on your computer"))
+				]
+				+ SHorizontalBox::Slot()
+				.AutoWidth()
+				.Padding(8.0f, 0.0f, 0.0f, 0.0f)
+				[
+					// reset defaults button
+					SNew(SButton)
+					.Visibility(this, &SSettingsSectionHeader::HandleSetAsDefaultButtonVisibility)
+					.IsEnabled(this, &SSettingsSectionHeader::HandleResetToDefaultsButtonEnabled)
+					.OnClicked(this, &SSettingsSectionHeader::HandleResetDefaultsButtonClicked)
+					.Text(LOCTEXT("ResetDefaultsButtonText", "Reset to Defaults"))
+					.ToolTipText(LOCTEXT("ResetDefaultsButtonTooltip", "Reset the settings below to their default values"))
+				]
+			]
+		]
+		+ SVerticalBox::Slot()
+		.AutoHeight()
+		.Padding(0.0f, 8.0f, 0.0f, 0.0f)
+		[
+			FileWatcherWidget.ToSharedRef()
+		]
+	];
+
 }
 
 FText SSettingsSectionHeader::GetSettingsBoxTitleText() const
@@ -184,11 +162,16 @@ FText SSettingsSectionHeader::GetSettingsBoxDescriptionText() const
 	return FText::GetEmpty();
 }
 
-EVisibility SSettingsSectionHeader::HandleButtonRowVisibility() const
+EVisibility SSettingsSectionHeader::GetButtonRowVisibility() const
 {
 	return DetailsView.Pin()->HasActiveSearch() ? EVisibility::Hidden : EVisibility::Visible;
 }
 
+EVisibility SSettingsSectionHeader::GetCategoryDescriptionVisibility() const
+{
+	return TableRow.IsValid() && TableRow.Pin()->IsItemExpanded() ? EVisibility::Visible : EVisibility::Collapsed;
+}
+
 FReply SSettingsSectionHeader::HandleExportButtonClicked()
 {
 	if(SettingsSection.IsValid())
@@ -513,7 +496,7 @@ void SSettingsSectionHeader::HandleCheckoutNoticeFileProbablyModifiedExternally(
 }
 
 /** Callback for determining the visibility of the 'Locked' notice. */
-EVisibility SSettingsSectionHeader::HandleCheckoutNoticeVisibility() const
+EVisibility SSettingsSectionHeader::GetCheckoutNoticeVisibility() const
 {
 	// Only DefaultConfig are under source control, so the checkout notice should not be visible for the other cases
 	if(SettingsObject.IsValid() && SettingsObject->GetClass()->HasAnyClassFlags(CLASS_DefaultConfig) && !DetailsView.Pin()->HasActiveSearch())
@@ -543,12 +526,7 @@ void FSettingsDetailRootObjectCustomization::Initialize()
 	OnSelectedSectionChanged();
 }
 
-TSharedPtr<SWidget> FSettingsDetailRootObjectCustomization::CustomizeObjectHeader(const UObject* InRootObject)
-{
-	return SNew(SSettingsSectionHeader, InRootObject, Model, DetailsView.Pin());
-}
-
-TSharedPtr<SWidget> FSettingsDetailRootObjectCustomization::CustomizeObjectHeader(const UObject* InRootObject, const TSharedRef<ITableRow>& InTableRow)
+TSharedPtr<SWidget> FSettingsDetailRootObjectCustomization::CustomizeObjectHeader(const UObject* InRootObject, const TSharedPtr<ITableRow>& InTableRow)
 {
 	return SNew(SSettingsSectionHeader, InRootObject, Model, DetailsView.Pin(), InTableRow);
 }
diff --git a/Engine/Source/Developer/SettingsEditor/Private/Widgets/SSettingsSectionHeader.h b/Engine/Source/Developer/SettingsEditor/Private/Widgets/SSettingsSectionHeader.h
--- a/Engine/Source/Developer/SettingsEditor/Private/Widgets/SSettingsSectionHeader.h
+++ b/Engine/Source/Developer/SettingsEditor/Private/Widgets/SSettingsSectionHeader.h
@@ -22,8 +22,7 @@ public:
 	{}
 	SLATE_END_ARGS()
 
-	void Construct(const FArguments& InArgs, const UObject* InSettingsObject, ISettingsEditorModelPtr InModel, TSharedPtr<IDetailsView> InDetailsView);
-	void Construct(const FArguments& InArgs, const UObject* InSettingsObject, ISettingsEditorModelPtr InModel, TSharedPtr<IDetailsView> InDetailsView, const TSharedRef<ITableRow>& InTableRow);
+	void Construct(const FArguments& InArgs, const UObject* InSettingsObject, ISettingsEditorModelPtr InModel, TSharedPtr<IDetailsView> InDetailsView, const TSharedPtr<ITableRow>& InTableRow);
 
 private:
 
@@ -92,9 +91,11 @@ private:
 	void HandleCheckoutNoticeFileProbablyModifiedExternally();
 
 	/** Callback for determining the visibility of the 'Locked' notice. */
-	EVisibility HandleCheckoutNoticeVisibility() const;
+	EVisibility GetCheckoutNoticeVisibility() const;
 
-	EVisibility HandleButtonRowVisibility() const;
+	EVisibility GetButtonRowVisibility() const;
+
+	EVisibility GetCategoryDescriptionVisibility() const;
 
 	void OnSettingsSelectionChanged();
 private:
@@ -105,6 +106,7 @@ private:
 	ISettingsSectionPtr SettingsSection;
 	TWeakObjectPtr<UObject> SettingsObject;
 	TWeakPtr<IDetailsView> DetailsView;
+	TWeakPtr<ITableRow> TableRow;
 };
 
 
@@ -116,10 +118,10 @@ public:
 	void Initialize();
 
 	/** IDetailRootObjectCustomization interface */
-	virtual TSharedPtr<SWidget> CustomizeObjectHeader(const UObject* InRootObject) override;
-	virtual TSharedPtr<SWidget> CustomizeObjectHeader(const UObject* InRootObject, const TSharedRef<ITableRow>& InTableRow) override;
+	virtual TSharedPtr<SWidget> CustomizeObjectHeader(const UObject* InRootObject, const TSharedPtr<ITableRow>& InTableRow) override;
 	virtual bool IsObjectVisible(const UObject* InRootObject) const override;
 	virtual bool ShouldDisplayHeader(const UObject* InRootObject) const override;
+	virtual EExpansionArrowUsage GetExpansionArrowUsage() const override { return EExpansionArrowUsage::Custom; }
 private:
 	void OnSelectedSectionChanged();
 
diff --git a/Engine/Source/Editor/LevelEditor/Private/LevelEditor.cpp b/Engine/Source/Editor/LevelEditor/Private/LevelEditor.cpp
--- a/Engine/Source/Editor/LevelEditor/Private/LevelEditor.cpp
+++ b/Engine/Source/Editor/LevelEditor/Private/LevelEditor.cpp
@@ -228,6 +228,7 @@ TSharedRef<SDockTab> FLevelEditorModule::SpawnLevelEditor( const FSpawnTabArgs&
 			}
 		}
 
+		LevelEditorCreatedEvent.Broadcast(LevelEditorTmp);
 	}
 
 	IIntroTutorials& IntroTutorials = FModuleManager::LoadModuleChecked<IIntroTutorials>(TEXT("IntroTutorials"));
diff --git a/Engine/Source/Editor/LevelEditor/Private/SActorDetails.cpp b/Engine/Source/Editor/LevelEditor/Private/SActorDetails.cpp
--- a/Engine/Source/Editor/LevelEditor/Private/SActorDetails.cpp
+++ b/Engine/Source/Editor/LevelEditor/Private/SActorDetails.cpp
@@ -234,7 +234,7 @@ SActorDetails::~SActorDetails()
 	}
 }
 
-void SActorDetails::OnDetailsViewObjectArrayChanged(const FString& InTitle, const TArray<TWeakObjectPtr<UObject>>& InObjects)
+void SActorDetails::OnDetailsViewObjectArrayChanged(const FString& InTitle, const TArray<UObject*>& InObjects)
 {
 	// The DetailsView will already check validity every tick and hide itself when invalid, so this piggy-backs on that code instead of needing a second tick function.
 	if (InObjects.Num() == 0 && !LockedActorSelection.IsValid())
@@ -318,6 +318,11 @@ void SActorDetails::PostRedo(bool bSuccess)
 	PostUndo(bSuccess);
 }
 
+void SActorDetails::SetActorDetailsFilter(TSharedPtr<FDetailsViewObjectFilter> InFilter)
+{
+	DetailsView->SetObjectFilter(InFilter);
+}
+
 void SActorDetails::OnComponentsEditedInWorld()
 {
 	if (GetSelectedActorInEditor() == GetActorContext())
diff --git a/Engine/Source/Editor/LevelEditor/Private/SActorDetails.h b/Engine/Source/Editor/LevelEditor/Private/SActorDetails.h
--- a/Engine/Source/Editor/LevelEditor/Private/SActorDetails.h
+++ b/Engine/Source/Editor/LevelEditor/Private/SActorDetails.h
@@ -45,7 +45,13 @@ public:
 	virtual void PostRedo(bool bSuccess) override;
 
 	static LEVELEDITOR_API IConsoleVariable* ShowComponents;
-
+	
+	/**
+	 * Sets the filter that should be used to filter incoming actors in or out of the details panel
+	 *
+	 * @param InFilter	The filter to use or nullptr to remove the active filter
+	 */
+	void SetActorDetailsFilter(TSharedPtr<FDetailsViewObjectFilter> InFilter);
 private:
 	AActor* GetSelectedActorInEditor() const;
 	AActor* GetActorContext() const;
@@ -57,7 +63,7 @@ private:
 	void OnSCSEditorTreeViewSelectionChanged(const TArray<TSharedPtr<class FSCSEditorTreeNode> >& SelectedNodes);
 	void OnSCSEditorTreeViewItemDoubleClicked(const TSharedPtr<class FSCSEditorTreeNode> ClickedNode);
 	void UpdateComponentTreeFromEditorSelection();
-	void OnDetailsViewObjectArrayChanged(const FString& InTitle, const TArray<TWeakObjectPtr<UObject>>& InObjects);
+	void OnDetailsViewObjectArrayChanged(const FString& InTitle, const TArray<UObject*>& InObjects);
 
 	bool IsPropertyReadOnly(const struct FPropertyAndParent& PropertyAndParent) const;
 	bool IsPropertyEditingEnabled() const;
diff --git a/Engine/Source/Editor/LevelEditor/Private/SLevelEditor.cpp b/Engine/Source/Editor/LevelEditor/Private/SLevelEditor.cpp
--- a/Engine/Source/Editor/LevelEditor/Private/SLevelEditor.cpp
+++ b/Engine/Source/Editor/LevelEditor/Private/SLevelEditor.cpp
@@ -1702,6 +1702,17 @@ TSharedRef<SWidget> SLevelEditor::CreateActorDetails( const FName TabIdentifier
 }
 
 
+void SLevelEditor::SetActorDetailsFilter(TSharedPtr<FDetailsViewObjectFilter> ActorDetailsFilter)
+{
+	for (TWeakPtr<SActorDetails> Details : AllActorDetailPanels)
+	{
+		if (Details.IsValid())
+		{
+			Details.Pin()->SetActorDetailsFilter(ActorDetailsFilter);
+		}
+	}
+}
+
 TSharedRef<SWidget> SLevelEditor::CreateToolBox()
 {
 	TSharedRef<SLevelEditorToolBox> NewToolBox =
diff --git a/Engine/Source/Editor/LevelEditor/Private/SLevelEditor.h b/Engine/Source/Editor/LevelEditor/Private/SLevelEditor.h
--- a/Engine/Source/Editor/LevelEditor/Private/SLevelEditor.h
+++ b/Engine/Source/Editor/LevelEditor/Private/SLevelEditor.h
@@ -121,6 +121,7 @@ public:
 	virtual void OnToolkitHostingFinished( const TSharedRef< class IToolkit >& Toolkit ) override;
 	virtual UWorld* GetWorld() const override;
 	virtual TSharedRef<SWidget> CreateActorDetails( const FName TabIdentifier ) override;
+	virtual void SetActorDetailsFilter(TSharedPtr<FDetailsViewObjectFilter> ActorDetailsFilter) override;
 	virtual TSharedRef<SWidget> CreateToolBox() override;
 
 	/** SWidget overrides */
diff --git a/Engine/Source/Editor/LevelEditor/Public/ILevelEditor.h b/Engine/Source/Editor/LevelEditor/Public/ILevelEditor.h
--- a/Engine/Source/Editor/LevelEditor/Public/ILevelEditor.h
+++ b/Engine/Source/Editor/LevelEditor/Public/ILevelEditor.h
@@ -12,6 +12,7 @@
 #include "AssetThumbnail.h"
 #include "Toolkits/IToolkitHost.h"
 
+class FDetailsViewObjectFilter;
 class ISceneOutliner;
 class IAssetViewport;
 class SLevelViewport;
@@ -58,6 +59,9 @@ public:
 	/** Spawns an Actor Details widget */
 	virtual TSharedRef<SWidget> CreateActorDetails( const FName TabIdentifier ) = 0;
 
+	/** Set the filter that should be used to determine the set of objects that should be shown in a details panel when an actor in the level editor is selected */
+	virtual void SetActorDetailsFilter(TSharedPtr<FDetailsViewObjectFilter> ActorDetailsFilter) = 0;
+
 	/** Spawns a level editor ToolBox widget (aka. "Modes") */
 	virtual TSharedRef<SWidget> CreateToolBox() = 0;
 
diff --git a/Engine/Source/Editor/LevelEditor/Public/LevelEditor.h b/Engine/Source/Editor/LevelEditor/Public/LevelEditor.h
--- a/Engine/Source/Editor/LevelEditor/Public/LevelEditor.h
+++ b/Engine/Source/Editor/LevelEditor/Public/LevelEditor.h
@@ -199,7 +199,9 @@ public:
 	DECLARE_EVENT(FLevelEditorModule, FTabContentChangedEvent);
 	virtual FTabContentChangedEvent& OnTabContentChanged() { return TabContentChangedEvent; }
 
-
+	/** Called when a level editor widget has been created */
+	DECLARE_EVENT_OneParam(FLevelEditorModule, FOnLevelEditorCreated, TSharedPtr<ILevelEditor>);
+	virtual FOnLevelEditorCreated& OnLevelEditorCreated() { return LevelEditorCreatedEvent; }
 	/**
 	 * Called when actor selection changes
 	 * 
@@ -401,6 +403,9 @@ private:
 	/** Multicast delegate executed when the tab content is changed */
 	FTabContentChangedEvent TabContentChangedEvent;
 
+	/** Multicast delegate executed when a level editor has been created */
+	FOnLevelEditorCreated LevelEditorCreatedEvent;
+
 	/** Multicast delegate executed when actor selection changes */
 	FActorSelectionChangedEvent ActorSelectionChangedEvent;
 
diff --git a/Engine/Source/Editor/PropertyEditor/Private/DetailLayoutBuilderImpl.cpp b/Engine/Source/Editor/PropertyEditor/Private/DetailLayoutBuilderImpl.cpp
--- a/Engine/Source/Editor/PropertyEditor/Private/DetailLayoutBuilderImpl.cpp
+++ b/Engine/Source/Editor/PropertyEditor/Private/DetailLayoutBuilderImpl.cpp
@@ -18,8 +18,8 @@ FDetailLayoutBuilderImpl::FDetailLayoutBuilderImpl(TSharedPtr<FComplexPropertyNo
 	, PropertyGenerationUtilities( InPropertyGenerationUtilities )
 	, DetailsView( InDetailsView.Get() )
 	, CurrentCustomizationClass( nullptr )
+	, bLayoutForExternalRoot(bIsExternal)
 {
-	bLayoutForExternalRoot = bIsExternal;
 }
 
 
@@ -326,10 +326,11 @@ void FDetailLayoutBuilderImpl::GenerateDetailLayout()
 		CategoryNodes.AddUnique(AdvancedOnlyCategories[CategoryIndex]);
 	}
 
-	if(DetailsView && DetailsView->ContainsMultipleTopLevelObjects())
+	TSharedPtr<FComplexPropertyNode> RootNodePinned = RootNode.Pin();
+	if(DetailsView->GetRootObjectCustomization() && RootNodePinned->GetInstancesNum())
 	{
 		// This should always exist here
-		UObject* RootObject = RootNode.Pin()->AsObjectNode()->GetUObject(0);
+		UObject* RootObject = RootNodePinned->AsObjectNode()->GetUObject(0);
 		check(RootObject);
 
 		TSharedPtr<IDetailRootObjectCustomization> RootObjectCustomization = DetailsView->GetRootObjectCustomization();
diff --git a/Engine/Source/Editor/PropertyEditor/Private/DetailMultiTopLevelObjectRootNode.cpp b/Engine/Source/Editor/PropertyEditor/Private/DetailMultiTopLevelObjectRootNode.cpp
--- a/Engine/Source/Editor/PropertyEditor/Private/DetailMultiTopLevelObjectRootNode.cpp
+++ b/Engine/Source/Editor/PropertyEditor/Private/DetailMultiTopLevelObjectRootNode.cpp
@@ -4,28 +4,30 @@
 #include "IDetailRootObjectCustomization.h"
 #include "DetailWidgetRow.h"
 
-void SDetailMultiTopLevelObjectTableRow::Construct(const FArguments& InArgs, const TSharedRef<FDetailTreeNode>& InOwnerTreeNode, const TSharedRef<SWidget>& InCustomizedWidgetContents, const TSharedRef<STableViewBase>& InOwnerTableView)
-{
-	Construct(InArgs, InOwnerTreeNode);
-	ChildSlotConstruct(InCustomizedWidgetContents, InOwnerTableView);
-	// Disable toggle functionality
-	bShowExpansionArrow = false;
-}
-
-void SDetailMultiTopLevelObjectTableRow::Construct(const FArguments& InArgs, const TSharedRef<FDetailTreeNode>& InOwnerTreeNode)
+void SDetailMultiTopLevelObjectTableRow::Construct( const FArguments& InArgs, TSharedRef<FDetailTreeNode> InOwnerTreeNode, const TSharedRef<STableViewBase>& InOwnerTableView )
 {
 	OwnerTreeNode = InOwnerTreeNode;
-	bShowExpansionArrow = InArgs._ShowExpansionArrow;
-}
+	ExpansionArrowUsage = InArgs._ExpansionArrowUsage;
 
-void SDetailMultiTopLevelObjectTableRow::ChildSlotConstruct(const TSharedRef<SWidget>& InCustomizedWidgetContents, const TSharedRef<STableViewBase>& InOwnerTableView)
-{
 	ChildSlot
-	[
+	[	
 		SNew( SBox )
 		.Padding( FMargin( 0.0f, 0.0f, SDetailTableRowBase::ScrollbarPaddingSize, 0.0f ) )
 		[
-			InCustomizedWidgetContents
+			SNew( SHorizontalBox )
+			+SHorizontalBox::Slot()
+			.VAlign(VAlign_Center)
+			.Padding(2.0f, 2.0f, 2.0f, 2.0f)
+			.AutoWidth()
+			[
+				SNew( SExpanderArrow, SharedThis(this) )
+				.Visibility(ExpansionArrowUsage == EExpansionArrowUsage::Default ? EVisibility::Visible : EVisibility::Collapsed)
+			]
+			+SHorizontalBox::Slot()
+			.Expose(ContentSlot)
+			[
+				SNullWidget::NullWidget
+			]
 		]
 	];
 
@@ -35,11 +37,17 @@ void SDetailMultiTopLevelObjectTableRow::ChildSlotConstruct(const TSharedRef<SWi
 			.ShowSelection(false),
 		InOwnerTableView
 	);
-	// Enable toggle functionality
-	bShowExpansionArrow = true;
 }
 
 
+void SDetailMultiTopLevelObjectTableRow::SetContent(TSharedRef<SWidget> InContent)
+{
+	(*ContentSlot)
+	[
+		InContent
+	];
+}
+
 const FSlateBrush* SDetailMultiTopLevelObjectTableRow::GetBackgroundImage() const
 {
 	if (IsHovered())
@@ -54,7 +62,7 @@ const FSlateBrush* SDetailMultiTopLevelObjectTableRow::GetBackgroundImage() cons
 
 FReply SDetailMultiTopLevelObjectTableRow::OnMouseButtonDown( const FGeometry& MyGeometry, const FPointerEvent& MouseEvent )
 {
-	if (bShowExpansionArrow && MouseEvent.GetEffectingButton() == EKeys::LeftMouseButton)
+	if (ExpansionArrowUsage != EExpansionArrowUsage::None && MouseEvent.GetEffectingButton() == EKeys::LeftMouseButton)
 	{
 		ToggleExpansion();
 		return FReply::Handled();
@@ -98,68 +106,31 @@ ENodeVisibility FDetailMultiTopLevelObjectRootNode::GetVisibility() const
 
 TSharedRef< ITableRow > FDetailMultiTopLevelObjectRootNode::GenerateWidgetForTableView(const TSharedRef<STableViewBase>& OwnerTable, const FDetailColumnSizeData& ColumnSizeData, bool bAllowFavoriteSystem)
 {
-	FDetailWidgetRow Row;
-	TSharedPtr<SDetailMultiTopLevelObjectTableRow> DetailMultiTopLevelObjectTableRow;
-	if (RootObjectCustomization.IsValid())
-	{
-		
-		DetailMultiTopLevelObjectTableRow = SNew(SDetailMultiTopLevelObjectTableRow, AsShared());
-		GenerateStandaloneWidget(Row, DetailMultiTopLevelObjectTableRow.ToSharedRef());
-		DetailMultiTopLevelObjectTableRow->ChildSlotConstruct(Row.NameWidget.Widget, OwnerTable);
-		bShouldShowOnlyChildren = false;
-	}
-	else
-	{
-		GenerateStandaloneWidget(Row);
-		DetailMultiTopLevelObjectTableRow = SNew(SDetailMultiTopLevelObjectTableRow, AsShared(), Row.NameWidget.Widget, OwnerTable);
-		bShouldShowOnlyChildren = true;
-	}
-	
-	return DetailMultiTopLevelObjectTableRow.ToSharedRef();
-}
+	EExpansionArrowUsage ExpansionArrowUsage = EExpansionArrowUsage::None;
 
-
-bool GenerateStandaloneWidgetInternal(FDetailWidgetRow& OutRow, TSharedPtr<SWidget>& HeaderWidget, const FName& NodeName)
-{
-	if (!HeaderWidget.IsValid())
+	TSharedPtr<IDetailRootObjectCustomization> Customization = RootObjectCustomization.Pin();
+	if (Customization)
 	{
-		// no customization was supplied or was passed back from the interface as invalid
-		// just make a text block with the name
-		HeaderWidget =
-			SNew(STextBlock)
-			.Font(FEditorStyle::GetFontStyle("DetailsView.CategoryFontStyle"))
-			.Text(FText::FromName(NodeName));
+		ExpansionArrowUsage = Customization->GetExpansionArrowUsage();
 	}
 
-	OutRow.NameContent()
-	[
-		HeaderWidget.ToSharedRef()
-	];
+	TSharedRef<SDetailMultiTopLevelObjectTableRow> TableRowWidget =
+		SNew(SDetailMultiTopLevelObjectTableRow, AsShared(), OwnerTable)
+		.ExpansionArrowUsage(ExpansionArrowUsage);
 
-	return true;
-}
+	FDetailWidgetRow Row;
+	GenerateWidget_Internal(Row, TableRowWidget);
 
-bool FDetailMultiTopLevelObjectRootNode::GenerateStandaloneWidget(FDetailWidgetRow& OutRow) const
-{
-	TSharedPtr<SWidget> HeaderWidget;
-	if (RootObjectCustomization.IsValid() && RootObject.IsValid())
-	{
-		HeaderWidget = RootObjectCustomization.Pin()->CustomizeObjectHeader(RootObject.Get());
-	}
+	TableRowWidget->SetContent(Row.NameWidget.Widget);
 
-	return GenerateStandaloneWidgetInternal(OutRow, HeaderWidget, NodeName);
+	return TableRowWidget;
 }
 
 
-bool FDetailMultiTopLevelObjectRootNode::GenerateStandaloneWidget(FDetailWidgetRow& OutRow, const TSharedRef<ITableRow>& InTableRow) const
+bool FDetailMultiTopLevelObjectRootNode::GenerateStandaloneWidget(FDetailWidgetRow& OutRow) const
 {
-	TSharedPtr<SWidget> HeaderWidget;
-	if (RootObjectCustomization.IsValid() && RootObject.IsValid())
-	{
-		HeaderWidget = RootObjectCustomization.Pin()->CustomizeObjectHeader(RootObject.Get(), InTableRow);
-	}
-
-	return GenerateStandaloneWidgetInternal(OutRow, HeaderWidget, NodeName);
+	GenerateWidget_Internal(OutRow, nullptr);
+	return true;
 }
 
 
@@ -207,3 +178,28 @@ bool FDetailMultiTopLevelObjectRootNode::ShouldShowOnlyChildren() const
 {
 	return RootObjectCustomization.IsValid() && RootObject.IsValid() ? !RootObjectCustomization.Pin()->ShouldDisplayHeader(RootObject.Get()) : bShouldShowOnlyChildren;
 }
+
+void FDetailMultiTopLevelObjectRootNode::GenerateWidget_Internal(FDetailWidgetRow& OutRow, TSharedPtr<SDetailMultiTopLevelObjectTableRow> TableRowWidget) const
+{
+	TSharedPtr<SWidget> HeaderWidget;
+	if (RootObjectCustomization.IsValid() && RootObject.IsValid())
+	{
+		HeaderWidget = RootObjectCustomization.Pin()->CustomizeObjectHeader(RootObject.Get(), TableRowWidget);
+	}
+
+	if (!HeaderWidget.IsValid())
+	{
+		// no customization was supplied or was passed back from the interface as invalid
+		// just make a text block with the name
+		HeaderWidget =
+			SNew(STextBlock)
+			.Font(FEditorStyle::GetFontStyle("DetailsView.CategoryFontStyle"))
+			.Text(FText::FromName(NodeName));
+	}
+
+	OutRow.NameContent()
+	[
+		HeaderWidget.ToSharedRef()
+	];
+
+}
diff --git a/Engine/Source/Editor/PropertyEditor/Private/DetailMultiTopLevelObjectRootNode.h b/Engine/Source/Editor/PropertyEditor/Private/DetailMultiTopLevelObjectRootNode.h
--- a/Engine/Source/Editor/PropertyEditor/Private/DetailMultiTopLevelObjectRootNode.h
+++ b/Engine/Source/Editor/PropertyEditor/Private/DetailMultiTopLevelObjectRootNode.h
@@ -12,38 +12,32 @@
 #include "Widgets/Views/STableRow.h"
 #include "SDetailsViewBase.h"
 #include "SDetailTableRowBase.h"
+#include "IDetailRootObjectCustomization.h"
 
-class IDetailRootObjectCustomization;
+using EExpansionArrowUsage = IDetailRootObjectCustomization::EExpansionArrowUsage;
 
 class SDetailMultiTopLevelObjectTableRow : public SDetailTableRowBase
 {
 public:
 	SLATE_BEGIN_ARGS(SDetailMultiTopLevelObjectTableRow)
 		: _DisplayName()
-		, _ShowExpansionArrow(false)
+		, _ExpansionArrowUsage(EExpansionArrowUsage::None)
 	{}
 		SLATE_ARGUMENT( FText, DisplayName )
-		SLATE_ARGUMENT( bool, ShowExpansionArrow )
+		SLATE_ARGUMENT(IDetailRootObjectCustomization::EExpansionArrowUsage, ExpansionArrowUsage)
 	SLATE_END_ARGS()
 
-	/**
-	* Use this construct when the toggle functionality (with the SExpanderArrow) is not desired.
-	* GenerateStandaloneWidget(Row) must be called before this function.
-	*/
-	void Construct( const FArguments& InArgs, const TSharedRef<FDetailTreeNode>& InOwnerTreeNode, const TSharedRef<SWidget>& InCustomizedWidgetContents, const TSharedRef<STableViewBase>& InOwnerTableView );
-	/**
-	* Use this construct together with ChildSlotConstruct when the toggle functionality (with the SExpanderArrow) is desired.
-	* GenerateStandaloneWidget(Row, DetailMultiTopLevelObjectTableRow) and ChildSlotConstruct(Row.NameWidget.Widget, OwnerTable) must be called after this function.
-	*/
-	void Construct( const FArguments& InArgs, const TSharedRef<FDetailTreeNode>& InOwnerTreeNode );
-	void ChildSlotConstruct(const TSharedRef<SWidget>& InCustomizedWidgetContents, const TSharedRef<STableViewBase>& InOwnerTableView);
+	void Construct(const FArguments& InArgs, TSharedRef<FDetailTreeNode> InOwnerTreeNode, const TSharedRef<STableViewBase>& InOwnerTableView);
+	void SetContent(TSharedRef<SWidget> InContent);
 private:
 	const FSlateBrush* GetBackgroundImage() const;
+
 private:
 	virtual FReply OnMouseButtonDown( const FGeometry& MyGeometry, const FPointerEvent& MouseEvent ) override;
 	virtual FReply OnMouseButtonDoubleClick( const FGeometry& InMyGeometry, const FPointerEvent& InMouseEvent ) override;
 private:
-	bool bShowExpansionArrow;
+	EExpansionArrowUsage ExpansionArrowUsage;
+	SHorizontalBox::FSlot* ContentSlot = nullptr;
 };
 
 
@@ -51,21 +45,13 @@ class FDetailMultiTopLevelObjectRootNode : public FDetailTreeNode, public TShare
 {
 public:
 	FDetailMultiTopLevelObjectRootNode( const FDetailNodeList& InChildNodes, const TSharedPtr<IDetailRootObjectCustomization>& RootObjectCustomization, IDetailsViewPrivate* InDetailsView, const UObject& InRootObject );
-
 private:
 	virtual IDetailsViewPrivate* GetDetailsView() const override{ return DetailsView; }
 	virtual void OnItemExpansionChanged( bool bIsExpanded, bool bShouldSaveState ) override {}
 	virtual bool ShouldBeExpanded() const override { return true; }
 	virtual ENodeVisibility GetVisibility() const override;
-	virtual TSharedRef< ITableRow > GenerateWidgetForTableView( const TSharedRef<STableViewBase>& OwnerTable, const FDetailColumnSizeData& ColumnSizeData, bool bAllowFavoriteSystem) override;
-	/**
-	* Use this version of GenerateStandaloneWidget when the toggle functionality (with the SExpanderArrow) is not used.
-	*/
+	virtual TSharedRef< ITableRow > GenerateWidgetForTableView(const TSharedRef<STableViewBase>& OwnerTable, const FDetailColumnSizeData& ColumnSizeData, bool bAllowFavoriteSystem) override;
 	virtual bool GenerateStandaloneWidget(FDetailWidgetRow& OutRow) const override;
-	/**
-	* Use this version of GenerateStandaloneWidget when the toggle functionality (with the SExpanderArrow) is being used.
-	*/
-	virtual bool GenerateStandaloneWidget(FDetailWidgetRow& OutRow, const TSharedRef<ITableRow>& InTableRow) const;
 	virtual void GetChildren(FDetailNodeList& OutChildren )  override;
 	virtual void FilterNode( const FDetailFilter& InFilter ) override;
 	virtual void Tick( float DeltaTime ) override {}
@@ -73,6 +59,8 @@ private:
 	virtual FName GetNodeName() const override { return NodeName; }
 	virtual EDetailNodeType GetNodeType() const override { return EDetailNodeType::Object; }
 	virtual TSharedPtr<IPropertyHandle> CreatePropertyHandle() const override { return nullptr; }
+private:
+	void GenerateWidget_Internal(FDetailWidgetRow& Row, TSharedPtr<SDetailMultiTopLevelObjectTableRow> TableRow) const;
 private:
 	FDetailNodeList ChildNodes;
 	IDetailsViewPrivate* DetailsView;
diff --git a/Engine/Source/Editor/PropertyEditor/Private/DetailsViewGenericObjectFilter.cpp b/Engine/Source/Editor/PropertyEditor/Private/DetailsViewGenericObjectFilter.cpp
new file mode 100644
--- /dev/null
+++ b/Engine/Source/Editor/PropertyEditor/Private/DetailsViewGenericObjectFilter.cpp
@@ -0,0 +1,28 @@
+// Copyright Epic Games, Inc. All Rights Reserved.
+
+#include "DetailsViewGenericObjectFilter.h"
+
+
+FDetailsViewDefaultObjectFilter::FDetailsViewDefaultObjectFilter(bool bInAllowMultipleRoots)
+	: bAllowMultipleRoots(bInAllowMultipleRoots)
+{
+}
+
+TArray<FDetailsViewObjectRoot> FDetailsViewDefaultObjectFilter::FilterObjects(const TArray<UObject*>& SourceObjects)
+{
+	TArray<FDetailsViewObjectRoot> Roots;
+	if (bAllowMultipleRoots)
+	{
+		Roots.Reserve(SourceObjects.Num());
+		for (UObject* Object : SourceObjects)
+		{
+			Roots.Emplace(Object);
+		}
+	}
+	else
+	{
+		Roots.Emplace(SourceObjects);
+	}
+
+	return Roots;
+}
diff --git a/Engine/Source/Editor/PropertyEditor/Private/DetailsViewGenericObjectFilter.h b/Engine/Source/Editor/PropertyEditor/Private/DetailsViewGenericObjectFilter.h
new file mode 100644
--- /dev/null
+++ b/Engine/Source/Editor/PropertyEditor/Private/DetailsViewGenericObjectFilter.h
@@ -0,0 +1,22 @@
+// Copyright Epic Games, Inc. All Rights Reserved.
+
+#pragma once
+
+#include "DetailsViewObjectFilter.h"
+
+/**
+ * The default object set is just a pass through to the details panel.  
+ */
+class FDetailsViewDefaultObjectFilter : public FDetailsViewObjectFilter
+{
+public:
+	FDetailsViewDefaultObjectFilter(bool bInAllowMultipleRoots);
+
+	/**
+	 * FDetailsViewObjectFilter interface
+	 */
+	virtual TArray<FDetailsViewObjectRoot> FilterObjects(const TArray<UObject*>& SourceObjects);
+
+private:
+	bool bAllowMultipleRoots;
+};
\ No newline at end of file
diff --git a/Engine/Source/Editor/PropertyEditor/Private/PropertyEditorModule.cpp b/Engine/Source/Editor/PropertyEditor/Private/PropertyEditorModule.cpp
--- a/Engine/Source/Editor/PropertyEditor/Private/PropertyEditorModule.cpp
+++ b/Engine/Source/Editor/PropertyEditor/Private/PropertyEditorModule.cpp
@@ -225,7 +225,7 @@ TSharedRef<SPropertyTreeViewImpl> FPropertyEditorModule::CreatePropertyView(
 
 	if( InObject )
 	{
-		TArray< TWeakObjectPtr< UObject > > Objects;
+		TArray<UObject*> Objects;
 		Objects.Add( InObject );
 		PropertyView->SetObjectArray( Objects );
 	}
@@ -256,9 +256,7 @@ TSharedRef<IDetailsView> FPropertyEditorModule::CreateDetailView( const FDetails
 		}
 	}
 
-	TSharedRef<SDetailsView> DetailView = 
-		SNew( SDetailsView )
-		.DetailsViewArgs( DetailsViewArgs );
+	TSharedRef<SDetailsView> DetailView = SNew(SDetailsView, DetailsViewArgs);
 
 	AllDetailViews.Add( DetailView );
 
diff --git a/Engine/Source/Editor/PropertyEditor/Private/PropertyEditorToolkit.cpp b/Engine/Source/Editor/PropertyEditor/Private/PropertyEditorToolkit.cpp
--- a/Engine/Source/Editor/PropertyEditor/Private/PropertyEditorToolkit.cpp
+++ b/Engine/Source/Editor/PropertyEditor/Private/PropertyEditorToolkit.cpp
@@ -164,7 +164,7 @@ void FPropertyEditorToolkit::Initialize( const EToolkitMode::Type Mode, const TS
 		const bool bCreateDefaultToolbar = false;
 		FAssetEditorToolkit::InitAssetEditor(Mode, InitToolkitHost, ApplicationId, StandaloneDefaultLayout, bCreateDefaultStandaloneMenu, bCreateDefaultToolbar, AdjustedObjectsToEdit);
 
-		TArray< TWeakObjectPtr<UObject> > AdjustedObjectsToEditWeak;
+		TArray<UObject*> AdjustedObjectsToEditWeak;
 		for (auto ObjectIter = AdjustedObjectsToEdit.CreateConstIterator(); ObjectIter; ++ObjectIter)
 		{
 			AdjustedObjectsToEditWeak.Add(*ObjectIter);
@@ -460,7 +460,7 @@ void FPropertyEditorToolkit::TableColumnsChanged()
 
 void FPropertyEditorToolkit::GridSelectionChanged()
 {
-	TArray< TWeakObjectPtr< UObject > > SelectedObjects;
+	TArray<TWeakObjectPtr<UObject>> SelectedObjects;
 	PropertyTable->GetSelectedTableObjects( SelectedObjects );
 
 	if (SelectedObjects.Num() == 0)
@@ -469,7 +469,15 @@ void FPropertyEditorToolkit::GridSelectionChanged()
 		SelectedObjects = PropertyTable->GetSelectedObjects();
 	}
 
-	PropertyTree->SetObjectArray( SelectedObjects );
+	TArray<UObject*> SelectedRawObjects;
+	SelectedRawObjects.Reserve(SelectedObjects.Num());
+
+	for (const TWeakObjectPtr<UObject>& Object : SelectedObjects)
+	{
+		SelectedRawObjects.Add(Object.Get());
+	}
+
+	PropertyTree->SetObjectArray(SelectedRawObjects);
 
 	const TSet< TSharedRef< IPropertyTableRow > > SelectedRows = PropertyTable->GetSelectedRows();
 
diff --git a/Engine/Source/Editor/PropertyEditor/Private/SDetailsView.cpp b/Engine/Source/Editor/PropertyEditor/Private/SDetailsView.cpp
--- a/Engine/Source/Editor/PropertyEditor/Private/SDetailsView.cpp
+++ b/Engine/Source/Editor/PropertyEditor/Private/SDetailsView.cpp
@@ -18,6 +18,7 @@
 #include "Classes/EditorStyleSettings.h"
 #include "DetailLayoutBuilderImpl.h"
 #include "DetailsViewPropertyGenerationUtilities.h"
+#include "DetailsViewGenericObjectFilter.h"
 
 
 #define LOCTEXT_NAMESPACE "SDetailsView"
@@ -36,9 +37,12 @@ SDetailsView::~SDetailsView()
  *
  * @param InArgs   Declaration from which to construct this widget.              
  */
-void SDetailsView::Construct(const FArguments& InArgs)
+void SDetailsView::Construct(const FArguments& InArgs, const FDetailsViewArgs& InDetailsViewArgs)
 {
-	DetailsViewArgs = InArgs._DetailsViewArgs;
+	DetailsViewArgs = InDetailsViewArgs;
+
+	SetObjectFilter(InDetailsViewArgs.ObjectFilter);
+
 	bViewingClassDefaultObject = false;
 
 	PropertyUtilities = MakeShareable( new FPropertyDetailsUtilities( *this ) );
@@ -368,7 +372,7 @@ EVisibility SDetailsView::GetScrollBarVisibility() const
 
 void SDetailsView::ForceRefresh()
 {
-	TArray< TWeakObjectPtr< UObject > > NewObjectList;
+	TArray<UObject*> NewObjectList;
 
 	const FRootPropertyNodeList& RootNodes = GetRootNodes();
 	for(const TSharedPtr<FComplexPropertyNode>& ComplexRootNode : RootNodes)
@@ -396,46 +400,52 @@ void SDetailsView::MoveScrollOffset(int32 DeltaOffset)
 	DetailTree->AddScrollOffset((float)DeltaOffset);
 }
 
-void SDetailsView::SetObjects( const TArray<UObject*>& InObjects, bool bForceRefresh/* = false*/, bool bOverrideLock/* = false*/ )
+void SDetailsView::SetObjects(const TArray<UObject*>& InObjects, bool bForceRefresh/* = false*/, bool bOverrideLock/* = false*/)
 {
 	if (!IsLocked() || bOverrideLock)
 	{
-		TArray< TWeakObjectPtr< UObject > > ObjectWeakPtrs;
-		
-		for( auto ObjectIter = InObjects.CreateConstIterator(); ObjectIter; ++ObjectIter )
-		{
-			ObjectWeakPtrs.Add( *ObjectIter );
-		}
-
-		if( bForceRefresh || ShouldSetNewObjects( ObjectWeakPtrs ) )
+		if( bForceRefresh || ShouldSetNewObjects(InObjects) )
 		{
-			SetObjectArrayPrivate( ObjectWeakPtrs );
+			SetObjectArrayPrivate(InObjects);
 		}
 	}
 }
 
-void SDetailsView::SetObjects( const TArray< TWeakObjectPtr< UObject > >& InObjects, bool bForceRefresh/* = false*/, bool bOverrideLock/* = false*/ )
+void SDetailsView::SetObjects(const TArray<TWeakObjectPtr<UObject>>& InObjects, bool bForceRefresh/* = false*/, bool bOverrideLock/* = false*/)
 {
 	if (!IsLocked() || bOverrideLock)
 	{
-		if( bForceRefresh || ShouldSetNewObjects( InObjects ) )
+		TArray<UObject*> SourceObjects;
+		SourceObjects.Reserve(InObjects.Num());
+
+		for (const TWeakObjectPtr<UObject>& Object : InObjects)
+		{
+			if (Object.IsValid())
+			{
+				SourceObjects.Add(Object.Get());
+			}
+		}
+
+		if(bForceRefresh || ShouldSetNewObjects(SourceObjects))
 		{
-			SetObjectArrayPrivate( InObjects );
+			SetObjectArrayPrivate(SourceObjects);
 		}
 	}
 }
 
-void SDetailsView::SetObject( UObject* InObject, bool bForceRefresh )
+void SDetailsView::SetObject(UObject* InObject, bool bForceRefresh)
 {
-	TArray< TWeakObjectPtr< UObject > > ObjectWeakPtrs;
-	ObjectWeakPtrs.Add( InObject );
+	TArray<UObject*> SourceObjects;
+	SourceObjects.Reserve(1);
 
-	SetObjects( ObjectWeakPtrs, bForceRefresh );
+	SourceObjects.Add(InObject);
+
+	SetObjects(SourceObjects, bForceRefresh);
 }
 
 void SDetailsView::RemoveInvalidObjects()
 {
-	TArray< TWeakObjectPtr< UObject > > ResetArray;
+	TArray<UObject*> ResetArray;
 
 	bool bAllFound = true;
 	
@@ -450,7 +460,7 @@ void SDetailsView::RemoveInvalidObjects()
 
 				if(Object.IsValid() && !Object->IsPendingKill())
 				{
-					ResetArray.Add(Object);
+					ResetArray.Add(Object.Get());
 				}
 				else
 				{
@@ -491,7 +501,17 @@ void SDetailsView::ClearSearch()
 	RerunCurrentFilter();
 }
 
-bool SDetailsView::ShouldSetNewObjects(const TArray< TWeakObjectPtr< UObject > >& InObjects) const
+void SDetailsView::SetObjectFilter(TSharedPtr<FDetailsViewObjectFilter> InFilter)
+{
+	ObjectFilter = InFilter;
+
+	if (!ObjectFilter.IsValid())
+	{
+		ObjectFilter = MakeShared<FDetailsViewDefaultObjectFilter>(!!DetailsViewArgs.bAllowMultipleTopLevelObjects);
+	}
+}
+
+bool SDetailsView::ShouldSetNewObjects(const TArray<UObject*>& InObjects) const
 {
 	bool bShouldSetObjects = false;
 
@@ -514,12 +534,11 @@ bool SDetailsView::ShouldSetNewObjects(const TArray< TWeakObjectPtr< UObject > >
 	else
 	{
 		// Check to see if the objects passed in are different. If not we do not need to set anything
-		TSet< TWeakObjectPtr< UObject > > NewObjects;
+		TSet<UObject*> NewObjects;
 		NewObjects.Append(InObjects);
 
-		if(DetailsViewArgs.bAllowMultipleTopLevelObjects)
+		if(RootPropertyNodes.Num() > 1)
 		{
-			
 			// For multiple top level node support, if the single object in each node is not found in the new object set
 			// then we need to refresh
 			for(int32 RootNodeIndex = 0; RootNodeIndex < RootPropertyNodes.Num(); ++RootNodeIndex)
@@ -543,15 +562,13 @@ bool SDetailsView::ShouldSetNewObjects(const TArray< TWeakObjectPtr< UObject > >
 		}
 		else
 		{
-
-			ensure(RootPropertyNodes.Num() == 1);
 			FObjectPropertyNode* RootPropertyNode = RootPropertyNodes[0]->AsObjectNode();
 			if( RootPropertyNode )
 			{
 				for(TPropObjectIterator Itor(RootPropertyNode->ObjectIterator()); Itor; ++Itor)
 				{
 					TWeakObjectPtr<UObject> Object = *Itor;
-					if(Object.IsValid() && !NewObjects.Contains(Object))
+					if(Object.IsValid() && !NewObjects.Contains(Object.Get()))
 					{
 						// An existing object is not in the list of new objects to set
 						bShouldSetObjects = true;
@@ -582,7 +599,7 @@ bool SDetailsView::ShouldSetNewObjects(const TArray< TWeakObjectPtr< UObject > >
 
 int32 SDetailsView::GetNumObjects() const
 {
-	if(DetailsViewArgs.bAllowMultipleTopLevelObjects)
+	if(RootPropertyNodes.Num() > 1)
 	{
 		return RootPropertyNodes.Num();
 	}
@@ -594,79 +611,77 @@ int32 SDetailsView::GetNumObjects() const
 	return 0;
 }
 
-void SDetailsView::SetObjectArrayPrivate(const TArray< TWeakObjectPtr< UObject > >& InObjects)
+void SDetailsView::SetObjectArrayPrivate(const TArray<UObject*>& InObjects)
 {
 	double StartTime = FPlatformTime::Seconds();
 
-	PreSetObject(InObjects.Num());
+	const TArray<FDetailsViewObjectRoot> Roots = ObjectFilter->FilterObjects(InObjects);
+
+	PreSetObject(Roots.Num());
 
 	// Selected actors for building SelectedActorInfo
 	TArray<AActor*> SelectedRawActors;
 
 	bViewingClassDefaultObject = InObjects.Num() > 0 ? true : false;
 	bool bOwnedByLockedLevel = false;
-	for( int32 ObjectIndex = 0 ; ObjectIndex < InObjects.Num(); ++ObjectIndex )
+
+	check(RootPropertyNodes.Num() == Roots.Num());
+
+	for(int32 RootIndex = 0; RootIndex < Roots.Num(); ++RootIndex)
 	{
-		TWeakObjectPtr< UObject > Object = InObjects[ObjectIndex];
+		const FDetailsViewObjectRoot& Root = Roots[RootIndex];
 
-		if( Object.IsValid() )
-		{
-			bViewingClassDefaultObject &= Object->HasAnyFlags( RF_ClassDefaultObject );
+		FObjectPropertyNode* RootNode = RootPropertyNodes[RootIndex]->AsObjectNode();
 
-			if(DetailsViewArgs.bAllowMultipleTopLevelObjects)
-			{
-				check(RootPropertyNodes.Num() == InObjects.Num());
-				RootPropertyNodes[ObjectIndex]->AsObjectNode()->AddObject( Object.Get() );
-			}
-			else
+		for(const TWeakObjectPtr<UObject>& Object : Root.Objects)
+		{
+			if (Object.IsValid())
 			{
-				RootPropertyNodes[0]->AsObjectNode()->AddObject( Object.Get() );
-			}
+				bViewingClassDefaultObject &= Object->HasAnyFlags(RF_ClassDefaultObject);
 
-			SelectedObjects.Add( Object );
-			AActor* Actor = Cast<AActor>( Object.Get() );
-			if( Actor )
-			{
-				SelectedActors.Add( Actor );
-				SelectedRawActors.Add( Actor );
+				RootNode->AddObject(Object.Get());
+
+				SelectedObjects.Add(Object);
+				if (AActor* Actor = Cast<AActor>(Object.Get()))
+				{
+					SelectedActors.Add(Actor);
+					SelectedRawActors.Add(Actor);
+				}
 			}
 		}
 	}
 
-	if( InObjects.Num() == 0 )
+	if(SelectedObjects.Num() == 0)
 	{
 		// Unlock the view automatically if we are viewing nothing
 		bIsLocked = false;
 	}
 
 	// Selection changed, refresh the detail area
-	if ( DetailsViewArgs.NameAreaSettings != FDetailsViewArgs::ActorsUseNameArea && DetailsViewArgs.NameAreaSettings != FDetailsViewArgs::ComponentsAndActorsUseNameArea )
+	if (DetailsViewArgs.NameAreaSettings != FDetailsViewArgs::ActorsUseNameArea && DetailsViewArgs.NameAreaSettings != FDetailsViewArgs::ComponentsAndActorsUseNameArea)
 	{
-		NameArea->Refresh( SelectedObjects );
+		NameArea->Refresh(SelectedObjects);
 	}
 	else
 	{
-		NameArea->Refresh( SelectedActors, SelectedObjects, DetailsViewArgs.NameAreaSettings );
+		NameArea->Refresh(SelectedActors, SelectedObjects, DetailsViewArgs.NameAreaSettings);
 	}
 	
 	// When selection changes rebuild information about the selection
-	SelectedActorInfo = AssetSelectionUtils::BuildSelectedActorInfo( SelectedRawActors );
+	SelectedActorInfo = AssetSelectionUtils::BuildSelectedActorInfo(SelectedRawActors);
 
-	// @todo Slate Property Window
-	//SetFlags(EPropertyWindowFlags::ReadOnly, bOwnedByLockedLevel);
 
-
-	PostSetObject();
+	PostSetObject(Roots);
 
 	// Set the title of the window based on the objects we are viewing
 	// Or call the delegate for handling when the title changed
 	FString Title;
 
-	if( GetNumObjects() == 0 )
+	if(SelectedObjects.Num() == 0 )
 	{
 		Title = NSLOCTEXT("PropertyView", "NothingSelectedTitle", "Nothing selected").ToString();
 	}
-	else if( GetNumObjects() == 1 && RootPropertyNodes[0]->AsObjectNode()->GetNumObjects() > 0)
+	else if(Roots.Num() == 1 && RootPropertyNodes[0]->AsObjectNode()->GetNumObjects() == 1)
 	{
 		// if the object is the default metaobject for a UClass, use the UClass's name instead
 		UObject* Object = RootPropertyNodes[0]->AsObjectNode()->GetUObject(0);
@@ -690,9 +705,9 @@ void SDetailsView::SetObjectArrayPrivate(const TArray< TWeakObjectPtr< UObject >
 
 		Title = ObjectName;
 	}
-	else if(DetailsViewArgs.bAllowMultipleTopLevelObjects)
+	else if(Roots.Num() > 1)
 	{
-		Title = FText::Format(NSLOCTEXT("PropertyView", "MultipleToLevelObjectsSelectedFmt", "{0} selected"), GetNumObjects()).ToString();
+		Title = FText::Format(NSLOCTEXT("PropertyView", "MultipleToLevelObjectsSelectedFmt", "{0} selected"), Roots.Num()).ToString();
 	}
 	else
 	{
@@ -707,7 +722,7 @@ void SDetailsView::SetObjectArrayPrivate(const TArray< TWeakObjectPtr< UObject >
 
 void SDetailsView::ReplaceObjects( const TMap<UObject*, UObject*>& OldToNewObjectMap )
 {
-	TArray<TWeakObjectPtr<UObject>> NewObjectList;
+	TArray<UObject*> NewObjectList;
 	bool bObjectsReplaced = false;
 
 	TArray<FObjectPropertyNode*> ObjectNodes;
@@ -749,18 +764,18 @@ void SDetailsView::ReplaceObjects( const TMap<UObject*, UObject*>& OldToNewObjec
 	}
 }
 
-void SDetailsView::RemoveDeletedObjects( const TArray<UObject*>& DeletedObjects )
+void SDetailsView::RemoveDeletedObjects(const TArray<UObject*>& DeletedObjects)
 {
-	TArray< TWeakObjectPtr< UObject > > NewObjectList;
+	TArray<UObject*> NewObjectList;
 	bool bObjectsRemoved = false;
 
 	for(const TSharedPtr<FComplexPropertyNode>& ComplexRootNode : RootPropertyNodes)
 	{
 		FObjectPropertyNode* RootPropertyNode = ComplexRootNode->AsObjectNode();
 		// Scan all objects and look for objects which need to be replaced
-		for ( TPropObjectIterator Itor( RootPropertyNode->ObjectIterator() ); Itor; ++Itor )
+		for (TPropObjectIterator Itor( RootPropertyNode->ObjectIterator() ); Itor; ++Itor)
 		{
-			if( DeletedObjects.Contains( Itor->Get() ) )
+			if(DeletedObjects.Contains( Itor->Get()))
 			{
 				// An object we had needs to be removed
 				bObjectsRemoved = true;
@@ -776,7 +791,7 @@ void SDetailsView::RemoveDeletedObjects( const TArray<UObject*>& DeletedObjects
 	// if any objects were replaced update the observed objects
 	if( bObjectsRemoved )
 	{
-		SetObjectArrayPrivate( NewObjectList );
+		SetObjectArrayPrivate(NewObjectList);
 	}
 }
 
@@ -816,18 +831,11 @@ void SDetailsView::PreSetObject(int32 InNewNumObjects)
 	RootPropertyNodes.Empty(InNewNumObjects);
 	ExpandedDetailNodes.Empty();
 
-	if(DetailsViewArgs.bAllowMultipleTopLevelObjects)
-	{
-		for(int32 NewRootIndex = 0; NewRootIndex < InNewNumObjects; ++NewRootIndex)
-		{
-			RootPropertyNodes.Add(MakeShareable(new FObjectPropertyNode));
-		}
-	}
-	else
+
+	for (int32 NewRootIndex = 0; NewRootIndex < InNewNumObjects; ++NewRootIndex)
 	{
 		RootPropertyNodes.Add(MakeShareable(new FObjectPropertyNode));
-	}	
-
+	}
 
 	SelectedActors.Empty();
 	SelectedObjects.Empty();
@@ -835,7 +843,7 @@ void SDetailsView::PreSetObject(int32 InNewNumObjects)
 
 
 /** Called at the end of SetObjectArray after we change the objects being observed */
-void SDetailsView::PostSetObject()
+void SDetailsView::PostSetObject(const TArray<FDetailsViewObjectRoot>& Roots)
 {
 	TSharedPtr<SColorPicker> ExistingColorPicker = GetColorPicker();
 	if (ExistingColorPicker.IsValid()
diff --git a/Engine/Source/Editor/PropertyEditor/Private/SDetailsView.h b/Engine/Source/Editor/PropertyEditor/Private/SDetailsView.h
--- a/Engine/Source/Editor/PropertyEditor/Private/SDetailsView.h
+++ b/Engine/Source/Editor/PropertyEditor/Private/SDetailsView.h
@@ -14,17 +14,14 @@
 
 class AActor;
 class IDetailRootObjectCustomization;
+class FDetailsViewObjectFilter;
 
 class SDetailsView : public SDetailsViewBase
 {
 	friend class FPropertyDetailsUtilities;
 public:
 
-	SLATE_BEGIN_ARGS( SDetailsView )
-		: _DetailsViewArgs()
-		{}
-		/** The user defined args for the details view */
-		SLATE_ARGUMENT( FDetailsViewArgs, DetailsViewArgs )
+	SLATE_BEGIN_ARGS(SDetailsView){}
 	SLATE_END_ARGS()
 
 	virtual ~SDetailsView();
@@ -39,42 +36,44 @@ public:
 	/**
 	 * Constructs the property view widgets                   
 	 */
-	void Construct(const FArguments& InArgs);
+	void Construct(const FArguments& InArgs, const FDetailsViewArgs& InDetailsViewArgs);
 
 	/** IDetailsView interface */
-	virtual void SetObjects( const TArray<UObject*>& InObjects, bool bForceRefresh = false, bool bOverrideLock = false ) override;
-	virtual void SetObjects( const TArray< TWeakObjectPtr< UObject > >& InObjects, bool bForceRefresh = false, bool bOverrideLock = false ) override;
-	virtual void SetObject( UObject* InObject, bool bForceRefresh = false ) override;
+	virtual void SetObjects(const TArray<UObject*>& InObjects, bool bForceRefresh = false, bool bOverrideLock = false) override;
+	virtual void SetObjects(const TArray<TWeakObjectPtr<UObject>>& InObjects, bool bForceRefresh = false, bool bOverrideLock = false) override;
+	virtual void SetObject(UObject* InObject, bool bForceRefresh = false) override;
+
 	virtual void RemoveInvalidObjects() override;
 	virtual void SetObjectPackageOverrides(const TMap<TWeakObjectPtr<UObject>, TWeakObjectPtr<UPackage>>& InMapping) override;
 	virtual void SetRootObjectCustomizationInstance(TSharedPtr<IDetailRootObjectCustomization> InRootObjectCustomization) override;
 	virtual void ClearSearch() override;
+	virtual void SetObjectFilter(TSharedPtr<FDetailsViewObjectFilter> InFilter) override;
 
 	/**
 	 * Replaces objects being observed by the view with new objects
 	 *
 	 * @param OldToNewObjectMap	Mapping from objects to replace to their replacement
 	 */
-	void ReplaceObjects( const TMap<UObject*, UObject*>& OldToNewObjectMap );
+	void ReplaceObjects(const TMap<UObject*, UObject*>& OldToNewObjectMap);
 
 	/**
 	 * Removes objects from the view because they are about to be deleted
 	 *
 	 * @param DeletedObjects	The objects to delete
 	 */
-	void RemoveDeletedObjects( const TArray<UObject*>& DeletedObjects );
+	void RemoveDeletedObjects(const TArray<UObject*>& DeletedObjects);
 
 	/** Sets the callback for when the property view changes */
-	virtual void SetOnObjectArrayChanged( FOnObjectArrayChanged OnObjectArrayChangedDelegate) override;
+	virtual void SetOnObjectArrayChanged(FOnObjectArrayChanged OnObjectArrayChangedDelegate) override;
 
 	/** @return	Returns list of selected objects we're inspecting */
-	virtual const TArray< TWeakObjectPtr<UObject> >& GetSelectedObjects() const override
+	virtual const TArray<TWeakObjectPtr<UObject>>& GetSelectedObjects() const override
 	{
 		return SelectedObjects;
 	} 
 
 	/** @return	Returns list of selected actors we're inspecting */
-	virtual const TArray< TWeakObjectPtr<AActor> >& GetSelectedActors() const override
+	virtual const TArray<TWeakObjectPtr<AActor>>& GetSelectedActors() const override
 	{
 		return SelectedActors;
 	}
@@ -112,7 +111,7 @@ public:
 		return RootObjectCustomization;
 	}
 private:
-	void SetObjectArrayPrivate( const TArray< TWeakObjectPtr< UObject > >& InObjects );
+	void SetObjectArrayPrivate(const TArray<UObject*>& InObjects);
 
 	TSharedRef<SDetailTree> ConstructTreeView( TSharedRef<SScrollBar>& ScrollBar );
 
@@ -123,7 +122,7 @@ private:
 	 * @param InObjects The potential new objects to set
 	 * @return true if the new objects should be set
 	 */
-	bool ShouldSetNewObjects( const TArray< TWeakObjectPtr< UObject > >& InObjects ) const;
+	bool ShouldSetNewObjects(const TArray<UObject*>& InObjects) const;
 
 	/**
 	 * Returns the number of objects being edited by this details panel.
@@ -134,7 +133,7 @@ private:
 	void PreSetObject(int32 InNewNumObjects);
 
 	/** Called at the end of SetObjectArray after we change the objects being observed */
-	void PostSetObject();
+	void PostSetObject(const TArray<FDetailsViewObjectRoot>& Roots);
 	
 	/** Called to get the visibility of the actor name area */
 	EVisibility GetActorNameAreaVisibility() const;
@@ -160,10 +159,14 @@ private:
 	void OnShowHiddenPropertiesWhilePlayingClicked();
 
 private:
+	/** The filter for objects viewed by this details panel */
+	TSharedPtr<FDetailsViewObjectFilter> ObjectFilter;
+
 	/** Information about the current set of selected actors */
 	FSelectedActorInfo SelectedActorInfo;
-	/** Selected objects for this detail view.  */
-	TArray< TWeakObjectPtr<UObject> > SelectedObjects;
+
+	/** Final set of selected objects for this detail view.  It may be different from the set passed in through SetObjects if there is an active filter */
+	TArray<TWeakObjectPtr<UObject>> SelectedObjects;
 
 	/** 
 	 * Selected actors for this detail view.  Note that this is not necessarily the same editor selected actor set.  If this detail view is locked
diff --git a/Engine/Source/Editor/PropertyEditor/Private/SDetailsViewBase.h b/Engine/Source/Editor/PropertyEditor/Private/SDetailsViewBase.h
--- a/Engine/Source/Editor/PropertyEditor/Private/SDetailsViewBase.h
+++ b/Engine/Source/Editor/PropertyEditor/Private/SDetailsViewBase.h
@@ -30,7 +30,7 @@ class IDetailKeyframeHandler;
 class IDetailPropertyExtensionHandler;
 class SDetailNameArea;
 class IPropertyGenerationUtilities;
-
+struct FDetailsViewObjectRoot;
 
 
 typedef STreeView< TSharedRef<class FDetailTreeNode> > SDetailTree;
diff --git a/Engine/Source/Editor/PropertyEditor/Private/SPropertyTreeViewImpl.cpp b/Engine/Source/Editor/PropertyEditor/Private/SPropertyTreeViewImpl.cpp
--- a/Engine/Source/Editor/PropertyEditor/Private/SPropertyTreeViewImpl.cpp
+++ b/Engine/Source/Editor/PropertyEditor/Private/SPropertyTreeViewImpl.cpp
@@ -760,16 +760,19 @@ void SPropertyTreeViewImpl::RequestRefresh()
 	FavoritesTree->RequestTreeRefresh();
 }
 
-void SPropertyTreeViewImpl::SetObjectArray( const TArray< TWeakObjectPtr< UObject > >& InObjects )
+void SPropertyTreeViewImpl::SetObjectArray( const TArray<UObject*>& InObjects )
 {
 	check( RootPropertyNode.IsValid() );
 
 	PreSetObject();
 
+	TArray<UObject*> SetObjects;
+	SetObjects.Reserve(InObjects.Num());
+
 	bool bOwnedByLockedLevel = false;
-	for( int32 ObjectIndex = 0 ; ObjectIndex < InObjects.Num() ; ++ObjectIndex )
+	for(UObject* Object : InObjects)
 	{
-		RootPropertyNode->AddObject( InObjects[ObjectIndex].Get() );
+		RootPropertyNode->AddObject(Object);
 	}
 
 	// @todo Slate Property Window
@@ -809,7 +812,7 @@ void SPropertyTreeViewImpl::SetObjectArray( const TArray< TWeakObjectPtr< UObjec
 		Title = FText::Format( NSLOCTEXT("PropertyView", "MultipleSelectedFmt", "{0} ({1} selected)"), FText::FromString(RootPropertyNode->GetObjectBaseClass()->GetName()), RootPropertyNode->GetNumObjects() ).ToString();
 	}
 
-	OnObjectArrayChanged.ExecuteIfBound(Title, InObjects);
+	OnObjectArrayChanged.ExecuteIfBound(Title, SetObjects);
 }
 
 TSharedRef< FPropertyPath > SPropertyTreeViewImpl::GetRootPath() const
@@ -841,7 +844,7 @@ void SPropertyTreeViewImpl::SetRootPath( const TSharedPtr< FPropertyPath >& Path
 void SPropertyTreeViewImpl::ReplaceObjects( const TMap<UObject*, UObject*>& OldToNewObjectMap )
 {
 
-	TArray< TWeakObjectPtr< UObject > > NewObjectList;
+	TArray<UObject*> NewObjectList;
 	bool bObjectsReplaced = false;
 
 	TArray< FObjectPropertyNode* > ObjectNodes;
@@ -880,7 +883,7 @@ void SPropertyTreeViewImpl::ReplaceObjects( const TMap<UObject*, UObject*>& OldT
 
 void SPropertyTreeViewImpl::RemoveDeletedObjects( const TArray<UObject*>& DeletedObjects )
 {
-	TArray< TWeakObjectPtr< UObject > > NewObjectList;
+	TArray<UObject*> NewObjectList;
 	bool bObjectsRemoved = false;
 
 	// Scan all objects and look for objects which need to be replaced
@@ -912,7 +915,7 @@ void SPropertyTreeViewImpl::RemoveDeletedObjects( const TArray<UObject*>& Delete
  */
 void SPropertyTreeViewImpl::RemoveInvalidActors( const TSet<AActor*>& ValidActors )
 {
-	TArray< TWeakObjectPtr< UObject > > ResetArray;
+	TArray<UObject*> ResetArray;
 
 	bool bAllFound = true;
 	for ( TPropObjectIterator Itor( RootPropertyNode->ObjectIterator() ); Itor; ++Itor )
@@ -1039,7 +1042,7 @@ void SPropertyTreeViewImpl::Tick( const FGeometry& AllottedGeometry, const doubl
 	}
 	else if( Result == EPropertyDataValidationResult::ObjectInvalid && !bNodeTreeExternallyManaged )
 	{
-		TArray< TWeakObjectPtr< UObject > > ResetArray;
+		TArray<UObject*> ResetArray;
 		for ( TPropObjectIterator Itor( RootPropertyNode->ObjectIterator() ) ; Itor ; ++Itor )
 		{
 			TWeakObjectPtr<UObject> Object = *Itor;
@@ -1258,7 +1261,7 @@ void SPropertyTreeViewImpl::SetIsPropertyVisible(FIsPropertyVisible IsPropertyVi
 
 	if( RootPropertyNode.IsValid() )
 	{
-		TArray< TWeakObjectPtr< UObject > > Objects;
+		TArray<UObject*> Objects;
 		for( int32 ObjIndex = 0; ObjIndex < RootPropertyNode->GetNumObjects(); ++ObjIndex )
 		{
 			Objects.Add( RootPropertyNode->GetUObject( ObjIndex ) );
diff --git a/Engine/Source/Editor/PropertyEditor/Private/SPropertyTreeViewImpl.h b/Engine/Source/Editor/PropertyEditor/Private/SPropertyTreeViewImpl.h
--- a/Engine/Source/Editor/PropertyEditor/Private/SPropertyTreeViewImpl.h
+++ b/Engine/Source/Editor/PropertyEditor/Private/SPropertyTreeViewImpl.h
@@ -66,7 +66,7 @@ public:
 	 *
 	 * @param InObjects	The list of objects to observe
 	 */
-	virtual void SetObjectArray( const TArray< TWeakObjectPtr< UObject > >& InObjects );
+	virtual void SetObjectArray( const TArray<UObject*>& InObjects );
 
 	virtual TSharedRef< FPropertyPath > GetRootPath() const;
 
diff --git a/Engine/Source/Editor/PropertyEditor/Private/SStructureDetailsView.h b/Engine/Source/Editor/PropertyEditor/Private/SStructureDetailsView.h
--- a/Engine/Source/Editor/PropertyEditor/Private/SStructureDetailsView.h
+++ b/Engine/Source/Editor/PropertyEditor/Private/SStructureDetailsView.h
@@ -94,6 +94,7 @@ public:
 	virtual void SetObjectPackageOverrides(const TMap<TWeakObjectPtr<UObject>, TWeakObjectPtr<UPackage>>& InMapping) override {}
 	virtual void SetRootObjectCustomizationInstance(TSharedPtr<IDetailRootObjectCustomization> InRootObjectCustomization) override {}
 	virtual TSharedPtr<class IDetailRootObjectCustomization> GetRootObjectCustomization() const override { return nullptr; }
+	virtual void SetObjectFilter(TSharedPtr<FDetailsViewObjectFilter> InFilter) override {}
 
 	/* This is required by the base class but there is only ever one root node in a structure details view */
 	virtual FRootPropertyNodeList& GetRootNodes() override;
diff --git a/Engine/Source/Editor/PropertyEditor/Public/DetailsViewObjectFilter.h b/Engine/Source/Editor/PropertyEditor/Public/DetailsViewObjectFilter.h
new file mode 100644
--- /dev/null
+++ b/Engine/Source/Editor/PropertyEditor/Public/DetailsViewObjectFilter.h
@@ -0,0 +1,49 @@
+// Copyright Epic Games, Inc. All Rights Reserved.
+
+#pragma once
+
+#include "CoreMinimal.h"
+
+/**
+ * An object root is a collection of UObjects that represent a top level set of properties in a details panel
+ * When there are multiple objects in the root, the common base class for all those objects are found and the properties on that common base class are displayed
+ * When a user edits one of those properties, it will propagate the result to all objects in the root.  
+ * If multiple differing values are found on a single property, the details panel UI displays "multiple values" 
+ */
+struct FDetailsViewObjectRoot
+{
+	FDetailsViewObjectRoot(bool bInShowRootHeader = true)
+		: bShowRootHeader(bInShowRootHeader)
+	{}
+
+	FDetailsViewObjectRoot(UObject* InObject, bool bInShowRootHeader=true)
+		: bShowRootHeader(bInShowRootHeader)
+	{
+		Objects.Add(InObject);
+	}
+
+	FDetailsViewObjectRoot(const TArray<UObject*>& InObjects, bool bInShowRootHeader=true)
+		: bShowRootHeader(bInShowRootHeader)
+	{
+		Objects.Reserve(InObjects.Num());
+		for (UObject* Object : InObjects)
+		{
+			Objects.Add(Object);
+		}
+	}
+
+	TArray<UObject*> Objects;
+	bool bShowRootHeader;
+};
+
+/**
+ * An object filter determines the root objects that should be displayed from a set of given source objects passed to the details panel
+ */
+class FDetailsViewObjectFilter
+{
+public:
+	virtual ~FDetailsViewObjectFilter() {}
+
+	virtual TArray<FDetailsViewObjectRoot> FilterObjects(const TArray<UObject*>& SourceObjects) = 0;
+
+};
diff --git a/Engine/Source/Editor/PropertyEditor/Public/IDetailRootObjectCustomization.h b/Engine/Source/Editor/PropertyEditor/Public/IDetailRootObjectCustomization.h
--- a/Engine/Source/Editor/PropertyEditor/Public/IDetailRootObjectCustomization.h
+++ b/Engine/Source/Editor/PropertyEditor/Public/IDetailRootObjectCustomization.h
@@ -12,6 +12,13 @@ class ITableRow;
 class IDetailRootObjectCustomization : public TSharedFromThis<IDetailRootObjectCustomization>
 {
 public:
+	enum EExpansionArrowUsage
+	{
+		Custom,
+		Default,
+		None,
+	};
+
 	virtual ~IDetailRootObjectCustomization() {}
 
 	/** 
@@ -19,15 +26,15 @@ public:
 	 *
 	 * @param	InRootObject	The object whose header is being customized
  	 */
-	virtual TSharedPtr<SWidget> CustomizeObjectHeader(const UObject* InRootObject) = 0;
+	virtual TSharedPtr<SWidget> CustomizeObjectHeader(const UObject* InRootObject) { return nullptr; }
 
 	/** 
 	 * Called when the details panel wants to display an object header widget for a given object
 	 *
 	 * @param	InRootObject	The object whose header is being customized
-	 * @param	InTableRow		The ITableRow object (table views to talk to their rows) that will use the current IDetailRootObjectCustomization element
+	 * @param	InTableRow		The ITableRow object (table views to talk to their rows) that will use the current IDetailRootObjectCustomization element.  This may be null if the customization is not being shown in a table view
  	 */
-	virtual TSharedPtr<SWidget> CustomizeObjectHeader(const UObject* InRootObject, const TSharedRef<ITableRow>& InTableRow)
+	virtual TSharedPtr<SWidget> CustomizeObjectHeader(const UObject* InRootObject, const TSharedPtr<ITableRow>& InTableRow)
 	{
 		return CustomizeObjectHeader(InRootObject);
 	}
@@ -38,9 +45,18 @@ public:
 	virtual bool IsObjectVisible(const UObject* InRootObject) const = 0;
 
 	/**
-	 * Whether or object should have a header displayed or just show the children directly
+	 * Whether or not the object should have a header displayed or just show the children directly
 	 *
 	 * @return true if the this customization should be displayed or false to just show the children directly
 	 */
 	virtual bool ShouldDisplayHeader(const UObject* InRootObject) const = 0;
+
+	/**
+	 * Gets the setup for expansion arrows in this customization
+	 */
+	virtual EExpansionArrowUsage GetExpansionArrowUsage() const
+	{ 
+		// Note: Returns none for backwards compatibility
+		return EExpansionArrowUsage::None; 
+	}
 };
diff --git a/Engine/Source/Editor/PropertyEditor/Public/IDetailsView.h b/Engine/Source/Editor/PropertyEditor/Public/IDetailsView.h
--- a/Engine/Source/Editor/PropertyEditor/Public/IDetailsView.h
+++ b/Engine/Source/Editor/PropertyEditor/Public/IDetailsView.h
@@ -15,7 +15,7 @@ class IDetailKeyframeHandler;
 class IDetailPropertyExtensionHandler;
 class IDetailRootObjectCustomization;
 class IPropertyTypeIdentifier;
-
+class FDetailsViewObjectFilter;
 
 enum class EEditDefaultsOnlyNodeVisibility : uint8
 {
@@ -51,6 +51,8 @@ struct FDetailsViewArgs
 	TSharedPtr<class FUICommandList> HostCommandList;
 	/** The tab manager from the host of the details view, allowing child widgets to spawn tabs */
 	TSharedPtr<FTabManager> HostTabManager;
+	/** Optional object filter to use for more complex handling of what a details panel is viewing. */
+	TSharedPtr<FDetailsViewObjectFilter> ObjectFilter;
 
 	/** Identifier for this details view; NAME_None if this view is anonymous */
 	FName ViewIdentifier;
@@ -353,6 +355,9 @@ public:
 	/** Force refresh */
 	virtual void ForceRefresh() = 0;
 	
+	/** Sets an optional object filter to use for more complex handling of what a details panel is viewing. */
+	virtual void SetObjectFilter(TSharedPtr<FDetailsViewObjectFilter> InFilter) = 0;
+
 	/** Allows other systems to add a custom filter in the details panel */
 	virtual void SetCustomFilterDelegate(FSimpleDelegate InDelegate) = 0;
 
diff --git a/Engine/Source/Editor/PropertyEditor/Public/PropertyEditorDelegates.h b/Engine/Source/Editor/PropertyEditor/Public/PropertyEditorDelegates.h
--- a/Engine/Source/Editor/PropertyEditor/Public/PropertyEditorDelegates.h
+++ b/Engine/Source/Editor/PropertyEditor/Public/PropertyEditorDelegates.h
@@ -45,7 +45,7 @@ DECLARE_DELEGATE_RetVal( TSharedRef<class IDetailCustomization>, FOnGetDetailCus
 DECLARE_DELEGATE_RetVal( TSharedRef<class IPropertyTypeCustomization>, FOnGetPropertyTypeCustomizationInstance );
 
 /** Notification for when a property view changes */
-DECLARE_DELEGATE_TwoParams( FOnObjectArrayChanged, const FString&, const TArray< TWeakObjectPtr< UObject > >& );
+DECLARE_DELEGATE_TwoParams( FOnObjectArrayChanged, const FString&, const TArray<UObject*>& );
 
 /** Notification for when displayed properties changes (for instance, because the user has filtered some properties */
 DECLARE_DELEGATE( FOnDisplayedPropertiesChanged );
diff --git a/Engine/Source/Editor/UMGEditor/Private/Details/SWidgetDetailsView.cpp b/Engine/Source/Editor/UMGEditor/Private/Details/SWidgetDetailsView.cpp
--- a/Engine/Source/Editor/UMGEditor/Private/Details/SWidgetDetailsView.cpp
+++ b/Engine/Source/Editor/UMGEditor/Private/Details/SWidgetDetailsView.cpp
@@ -280,7 +280,7 @@ void SWidgetDetailsView::OnEditorSelectionChanged()
 	PropertyView->SetObjects(SelectedObjects, bForceRefresh);
 }
 
-void SWidgetDetailsView::OnPropertyViewObjectArrayChanged(const FString& InTitle, const TArray<TWeakObjectPtr<UObject>>& InObjects)
+void SWidgetDetailsView::OnPropertyViewObjectArrayChanged(const FString& InTitle, const TArray<UObject*>& InObjects)
 {
 	// Update the package remapping for all selected objects so that we can find the correct localization ID when editing text properties (since we edit a copy of the real data, not connected to the asset package)
 	UBlueprint* UMGBlueprint = BlueprintEditor.Pin()->GetBlueprintObj();
diff --git a/Engine/Source/Editor/UMGEditor/Private/Details/SWidgetDetailsView.h b/Engine/Source/Editor/UMGEditor/Private/Details/SWidgetDetailsView.h
--- a/Engine/Source/Editor/UMGEditor/Private/Details/SWidgetDetailsView.h
+++ b/Engine/Source/Editor/UMGEditor/Private/Details/SWidgetDetailsView.h
@@ -44,7 +44,7 @@ private:
 	void OnEditorSelectionChanged();
 
 	/** Handles the callback from the property detail view confirming the list of objects being edited has changed */
-	void OnPropertyViewObjectArrayChanged(const FString& InTitle, const TArray<TWeakObjectPtr<UObject>>& InObjects);
+	void OnPropertyViewObjectArrayChanged(const FString& InTitle, const TArray<UObject*>& InObjects);
 
 	void ClearFocusIfOwned();
 
