Fixed material editor crash when viewing the parameters overview

diff --git a/Engine/Source/Editor/PropertyEditor/Private/DetailLayoutBuilderImpl.cpp b/Engine/Source/Editor/PropertyEditor/Private/DetailLayoutBuilderImpl.cpp
--- a/Engine/Source/Editor/PropertyEditor/Private/DetailLayoutBuilderImpl.cpp
+++ b/Engine/Source/Editor/PropertyEditor/Private/DetailLayoutBuilderImpl.cpp
@@ -327,7 +327,7 @@ void FDetailLayoutBuilderImpl::GenerateDetailLayout()
 	}
 
 	TSharedPtr<FComplexPropertyNode> RootNodePinned = RootNode.Pin();
-	if(DetailsView->GetRootObjectCustomization() && RootNodePinned->GetInstancesNum())
+	if(DetailsView && DetailsView->GetRootObjectCustomization() && RootNodePinned->GetInstancesNum())
 	{
 		// This should always exist here
 		UObject* RootObject = RootNodePinned->AsObjectNode()->GetUObject(0);
