StructOnScope 'fake' properties will now be placed in a struct container object to make sure they are properly resolved by TFieldPaths

diff --git a/Engine/Source/Editor/PropertyEditor/Private/PropertyEditorModule.cpp b/Engine/Source/Editor/PropertyEditor/Private/PropertyEditorModule.cpp
--- a/Engine/Source/Editor/PropertyEditor/Private/PropertyEditorModule.cpp
+++ b/Engine/Source/Editor/PropertyEditor/Private/PropertyEditorModule.cpp
@@ -89,10 +89,14 @@ const FPropertyTypeLayoutCallback& FPropertyTypeLayoutCallbackList::Find( const
 
 void FPropertyEditorModule::StartupModule()
 {
+	StructOnScopePropertyOwner = nullptr;
 }
 
 void FPropertyEditorModule::ShutdownModule()
 {
+	// No need to remove this object from root since the final GC pass doesn't care about root flags
+	StructOnScopePropertyOwner = nullptr;
+
 	// NOTE: It's vital that we clean up everything created by this DLL here!  We need to make sure there
 	//       are no outstanding references to objects as the compiled code for this module's class will
 	//       literally be unloaded from memory after this function exits.  This even includes instantiated
@@ -375,8 +379,16 @@ FStructProperty* FPropertyEditorModule::RegisterStructOnScopeProperty(TSharedRef
 
 	if(!StructProperty)
 	{
+		if (!StructOnScopePropertyOwner)
+		{
+			// Create a container for all StructOnScope property objects.
+			// It's important that this container is the owner of these properties and maintains a linked list 
+			// to all of the properties created here. This is automatically handled by the specialized property constructor.
+			StructOnScopePropertyOwner = NewObject<UStruct>(GetTransientPackage(), TEXT("StructOnScope"), RF_Transient);
+			StructOnScopePropertyOwner->AddToRoot();
+		}
 		UScriptStruct* InnerStruct = Cast<UScriptStruct>(const_cast<UStruct*>(StructOnScope->GetStruct()));
-		StructProperty = new FStructProperty(InnerStruct, *MakeUniqueObjectName(InnerStruct, UField::StaticClass(), InnerStruct->GetFName()).ToString(), RF_NoFlags);
+		StructProperty = new FStructProperty(StructOnScopePropertyOwner, *MakeUniqueObjectName(StructOnScopePropertyOwner, UField::StaticClass(), InnerStruct->GetFName()).ToString(), RF_Transient, 0, CPF_None, InnerStruct);
 		StructProperty->Struct = InnerStruct;
 		StructProperty->ElementSize = StructOnScope->GetStruct()->GetStructureSize();
 
diff --git a/Engine/Source/Editor/PropertyEditor/Public/PropertyEditorModule.h b/Engine/Source/Editor/PropertyEditor/Public/PropertyEditorModule.h
--- a/Engine/Source/Editor/PropertyEditor/Public/PropertyEditorModule.h
+++ b/Engine/Source/Editor/PropertyEditor/Public/PropertyEditorModule.h
@@ -311,5 +311,7 @@ private:
 	TMap<FName, FStructProperty*> RegisteredStructToProxyMap;
 	/** Shared thumbnail pool used by property row generators */
 	TSharedPtr<class FAssetThumbnailPool> GlobalThumbnailPool;
+	/** Container for ScopeOnStruct FStructProperty objects */
+	UStruct* StructOnScopePropertyOwner;
 
 };
