Editor - Giving the details panel the ability to filter the values as well (before we only filtered based on the name of properties and any custom filter strings coming from the customization.  Now the values themselves are considered when filtering so that if you've referenced an asset or have some text property, you can search for specific text, or the asset name to see where it has been set.

diff --git a/Engine/Source/Editor/PropertyEditor/Private/DetailItemNode.cpp b/Engine/Source/Editor/PropertyEditor/Private/DetailItemNode.cpp
--- a/Engine/Source/Editor/PropertyEditor/Private/DetailItemNode.cpp
+++ b/Engine/Source/Editor/PropertyEditor/Private/DetailItemNode.cpp
@@ -476,6 +476,53 @@ static bool PassesAllFilters( FDetailItemNode* ItemNode, const FDetailLayoutCust
 			}
 			return false;
 		}
+
+		static FString GetPropertyNodeValueFilterString(const FDetailLayoutCustomization& InCustomization, TSharedPtr<FPropertyNode> PropertyNode)
+		{
+			if (PropertyNode.IsValid())
+			{
+				// Is it a container (array, map, set?) - if so, ignore it, we don't care about these, only their inner nodes.
+				if (CastField<FArrayProperty>(PropertyNode->GetProperty()) || CastField<FMapProperty>(PropertyNode->GetProperty()) || CastField<FSetProperty>(PropertyNode->GetProperty()))
+				{
+					return FString();
+				}
+
+				// Is it a struct?  If so, some structs are useful, like FGameplayTag, or FGameplayTags, but if it's a user struct for the game
+				// like FMyGameplayStruct, with a bunch of other sub nodes, that will individually be matched and filtered, there's no reason
+				// to filter on the struct as a whole, so essentially what we're doing here is only checking structs that are leaf nodes.
+				if (CastField<FStructProperty>(PropertyNode->GetProperty()))
+				{
+					if (PropertyNode->GetNumChildNodes() > 0)
+					{
+						return FString();
+					}
+				}
+
+				// TODO Will have to do something special for EditInlineNew UObjects, rather than just a simple object path.
+				if (FObjectProperty* ObjectProperty = CastField<FObjectProperty>(PropertyNode->GetProperty()))
+				{
+					uint8* ValueAddress = nullptr;
+					FPropertyAccess::Result Result = PropertyNode->GetSingleReadAddress(ValueAddress);
+					if (ValueAddress != nullptr)
+					{
+						if (UObject* ObjectValue = ObjectProperty->GetObjectPropertyValue(ValueAddress))
+						{
+							return ObjectValue->GetName();
+						}
+					}
+				}
+				else
+				{
+					// PPF_SimpleObjectText, seems to get the most reasonable string for searching.
+					FString OutString;
+					PropertyNode->GetPropertyValueString(OutString, true, PPF_SimpleObjectText);
+
+					return OutString;
+				}
+			}
+
+			return FString();
+		}
 	};
 
 	bool bPassesAllFilters = true;
@@ -488,6 +535,7 @@ static bool PassesAllFilters( FDetailItemNode* ItemNode, const FDetailLayoutCust
 		TSharedPtr<FPropertyNode> PropertyNodePin = InCustomization.GetPropertyNode();
 		
 		const bool bPassesCategoryFilter = !bSearchFilterIsEmpty && InFilter.bShowAllChildrenIfCategoryMatches ? Local::StringPassesFilter(InFilter, InCategoryName) : false;
+		const bool bPassesValueFilter = !bSearchFilterIsEmpty && Local::StringPassesFilter(InFilter, Local::GetPropertyNodeValueFilterString(InCustomization, PropertyNodePin));
 
 		bPassesAllFilters = false;
 		if( PropertyNodePin.IsValid() && !PropertyNodePin->AsCategoryNode())
@@ -496,7 +544,7 @@ static bool PassesAllFilters( FDetailItemNode* ItemNode, const FDetailLayoutCust
 			const bool bIsSeenDueToFiltering = PropertyNodePin->HasNodeFlags(EPropertyNodeFlags::IsSeenDueToFiltering) != 0;
 			const bool bIsParentSeenDueToFiltering = PropertyNodePin->HasNodeFlags(EPropertyNodeFlags::IsParentSeenDueToFiltering) != 0;
 
-			const bool bPassesSearchFilter = bPassesCategoryFilter || bSearchFilterIsEmpty || ( bIsNotBeingFiltered || bIsSeenDueToFiltering || bIsParentSeenDueToFiltering );
+			const bool bPassesSearchFilter = bPassesCategoryFilter || bPassesValueFilter || bSearchFilterIsEmpty || ( bIsNotBeingFiltered || bIsSeenDueToFiltering || bIsParentSeenDueToFiltering );
 			const bool bPassesModifiedFilter = bPassesSearchFilter && ( InFilter.bShowOnlyModifiedProperties == false || PropertyNodePin->GetDiffersFromDefault() == true );
 			const bool bPassesDifferingFilter = InFilter.bShowOnlyDiffering ? InFilter.WhitelistedProperties.Find(*FPropertyNode::CreatePropertyPath(PropertyNodePin.ToSharedRef())) != nullptr : true;
 
@@ -521,7 +569,7 @@ static bool PassesAllFilters( FDetailItemNode* ItemNode, const FDetailLayoutCust
 		}
 		else if (InCustomization.HasCustomWidget())
 		{
-			const bool bPassesTextFilter = bPassesCategoryFilter || Local::StringPassesFilter(InFilter, InCustomization.WidgetDecl->FilterTextString.ToString());
+			const bool bPassesTextFilter = bPassesCategoryFilter || bPassesValueFilter || Local::StringPassesFilter(InFilter, InCustomization.WidgetDecl->FilterTextString.ToString());
 			const bool bPassesModifiedFilter = (InFilter.bShowOnlyModifiedProperties == false || InCustomization.WidgetDecl->DiffersFromDefaultAttr.Get() == true);
 			//@todo we need to support custom widgets for keyable,animated, in particularly for transforms(ComponentTransformDetails).
 			const bool bPassesKeyableFilter = (InFilter.bShowKeyable == false);
@@ -530,7 +578,7 @@ static bool PassesAllFilters( FDetailItemNode* ItemNode, const FDetailLayoutCust
 		}
 		else if (InCustomization.HasCustomBuilder())
 		{
-			const bool bPassesTextFilter = bPassesCategoryFilter || Local::StringPassesFilter(InFilter, InCustomization.CustomBuilderRow->GetWidgetRow().FilterTextString.ToString());
+			const bool bPassesTextFilter = bPassesCategoryFilter || bPassesValueFilter || Local::StringPassesFilter(InFilter, InCustomization.CustomBuilderRow->GetWidgetRow().FilterTextString.ToString());
 			//@todo we need to support custom builders for modified, keyable, animated, in particularly for transforms(ComponentTransformDetails).
 			const bool bPassesModifiedFilter = (InFilter.bShowOnlyModifiedProperties == false);
 			const bool bPassesKeyableFilter = (InFilter.bShowKeyable == false);
diff --git a/Engine/Source/Editor/PropertyEditor/Private/ItemPropertyNode.cpp b/Engine/Source/Editor/PropertyEditor/Private/ItemPropertyNode.cpp
--- a/Engine/Source/Editor/PropertyEditor/Private/ItemPropertyNode.cpp
+++ b/Engine/Source/Editor/PropertyEditor/Private/ItemPropertyNode.cpp
@@ -20,14 +20,14 @@ FItemPropertyNode::~FItemPropertyNode(void)
 
 }
 
-uint8* FItemPropertyNode::GetValueBaseAddress(uint8* StartAddress, bool bIsSparseData)
+uint8* FItemPropertyNode::GetValueBaseAddress(uint8* StartAddress, bool bIsSparseData) const
 {
-	FProperty* MyProperty = GetProperty();
+	const FProperty* MyProperty = GetProperty();
 	if( MyProperty && ParentNodeWeakPtr.IsValid())
 	{
-		FArrayProperty* OuterArrayProp = MyProperty->GetOwner<FArrayProperty>();
-		FSetProperty* OuterSetProp = MyProperty->GetOwner<FSetProperty>();
-		FMapProperty* OuterMapProp = MyProperty->GetOwner<FMapProperty>();
+		const FArrayProperty* OuterArrayProp = MyProperty->GetOwner<FArrayProperty>();
+		const FSetProperty* OuterSetProp = MyProperty->GetOwner<FSetProperty>();
+		const FMapProperty* OuterMapProp = MyProperty->GetOwner<FMapProperty>();
 
 		uint8* ValueBaseAddress = ParentNode->GetValueBaseAddress(StartAddress, bIsSparseData);
 
@@ -84,19 +84,19 @@ uint8* FItemPropertyNode::GetValueBaseAddress(uint8* StartAddress, bool bIsSpars
 	return nullptr;
 }
 
-uint8* FItemPropertyNode::GetValueAddress(uint8* StartAddress, bool bIsSparseData)
+uint8* FItemPropertyNode::GetValueAddress(uint8* StartAddress, bool bIsSparseData) const
 {
 	uint8* Result = GetValueBaseAddress(StartAddress, bIsSparseData);
 
-	FProperty* MyProperty = GetProperty();
+	const FProperty* MyProperty = GetProperty();
 
-	FArrayProperty* ArrayProperty = CastField<FArrayProperty>(MyProperty);
-	FSetProperty* SetProperty = CastField<FSetProperty>(MyProperty);
-	FMapProperty* MapProperty = CastField<FMapProperty>(MyProperty);
+	const FArrayProperty* ArrayProperty = CastField<FArrayProperty>(MyProperty);
+	const FSetProperty* SetProperty = CastField<FSetProperty>(MyProperty);
+	const FMapProperty* MapProperty = CastField<FMapProperty>(MyProperty);
 
 	if( Result && ArrayProperty)
 	{
-		FScriptArrayHelper ArrayHelper(ArrayProperty,Result);
+		FScriptArrayHelper ArrayHelper(ArrayProperty, Result);
 		Result = ArrayHelper.GetRawPtr();
 	}
 	else if (Result && SetProperty)
diff --git a/Engine/Source/Editor/PropertyEditor/Private/ItemPropertyNode.h b/Engine/Source/Editor/PropertyEditor/Private/ItemPropertyNode.h
--- a/Engine/Source/Editor/PropertyEditor/Private/ItemPropertyNode.h
+++ b/Engine/Source/Editor/PropertyEditor/Private/ItemPropertyNode.h
@@ -12,8 +12,8 @@ public:
 	FItemPropertyNode();
 	virtual ~FItemPropertyNode();
 
-	virtual uint8* GetValueBaseAddress(uint8* StartAddress, bool bIsSparseData) override;
-	virtual uint8* GetValueAddress(uint8* StartAddress, bool bIsSparseData) override;
+	virtual uint8* GetValueBaseAddress(uint8* StartAddress, bool bIsSparseData) const override;
+	virtual uint8* GetValueAddress(uint8* StartAddress, bool bIsSparseData) const override;
 
 	/**
 	 * Overridden function to get the derived object node
diff --git a/Engine/Source/Editor/PropertyEditor/Private/ObjectPropertyNode.cpp b/Engine/Source/Editor/PropertyEditor/Private/ObjectPropertyNode.cpp
--- a/Engine/Source/Editor/PropertyEditor/Private/ObjectPropertyNode.cpp
+++ b/Engine/Source/Editor/PropertyEditor/Private/ObjectPropertyNode.cpp
@@ -155,7 +155,7 @@ TArray<const UStruct*> FObjectPropertyNode::GetAllStructures() const
 // Intended for a function that gets the number of elements in a container
 DECLARE_DELEGATE_RetVal_OneParam(int32, FGetNumDelegate, const void*);
 
-bool FObjectPropertyNode::GetReadAddressUncached(FPropertyNode& InNode,
+bool FObjectPropertyNode::GetReadAddressUncached(const FPropertyNode& InNode,
 											   bool InRequiresSingleSelection,
 											   FReadAddressListData* OutAddresses,
 											   bool bComparePropertyContents,
@@ -168,7 +168,7 @@ bool FObjectPropertyNode::GetReadAddressUncached(FPropertyNode& InNode,
 		return false;
 	}
 
-	FProperty* InItemProperty = InNode.GetProperty();
+	const FProperty* InItemProperty = InNode.GetProperty();
 	// Is there a InItemProperty bound to the InItemProperty window?
 	if( !InItemProperty )
 	{
@@ -203,7 +203,7 @@ bool FObjectPropertyNode::GetReadAddressUncached(FPropertyNode& InNode,
 
 	if( ArrayOuter || SetOuter || MapOuter)
 	{
-		FPropertyNode* ParentPropertyNode = InNode.GetParentNode();
+		const FPropertyNode* ParentPropertyNode = InNode.GetParentNode();
 		check(ParentPropertyNode);
 		const UObject* TempObject = GetUObject(0);
 		if( TempObject )
@@ -270,9 +270,9 @@ bool FObjectPropertyNode::GetReadAddressUncached(FPropertyNode& InNode,
 		// If the item is an array or set itself, return NULL if there are a different number of
 		// items in the container in different objects, when multi-selecting.
 
-		FArrayProperty* ArrayProp = CastField<FArrayProperty>(InItemProperty);
-		FSetProperty* SetProp = CastField<FSetProperty>(InItemProperty);
-		FMapProperty* MapProp = CastField<FMapProperty>(InItemProperty);
+		const FArrayProperty* ArrayProp = CastField<FArrayProperty>(InItemProperty);
+		const FSetProperty* SetProp = CastField<FSetProperty>(InItemProperty);
+		const FMapProperty* MapProp = CastField<FMapProperty>(InItemProperty);
 
 		if (ArrayProp || SetProp || MapProp)
 		{
@@ -392,7 +392,7 @@ bool FObjectPropertyNode::GetReadAddressUncached(FPropertyNode& InNode,
  * @param InItem		The property to get objects from.
  * @param OutAddresses	Storage array for all of the objects' addresses.
  */
-bool FObjectPropertyNode::GetReadAddressUncached( FPropertyNode& InNode, FReadAddressListData& OutAddresses ) const
+bool FObjectPropertyNode::GetReadAddressUncached(const FPropertyNode& InNode, FReadAddressListData& OutAddresses ) const
 {
 	// Are any objects selected for property editing?
 	if( !GetNumObjects())
@@ -400,7 +400,7 @@ bool FObjectPropertyNode::GetReadAddressUncached( FPropertyNode& InNode, FReadAd
 		return false;
 	}
 
-	FProperty* InItemProperty = InNode.GetProperty();
+	const FProperty* InItemProperty = InNode.GetProperty();
 	// Is there a InItemProperty bound to the InItemProperty window?
 	if( !InItemProperty )
 	{
@@ -422,7 +422,7 @@ bool FObjectPropertyNode::GetReadAddressUncached( FPropertyNode& InNode, FReadAd
 	return true;
 }
 
-uint8* FObjectPropertyNode::GetValueBaseAddress(uint8* StartAddress, bool bIsSparseData)
+uint8* FObjectPropertyNode::GetValueBaseAddress(uint8* StartAddress, bool bIsSparseData) const
 {
 	uint8* Result = StartAddress;
 
diff --git a/Engine/Source/Editor/PropertyEditor/Private/ObjectPropertyNode.h b/Engine/Source/Editor/PropertyEditor/Private/ObjectPropertyNode.h
--- a/Engine/Source/Editor/PropertyEditor/Private/ObjectPropertyNode.h
+++ b/Engine/Source/Editor/PropertyEditor/Private/ObjectPropertyNode.h
@@ -24,8 +24,8 @@ public:
 	/** FPropertyNode Interface */
 	virtual FObjectPropertyNode* AsObjectNode() override { return this;}
 	virtual const FObjectPropertyNode* AsObjectNode() const override { return this; }
-	virtual bool GetReadAddressUncached(FPropertyNode& InNode, bool InRequiresSingleSelection, FReadAddressListData* OutAddresses, bool bComparePropertyContents = true, bool bObjectForceCompare = false, bool bArrayPropertiesCanDifferInSize = false) const override;
-	virtual bool GetReadAddressUncached(FPropertyNode& InNode, FReadAddressListData& OutAddresses) const override;
+	virtual bool GetReadAddressUncached(const FPropertyNode& InNode, bool InRequiresSingleSelection, FReadAddressListData* OutAddresses, bool bComparePropertyContents = true, bool bObjectForceCompare = false, bool bArrayPropertiesCanDifferInSize = false) const override;
+	virtual bool GetReadAddressUncached(const FPropertyNode& InNode, FReadAddressListData& OutAddresses) const override;
 
 	/**
 	 * Returns the UObject at index "n" of the Objects Array
@@ -131,7 +131,7 @@ protected:
 	virtual void InitBeforeNodeFlags() override;
 	virtual void InitChildNodes() override;
 	virtual bool GetQualifiedName( FString& PathPlusIndex, const bool bWithArrayIndex, const FPropertyNode* StopParent = nullptr, bool bIgnoreCategories = false ) const override;
-	virtual uint8* GetValueBaseAddress(uint8* Base, bool bIsSparseData) override;
+	virtual uint8* GetValueBaseAddress(uint8* Base, bool bIsSparseData) const override;
 	/**
 	 * Looks at the Objects array and creates the best base class.  Called by
 	 * Finalize(); that is, when the list of selected objects is being finalized.
diff --git a/Engine/Source/Editor/PropertyEditor/Private/PropertyHandleImpl.cpp b/Engine/Source/Editor/PropertyEditor/Private/PropertyHandleImpl.cpp
--- a/Engine/Source/Editor/PropertyEditor/Private/PropertyHandleImpl.cpp
+++ b/Engine/Source/Editor/PropertyEditor/Private/PropertyHandleImpl.cpp
@@ -27,78 +27,10 @@
 #include "ObjectEditorUtils.h"
 #include "SResetToDefaultPropertyEditor.h"
 #include "PropertyPathHelpers.h"
+#include "PropertyTextUtilities.h"
 
 #define LOCTEXT_NAMESPACE "PropertyHandleImplementation"
 
-void PropertyToTextHelper(FString& OutString, FPropertyNode* InPropertyNode, FProperty* Property, uint8* ValueAddress, EPropertyPortFlags PortFlags)
-{
-	if (InPropertyNode->GetArrayIndex() != INDEX_NONE || Property->ArrayDim == 1)
-	{
-		Property->ExportText_Direct(OutString, ValueAddress, ValueAddress, nullptr, PortFlags);
-	}
-	else
-	{
-		FArrayProperty::ExportTextInnerItem(OutString, Property, ValueAddress, Property->ArrayDim, ValueAddress, Property->ArrayDim, nullptr, PortFlags);
-	}
-}
-
-void PropertyToTextHelper(FString& OutString, FPropertyNode* InPropertyNode, FProperty* Property, const FObjectBaseAddress& ObjectAddress, EPropertyPortFlags PortFlags)
-{
-	bool bIsSparseProperty = !!InPropertyNode->HasNodeFlags(EPropertyNodeFlags::IsSparseClassData);
-	bool bIsInContainer = false;
-	FProperty* Outer = Property->GetOwner<FProperty>();
-	if (bIsSparseProperty)
-	{
-		while (Outer)
-		{
-			FArrayProperty* ArrayOuter = Property->GetOwner<FArrayProperty>();
-			FSetProperty* SetOuter = Property->GetOwner<FSetProperty>();
-			FMapProperty* MapOuter = Property->GetOwner<FMapProperty>();
-			if (ArrayOuter || SetOuter || MapOuter)
-			{
-				bIsInContainer = true;
-				break;
-			}
-
-			Outer = Outer->GetOwner<FProperty>();
-		}
-	}
-	if (!bIsSparseProperty || bIsInContainer)
-	{
-		PropertyToTextHelper(OutString, InPropertyNode, Property, ObjectAddress.BaseAddress, PortFlags);
-	}
-	else
-	{
-// TODO: once we're sure that these don't differ we should always use the call to PropertyToTextHelper
-		void* BaseAddress = ObjectAddress.Object->GetClass()->GetOrCreateSparseClassData();
-		void* ValueAddress = Property->ContainerPtrToValuePtr<void>(BaseAddress);
-		Property->ExportText_Direct(OutString, ValueAddress, ValueAddress, nullptr, PortFlags);
-
-		FString Test;
-		PropertyToTextHelper(Test, InPropertyNode, Property, (uint8*)ValueAddress, PortFlags);
-		check(Test.Compare(OutString) == 0);
-	}
-}
-
-void TextToPropertyHelper(const TCHAR* Buffer, FPropertyNode* InPropertyNode, FProperty* Property, uint8* ValueAddress, UObject* Object)
-{
-	if (InPropertyNode->GetArrayIndex() != INDEX_NONE || Property->ArrayDim == 1)
-	{
-		Property->ImportText(Buffer, ValueAddress, 0, Object);
-	}
-	else
-	{
-		FArrayProperty::ImportTextInnerItem(Buffer, Property, ValueAddress, 0, Object);
-	}
-}
-
-void TextToPropertyHelper(const TCHAR* Buffer, FPropertyNode* InPropertyNode, FProperty* Property, const FObjectBaseAddress& ObjectAddress)
-{
-
-	uint8* BaseAddress = InPropertyNode ? InPropertyNode->GetValueBaseAddressFromObject(ObjectAddress.Object) : ObjectAddress.BaseAddress;
-	TextToPropertyHelper(Buffer, InPropertyNode, Property, BaseAddress, ObjectAddress.Object);
-}
-
 FPropertyValueImpl::FPropertyValueImpl( TSharedPtr<FPropertyNode> InPropertyNode, FNotifyHook* InNotifyHook, TSharedPtr<IPropertyUtilities> InPropertyUtilities )
 	: PropertyNode( InPropertyNode )
 	, PropertyUtilities( InPropertyUtilities )
@@ -141,141 +73,6 @@ void FPropertyValueImpl::GetObjectsToModify( TArray<FObjectBaseAddress>& Objects
 	});
 }
 
-FPropertyAccess::Result FPropertyValueImpl::GetPropertyValueString( FString& OutString, FPropertyNode* InPropertyNode, const bool bAllowAlternateDisplayValue, EPropertyPortFlags PortFlags ) const
-{
-	FPropertyAccess::Result Result = FPropertyAccess::Success;
-
-	uint8* ValueAddress = nullptr;
-	FReadAddressList ReadAddresses;
-	bool bAllValuesTheSame = InPropertyNode->GetReadAddress( !!InPropertyNode->HasNodeFlags(EPropertyNodeFlags::SingleSelectOnly), ReadAddresses, false, true );
-
-	if( (ReadAddresses.Num() > 0 && bAllValuesTheSame) || ReadAddresses.Num() == 1 ) 
-	{
-		ValueAddress = ReadAddresses.GetAddress(0);
-
-		if( ValueAddress != nullptr )
-		{
-			FProperty* Property = InPropertyNode->GetProperty();
-
-			// Check for bogus data
-			if( Property != nullptr && InPropertyNode->GetParentNode() != nullptr )
-			{
-				PropertyToTextHelper(OutString, InPropertyNode, Property, ValueAddress, PortFlags);
-
-				UEnum* Enum = nullptr;
-				int64 EnumValue = 0;
-				if ( FByteProperty* ByteProperty = CastField<FByteProperty>(Property) )
-				{
-					if ( ByteProperty->Enum != nullptr )
-					{
-						Enum = ByteProperty->Enum;
-						EnumValue = ByteProperty->GetPropertyValue(ValueAddress);
-					}
-				}
-				else if ( FEnumProperty* EnumProperty = CastField<FEnumProperty>(Property) )
-				{
-					Enum = EnumProperty->GetEnum();
-					EnumValue = EnumProperty->GetUnderlyingProperty()->GetSignedIntPropertyValue(ValueAddress);
-				}
-
-				if ( Enum != nullptr )
-				{
-					if (Enum->IsValidEnumValue(EnumValue))
-					{
-						// See if we specified an alternate name for this value using metadata
-						OutString = Enum->GetDisplayNameTextByValue(EnumValue).ToString();
-						if(!bAllowAlternateDisplayValue || OutString.Len() == 0) 
-						{
-							OutString = Enum->GetNameStringByValue(EnumValue);
-						}
-					}
-					else
-					{
-						Result = FPropertyAccess::Fail;
-					}
-				}
-			}
-			else
-			{
-				Result = FPropertyAccess::Fail;
-			}
-		}
-
-	}
-	else
-	{
-		Result = ReadAddresses.Num() > 1 ? FPropertyAccess::MultipleValues : FPropertyAccess::Fail;
-	}
-
-	return Result;
-}
-
-FPropertyAccess::Result FPropertyValueImpl::GetPropertyValueText( FText& OutText, FPropertyNode* InPropertyNode, const bool bAllowAlternateDisplayValue ) const
-{
-	FPropertyAccess::Result Result = FPropertyAccess::Success;
-
-	uint8* ValueAddress = nullptr;
-	FReadAddressList ReadAddresses;
-	bool bAllValuesTheSame = InPropertyNode->GetReadAddress( !!InPropertyNode->HasNodeFlags(EPropertyNodeFlags::SingleSelectOnly), ReadAddresses, false, true );
-
-	if( ReadAddresses.Num() > 0 && InPropertyNode->GetProperty() != nullptr && (bAllValuesTheSame || ReadAddresses.Num() == 1) ) 
-	{
-		ValueAddress = ReadAddresses.GetAddress(0);
-
-		if( ValueAddress != nullptr )
-		{
-			FProperty* Property = InPropertyNode->GetProperty();
-
-			if(Property->IsA(FTextProperty::StaticClass()))
-			{
-				OutText = CastField<FTextProperty>(Property)->GetPropertyValue(ValueAddress);
-			}
-			else
-			{
-				FString ExportedTextString;
-				PropertyToTextHelper(ExportedTextString, InPropertyNode, Property, ValueAddress, PPF_PropertyWindow);
-
-				UEnum* Enum = nullptr;
-				int64 EnumValue = 0;
-				if (FByteProperty* ByteProperty = CastField<FByteProperty>(Property))
-				{
-					Enum = ByteProperty->Enum;
-					EnumValue = ByteProperty->GetPropertyValue(ValueAddress);
-				}
-				else if (FEnumProperty* EnumProperty = CastField<FEnumProperty>(Property))
-				{
-					Enum = EnumProperty->GetEnum();
-					EnumValue = EnumProperty->GetUnderlyingProperty()->GetSignedIntPropertyValue(ValueAddress);
-				}
-
-				if (Enum)
-				{
-					if (Enum->IsValidEnumValue(EnumValue))
-					{
-						// Text form is always display name
-						OutText = Enum->GetDisplayNameTextByValue(EnumValue);
-					}
-					else
-					{
-						Result = FPropertyAccess::Fail;
-					}
-				}
-				else
-				{
-					OutText = FText::FromString(ExportedTextString);
-				}
-			}
-		}
-
-	}
-	else
-	{
-		Result = ReadAddresses.Num() > 1 ? FPropertyAccess::MultipleValues : FPropertyAccess::Fail;
-	}
-
-	return Result;
-}
-
 FPropertyAccess::Result FPropertyValueImpl::GetValueData( void*& OutAddress ) const
 {
 	FPropertyAccess::Result Res = FPropertyAccess::Fail;
@@ -478,7 +275,7 @@ FPropertyAccess::Result FPropertyValueImpl::ImportText( const TArray<FObjectBase
 
 			// Cache the value of the property before modifying it.
 			FString PreviousValue;
-			PropertyToTextHelper(PreviousValue, InPropertyNode, NodeProperty, Cur, PPF_None);
+			FPropertyTextUtilities::PropertyToTextHelper(PreviousValue, InPropertyNode, NodeProperty, Cur, PPF_None);
 
 			// If this property is the inner-property of a container, cache the current value as well
 			FString PreviousContainerValue;
@@ -568,13 +365,13 @@ FPropertyAccess::Result FPropertyValueImpl::ImportText( const TArray<FObjectBase
 
 			// Set the new value.
 			const TCHAR* NewValue = *InValues[ObjectIndex];
-			TextToPropertyHelper(NewValue, InPropertyNode, NodeProperty, Cur);
+			FPropertyTextUtilities::TextToPropertyHelper(NewValue, InPropertyNode, NodeProperty, Cur);
 
 			if (Cur.Object)
 			{
 				// Cache the value of the property after having modified it.
 				FString ValueAfterImport;
-				PropertyToTextHelper(ValueAfterImport, InPropertyNode, NodeProperty, Cur, PPF_None);
+				FPropertyTextUtilities::PropertyToTextHelper(ValueAfterImport, InPropertyNode, NodeProperty, Cur, PPF_None);
 
 				if ((Cur.Object->HasAnyFlags(RF_ClassDefaultObject | RF_ArchetypeObject) ||
 					(Cur.Object->HasAnyFlags(RF_DefaultSubObject) && Cur.Object->GetOuter()->HasAnyFlags(RF_ClassDefaultObject | RF_ArchetypeObject))) &&
@@ -800,7 +597,7 @@ FPropertyAccess::Result FPropertyValueImpl::GetValueAsString( FString& OutString
 	if( PropertyNodePin.IsValid() )
 	{
 		const bool bAllowAlternateDisplayValue = false;
-		Res = GetPropertyValueString( OutString, PropertyNodePin.Get(), bAllowAlternateDisplayValue, PortFlags );
+		Res = PropertyNodePin->GetPropertyValueString( OutString, bAllowAlternateDisplayValue, PortFlags );
 	}
 	else
 	{
@@ -819,7 +616,7 @@ FPropertyAccess::Result FPropertyValueImpl::GetValueAsDisplayString( FString& Ou
 	if( PropertyNodePin.IsValid() )
 	{
 		const bool bAllowAlternateDisplayValue = true;
-		Res = GetPropertyValueString(OutString, PropertyNodePin.Get(), bAllowAlternateDisplayValue, PortFlags);
+		Res = PropertyNodePin->GetPropertyValueString(OutString, bAllowAlternateDisplayValue, PortFlags);
 	}
 	else
 	{
@@ -837,7 +634,7 @@ FPropertyAccess::Result FPropertyValueImpl::GetValueAsText( FText& OutText ) con
 
 	if( PropertyNodePin.IsValid() )
 	{
-		Res = GetPropertyValueText( OutText, PropertyNodePin.Get(), false/*bAllowAlternateDisplayValue*/ );
+		Res = PropertyNodePin->GetPropertyValueText( OutText, false/*bAllowAlternateDisplayValue*/ );
 	}
 	else
 	{
@@ -855,7 +652,7 @@ FPropertyAccess::Result FPropertyValueImpl::GetValueAsDisplayText( FText& OutTex
 
 	if( PropertyNodePin.IsValid() )
 	{
-		Res = GetPropertyValueText( OutText, PropertyNodePin.Get(), true/*bAllowAlternateDisplayValue*/ );
+		Res = PropertyNodePin->GetPropertyValueText( OutText, true/*bAllowAlternateDisplayValue*/ );
 	}
 	else
 	{
@@ -2764,7 +2561,7 @@ void FPropertyHandleBase::SetToolTipText( const FText& ToolTip )
 	}
 }
 
-uint8* FPropertyHandleBase::GetValueBaseAddress(uint8* Base)
+uint8* FPropertyHandleBase::GetValueBaseAddress(uint8* Base) const
 {
 	TSharedPtr<FPropertyNode> PropertyNode = Implementation->GetPropertyNode();
 	if (PropertyNode.IsValid())
diff --git a/Engine/Source/Editor/PropertyEditor/Private/PropertyHandleImpl.h b/Engine/Source/Editor/PropertyEditor/Private/PropertyHandleImpl.h
--- a/Engine/Source/Editor/PropertyEditor/Private/PropertyHandleImpl.h
+++ b/Engine/Source/Editor/PropertyEditor/Private/PropertyHandleImpl.h
@@ -345,25 +345,6 @@ public:
 
 	void ShowInvalidOperationError(const FText& ErrorText);
 
-protected:
-	/**
-	 * 
-	 * @param InPropertyNode				The property to get value from
-	 * @param OutText						The property formatted in a string
-	 * @param bAllowAlternateDisplayValue	Allow the function to potentially use an alternate form more suitable for display in the UI
-	 * @param PortFlags						Determines how the property's value is accessed. Defaults to PPF_PropertyWindow
-	 * @return true if the value was retrieved successfully
-	 */
-	FPropertyAccess::Result GetPropertyValueString( FString& OutString, FPropertyNode* InPropertyNode, const bool bAllowAlternateDisplayValue , EPropertyPortFlags PortFlags = PPF_PropertyWindow ) const;
-
-	/**
-	 * @param InPropertyNode	The property to get value from
-	 * @param OutText			The property formatted in text
-	 * @param bAllowAlternateDisplayValue Allow the function to potentially use an alternate form more suitable for display in the UI
-	 * @return true if the value was retrieved successfully
-	 */
-	FPropertyAccess::Result GetPropertyValueText( FText& OutText, FPropertyNode* InPropertyNode, const bool bAllowAlternateDisplayValue ) const;
-
 protected:
 	/** Property node used to access FProperty and address of object to change */
 	TWeakPtr<FPropertyNode> PropertyNode;
@@ -482,7 +463,7 @@ public:
 	virtual bool HasDocumentation() override { return false; }
 	virtual FString GetDocumentationLink() override { return FString(); }
 	virtual FString GetDocumentationExcerptName() override { return FString(); }
-	virtual uint8* GetValueBaseAddress( uint8* Base ) override;
+	virtual uint8* GetValueBaseAddress( uint8* Base ) const override;
 	virtual int32 GetNumPerObjectValues() const override;
 	virtual FPropertyAccess::Result SetPerObjectValues( const TArray<FString>& InPerObjectValues,  EPropertyValueSetFlags::Type Flags = EPropertyValueSetFlags::DefaultFlags ) override;
 	virtual FPropertyAccess::Result GetPerObjectValues( TArray<FString>& OutPerObjectValues ) const override;
diff --git a/Engine/Source/Editor/PropertyEditor/Private/PropertyNode.cpp b/Engine/Source/Editor/PropertyEditor/Private/PropertyNode.cpp
--- a/Engine/Source/Editor/PropertyEditor/Private/PropertyNode.cpp
+++ b/Engine/Source/Editor/PropertyEditor/Private/PropertyNode.cpp
@@ -22,6 +22,7 @@
 
 #include "Framework/Notifications/NotificationManager.h"
 #include "Widgets/Notifications/SNotificationList.h"
+#include "PropertyTextUtilities.h"
 
 #define LOCTEXT_NAMESPACE "PropertyNode"
 
@@ -665,12 +666,121 @@ EPropertyDataValidationResult FPropertyNode::EnsureDataIsValid()
 	return FinalResult;
 }
 
+FPropertyAccess::Result FPropertyNode::GetPropertyValueString(FString& OutString, const bool bAllowAlternateDisplayValue, EPropertyPortFlags PortFlags) const
+{
+	uint8* ValueAddress = nullptr;
+	FPropertyAccess::Result Result = GetSingleReadAddress(ValueAddress);
+
+	if (ValueAddress != nullptr)
+	{
+		const FProperty* PropertyPtr = GetProperty();
+
+		// Check for bogus data
+		if (PropertyPtr != nullptr && GetParentNode() != nullptr)
+		{
+			FPropertyTextUtilities::PropertyToTextHelper(OutString, this, PropertyPtr, ValueAddress, PortFlags);
+
+			UEnum* Enum = nullptr;
+			int64 EnumValue = 0;
+			if (const FByteProperty* ByteProperty = CastField<FByteProperty>(PropertyPtr))
+			{
+				if (ByteProperty->Enum != nullptr)
+				{
+					Enum = ByteProperty->Enum;
+					EnumValue = ByteProperty->GetPropertyValue(ValueAddress);
+				}
+			}
+			else if (const FEnumProperty* EnumProperty = CastField<FEnumProperty>(PropertyPtr))
+			{
+				Enum = EnumProperty->GetEnum();
+				EnumValue = EnumProperty->GetUnderlyingProperty()->GetSignedIntPropertyValue(ValueAddress);
+			}
+
+			if (Enum != nullptr)
+			{
+				if (Enum->IsValidEnumValue(EnumValue))
+				{
+					// See if we specified an alternate name for this value using metadata
+					OutString = Enum->GetDisplayNameTextByValue(EnumValue).ToString();
+					if (!bAllowAlternateDisplayValue || OutString.Len() == 0)
+					{
+						OutString = Enum->GetNameStringByValue(EnumValue);
+					}
+				}
+				else
+				{
+					Result = FPropertyAccess::Fail;
+				}
+			}
+		}
+		else
+		{
+			Result = FPropertyAccess::Fail;
+		}
+	}
+
+	return Result;
+}
+
+FPropertyAccess::Result FPropertyNode::GetPropertyValueText(FText& OutText, const bool bAllowAlternateDisplayValue) const
+{
+	uint8* ValueAddress = nullptr;
+	FPropertyAccess::Result Result = GetSingleReadAddress(ValueAddress);
+
+	if (ValueAddress != nullptr)
+	{
+		const FProperty* PropertyPtr = GetProperty();
+
+		if (PropertyPtr->IsA(FTextProperty::StaticClass()))
+		{
+			OutText = CastField<FTextProperty>(PropertyPtr)->GetPropertyValue(ValueAddress);
+		}
+		else
+		{
+			FString ExportedTextString;
+			FPropertyTextUtilities::PropertyToTextHelper(ExportedTextString, this, PropertyPtr, ValueAddress, PPF_PropertyWindow);
+
+			UEnum* Enum = nullptr;
+			int64 EnumValue = 0;
+			if (const FByteProperty* ByteProperty = CastField<FByteProperty>(PropertyPtr))
+			{
+				Enum = ByteProperty->Enum;
+				EnumValue = ByteProperty->GetPropertyValue(ValueAddress);
+			}
+			else if (const FEnumProperty* EnumProperty = CastField<FEnumProperty>(PropertyPtr))
+			{
+				Enum = EnumProperty->GetEnum();
+				EnumValue = EnumProperty->GetUnderlyingProperty()->GetSignedIntPropertyValue(ValueAddress);
+			}
+
+			if (Enum)
+			{
+				if (Enum->IsValidEnumValue(EnumValue))
+				{
+					// Text form is always display name
+					OutText = Enum->GetDisplayNameTextByValue(EnumValue);
+				}
+				else
+				{
+					Result = FPropertyAccess::Fail;
+				}
+			}
+			else
+			{
+				OutText = FText::FromString(ExportedTextString);
+			}
+		}
+	}
+
+	return Result;
+}
+
 /**
  * Sets the flags used by the window and the root node
  * @param InFlags - flags to turn on or off
  * @param InOnOff - whether to toggle the bits on or off
  */
-void FPropertyNode::SetNodeFlags (const EPropertyNodeFlags::Type InFlags, const bool InOnOff)
+void FPropertyNode::SetNodeFlags(const EPropertyNodeFlags::Type InFlags, const bool InOnOff)
 {
 	if (InOnOff)
 	{
@@ -835,7 +945,7 @@ bool FPropertyNode::GetQualifiedName( FString& PathPlusIndex, const bool bWithAr
 	return bAddedAnything;
 }
 
-bool FPropertyNode::GetReadAddressUncached( FPropertyNode& InPropertyNode,
+bool FPropertyNode::GetReadAddressUncached( const FPropertyNode& InPropertyNode,
 									bool InRequiresSingleSelection,
 									FReadAddressListData* OutAddresses,
 									bool bComparePropertyContents,
@@ -850,7 +960,7 @@ bool FPropertyNode::GetReadAddressUncached( FPropertyNode& InPropertyNode,
 	return false;
 }
 
-bool FPropertyNode::GetReadAddressUncached( FPropertyNode& InPropertyNode, FReadAddressListData& OutAddresses ) const
+bool FPropertyNode::GetReadAddressUncached( const FPropertyNode& InPropertyNode, FReadAddressListData& OutAddresses ) const
 {
 	if (ParentNodeWeakPtr.IsValid())
 	{
@@ -863,7 +973,7 @@ bool FPropertyNode::GetReadAddress(bool InRequiresSingleSelection,
 								   FReadAddressList& OutAddresses,
 								   bool bComparePropertyContents,
 								   bool bObjectForceCompare,
-								   bool bArrayPropertiesCanDifferInSize)
+								   bool bArrayPropertiesCanDifferInSize) const
 
 {
 
@@ -893,7 +1003,7 @@ bool FPropertyNode::GetReadAddress(bool InRequiresSingleSelection,
  * @param InItem		The property to get objects from.
  * @param OutAddresses	Storage array for all of the objects' addresses.
  */
-bool FPropertyNode::GetReadAddress( FReadAddressList& OutAddresses )
+bool FPropertyNode::GetReadAddress( FReadAddressList& OutAddresses ) const
 {
 	// @todo PropertyEditor Nodes which require validation cannot be cached
 	if( CachedReadAddresses.Num() && !HasNodeFlags(EPropertyNodeFlags::RequiresValidation) )
@@ -918,6 +1028,22 @@ bool FPropertyNode::GetReadAddress( FReadAddressList& OutAddresses )
 	return bSuccess;
 }
 
+FPropertyAccess::Result FPropertyNode::GetSingleReadAddress(uint8*& OutValueAddress) const
+{
+	OutValueAddress = nullptr;
+	FReadAddressList ReadAddresses;
+	bool bAllValuesTheSame = GetReadAddress(!!HasNodeFlags(EPropertyNodeFlags::SingleSelectOnly), ReadAddresses, false, true);
+
+	if ((ReadAddresses.Num() > 0 && bAllValuesTheSame) || ReadAddresses.Num() == 1)
+	{
+		OutValueAddress = ReadAddresses.GetAddress(0);
+
+		return FPropertyAccess::Success;
+	}
+
+	return ReadAddresses.Num() > 1 ? FPropertyAccess::MultipleValues : FPropertyAccess::Fail;
+}
+
 uint8* FPropertyNode::GetStartAddress(const UObject* Obj) const
 {
 	if (!Obj)
@@ -933,18 +1059,18 @@ uint8* FPropertyNode::GetStartAddress(const UObject* Obj) const
 	return (uint8*)Obj;
 }
 
-uint8* FPropertyNode::GetValueBaseAddressFromObject(const UObject* Obj)
+uint8* FPropertyNode::GetValueBaseAddressFromObject(const UObject* Obj) const
 {
 	return GetValueBaseAddress(GetStartAddress(Obj), HasNodeFlags(EPropertyNodeFlags::IsSparseClassData) != 0);
 }
 
-uint8* FPropertyNode::GetValueAddressFromObject(const UObject* Obj)
+uint8* FPropertyNode::GetValueAddressFromObject(const UObject* Obj) const
 {
 	return GetValueAddress(GetStartAddress(Obj), HasNodeFlags(EPropertyNodeFlags::IsSparseClassData) != 0);
 }
 
 
-uint8* FPropertyNode::GetValueBaseAddress(uint8* StartAddress, bool bIsSparseData)
+uint8* FPropertyNode::GetValueBaseAddress(uint8* StartAddress, bool bIsSparseData) const
 {
 	uint8* Result = NULL;
 
@@ -962,7 +1088,7 @@ uint8* FPropertyNode::GetValueBaseAddress(uint8* StartAddress, bool bIsSparseDat
 	return Result;
 }
 
-uint8* FPropertyNode::GetValueAddress(uint8* StartAddress, bool bIsSparseData)
+uint8* FPropertyNode::GetValueAddress(uint8* StartAddress, bool bIsSparseData) const
 {
 	return GetValueBaseAddress(StartAddress, bIsSparseData);
 }
diff --git a/Engine/Source/Editor/PropertyEditor/Private/PropertyNode.h b/Engine/Source/Editor/PropertyEditor/Private/PropertyNode.h
--- a/Engine/Source/Editor/PropertyEditor/Private/PropertyNode.h
+++ b/Engine/Source/Editor/PropertyEditor/Private/PropertyNode.h
@@ -345,6 +345,24 @@ public:
 	 */
 	EPropertyDataValidationResult EnsureDataIsValid();
 
+	//////////////////////////////////////////////////////////////////////////
+	// Text
+
+	/**
+	 * @param OutText						The property formatted in a string
+	 * @param bAllowAlternateDisplayValue	Allow the function to potentially use an alternate form more suitable for display in the UI
+	 * @param PortFlags						Determines how the property's value is accessed. Defaults to PPF_PropertyWindow
+	 * @return true if the value was retrieved successfully
+	 */
+	FPropertyAccess::Result GetPropertyValueString(FString& OutString, const bool bAllowAlternateDisplayValue, EPropertyPortFlags PortFlags = PPF_PropertyWindow) const;
+
+	/**
+	 * @param OutText			The property formatted in text
+	 * @param bAllowAlternateDisplayValue Allow the function to potentially use an alternate form more suitable for display in the UI
+	 * @return true if the value was retrieved successfully
+	 */
+	FPropertyAccess::Result GetPropertyValueText(FText& OutText, const bool bAllowAlternateDisplayValue) const;
+
 	//////////////////////////////////////////////////////////////////////////
 	//Flags
 	uint32 HasNodeFlags(const EPropertyNodeFlags::Type InTestFlags) const { return PropertyNodeFlags & InTestFlags; }
@@ -438,19 +456,26 @@ public:
 		FReadAddressList& OutAddresses,
 		bool bComparePropertyContents = true,
 		bool bObjectForceCompare = false,
-		bool bArrayPropertiesCanDifferInSize = false);
+		bool bArrayPropertiesCanDifferInSize = false) const;
 
 	/**
 	 * fills in the OutAddresses array with the addresses of all of the available objects.
 	 * @param OutAddresses	Storage array for all of the objects' addresses.
 	 */
-	bool GetReadAddress(FReadAddressList& OutAddresses);
+	bool GetReadAddress(FReadAddressList& OutAddresses) const;
+
+	/**
+	 * Fills in the OutValueAddress with the address of the value of all the available objects.
+	 * If multiple items are selected, this will return a null address unless they are all the same value.
+	 * @param OutValueAddress	The address of the item
+	 */
+	FPropertyAccess::Result GetSingleReadAddress(uint8*& OutValueAddress) const;
 
 	/**
 	 * Gets read addresses without accessing cached data.  Is less efficient but gets the must up to date data
 	 */
-	virtual bool GetReadAddressUncached(FPropertyNode& InNode, bool InRequiresSingleSelection, FReadAddressListData* OutAddresses, bool bComparePropertyContents = true, bool bObjectForceCompare = false, bool bArrayPropertiesCanDifferInSize = false) const;
-	virtual bool GetReadAddressUncached(FPropertyNode& InNode, FReadAddressListData& OutAddresses) const;
+	virtual bool GetReadAddressUncached(const FPropertyNode& InNode, bool InRequiresSingleSelection, FReadAddressListData* OutAddresses, bool bComparePropertyContents = true, bool bObjectForceCompare = false, bool bArrayPropertiesCanDifferInSize = false) const;
+	virtual bool GetReadAddressUncached(const FPropertyNode& InNode, FReadAddressListData& OutAddresses) const;
 
 	/**
 	 * Calculates the memory address for the data associated with this item's property.  This is typically the value of a FProperty or a UObject address.
@@ -460,7 +485,7 @@ public:
 	 *
 	 * @return	a pointer to a FProperty value or UObject.  (For dynamic arrays, you'd cast this value to an FArray*)
 	 */
-	virtual uint8* GetValueBaseAddress(uint8* StartAddress, bool bIsSparseData);
+	virtual uint8* GetValueBaseAddress(uint8* StartAddress, bool bIsSparseData) const;
 
 	/**
 	 * Calculates the memory address for the data associated with this item's value.  For most properties, identical to GetValueBaseAddress.  For items corresponding
@@ -471,7 +496,7 @@ public:
 	 *
 	 * @return	a pointer to a FProperty value or UObject.  (For dynamic arrays, you'd cast this value to whatever type is the Inner for the dynamic array)
 	 */
-	virtual uint8* GetValueAddress(uint8* StartAddress, bool bIsSparseData);
+	virtual uint8* GetValueAddress(uint8* StartAddress, bool bIsSparseData) const;
 
 	/**
 	 * Calculates the memory address for the data associated with this item's property.  This is typically the value of a FProperty or a UObject address.
@@ -480,7 +505,7 @@ public:
 	 *
 	 * @return	a pointer to a FProperty value or UObject.  (For dynamic arrays, you'd cast this value to an FArray*)
 	 */
-	uint8* GetValueBaseAddressFromObject(const UObject* Obj);
+	uint8* GetValueBaseAddressFromObject(const UObject* Obj) const;
 
 	/**
 	 * Calculates the memory address for the data associated with this item's value.  For most properties, identical to GetValueBaseAddress.  For items corresponding
@@ -490,7 +515,7 @@ public:
 	 *
 	 * @return	a pointer to a FProperty value or UObject.  (For dynamic arrays, you'd cast this value to whatever type is the Inner for the dynamic array)
 	 */
-	uint8* GetValueAddressFromObject(const UObject* Obj);
+	uint8* GetValueAddressFromObject(const UObject* Obj) const;
 
 	/**
 	 * Sets the display name override to use instead of the display name
@@ -959,7 +984,7 @@ protected:
 	TSharedPtr<FPropertyNode> PropertyKeyNode;
 
 	/** Cached read addresses for this property node */
-	FReadAddressListData CachedReadAddresses;
+	mutable FReadAddressListData CachedReadAddresses;
 
 	/** List of per object default value trackers associated with this property node */
 	TArray< TSharedPtr<FPropertyItemValueDataTrackerSlate> > ObjectDefaultValueTrackers;
diff --git a/Engine/Source/Editor/PropertyEditor/Private/PropertyTextUtilities.cpp b/Engine/Source/Editor/PropertyEditor/Private/PropertyTextUtilities.cpp
new file mode 100644
--- /dev/null
+++ b/Engine/Source/Editor/PropertyEditor/Private/PropertyTextUtilities.cpp
@@ -0,0 +1,73 @@
+// Copyright Epic Games, Inc. All Rights Reserved.
+
+#include "PropertyTextUtilities.h"
+#include "PropertyNode.h"
+#include "PropertyHandleImpl.h"
+
+void FPropertyTextUtilities::PropertyToTextHelper(FString& OutString, const FPropertyNode* InPropertyNode, const FProperty* Property, uint8* ValueAddress, EPropertyPortFlags PortFlags)
+{
+	if (InPropertyNode->GetArrayIndex() != INDEX_NONE || Property->ArrayDim == 1)
+	{
+		Property->ExportText_Direct(OutString, ValueAddress, ValueAddress, nullptr, PortFlags);
+	}
+	else
+	{
+		FArrayProperty::ExportTextInnerItem(OutString, Property, ValueAddress, Property->ArrayDim, ValueAddress, Property->ArrayDim, nullptr, PortFlags);
+	}
+}
+
+void FPropertyTextUtilities::PropertyToTextHelper(FString& OutString, const FPropertyNode* InPropertyNode, const FProperty* Property, const FObjectBaseAddress& ObjectAddress, EPropertyPortFlags PortFlags)
+{
+	bool bIsSparseProperty = !!InPropertyNode->HasNodeFlags(EPropertyNodeFlags::IsSparseClassData);
+	bool bIsInContainer = false;
+	const FProperty* Outer = Property->GetOwner<FProperty>();
+	if (bIsSparseProperty)
+	{
+		while (Outer)
+		{
+			const FArrayProperty* ArrayOuter = Property->GetOwner<FArrayProperty>();
+			const FSetProperty* SetOuter = Property->GetOwner<FSetProperty>();
+			const FMapProperty* MapOuter = Property->GetOwner<FMapProperty>();
+			if (ArrayOuter || SetOuter || MapOuter)
+			{
+				bIsInContainer = true;
+				break;
+			}
+
+			Outer = Outer->GetOwner<FProperty>();
+		}
+	}
+	if (!bIsSparseProperty || bIsInContainer)
+	{
+		PropertyToTextHelper(OutString, InPropertyNode, Property, ObjectAddress.BaseAddress, PortFlags);
+	}
+	else
+	{
+		// TODO: once we're sure that these don't differ we should always use the call to PropertyToTextHelper
+		void* BaseAddress = ObjectAddress.Object->GetClass()->GetOrCreateSparseClassData();
+		void* ValueAddress = Property->ContainerPtrToValuePtr<void>(BaseAddress);
+		Property->ExportText_Direct(OutString, ValueAddress, ValueAddress, nullptr, PortFlags);
+
+		FString Test;
+		PropertyToTextHelper(Test, InPropertyNode, Property, (uint8*)ValueAddress, PortFlags);
+		check(Test.Compare(OutString) == 0);
+	}
+}
+
+void FPropertyTextUtilities::TextToPropertyHelper(const TCHAR* Buffer, const FPropertyNode* InPropertyNode, const FProperty* Property, uint8* ValueAddress, UObject* Object)
+{
+	if (InPropertyNode->GetArrayIndex() != INDEX_NONE || Property->ArrayDim == 1)
+	{
+		Property->ImportText(Buffer, ValueAddress, 0, Object);
+	}
+	else
+	{
+		FArrayProperty::ImportTextInnerItem(Buffer, Property, ValueAddress, 0, Object);
+	}
+}
+
+void FPropertyTextUtilities::TextToPropertyHelper(const TCHAR* Buffer, const FPropertyNode* InPropertyNode, const FProperty* Property, const FObjectBaseAddress& ObjectAddress)
+{
+	uint8* BaseAddress = InPropertyNode ? InPropertyNode->GetValueBaseAddressFromObject(ObjectAddress.Object) : ObjectAddress.BaseAddress;
+	TextToPropertyHelper(Buffer, InPropertyNode, Property, BaseAddress, ObjectAddress.Object);
+}
\ No newline at end of file
diff --git a/Engine/Source/Editor/PropertyEditor/Private/PropertyTextUtilities.h b/Engine/Source/Editor/PropertyEditor/Private/PropertyTextUtilities.h
new file mode 100644
--- /dev/null
+++ b/Engine/Source/Editor/PropertyEditor/Private/PropertyTextUtilities.h
@@ -0,0 +1,17 @@
+// Copyright Epic Games, Inc. All Rights Reserved.
+
+#pragma once
+
+#include "UObject/PropertyPortFlags.h"
+
+class FPropertyNode;
+class FObjectBaseAddress;
+
+class FPropertyTextUtilities
+{
+public:
+	static void PropertyToTextHelper(FString& OutString, const FPropertyNode* InPropertyNode, const FProperty* Property, uint8* ValueAddress, EPropertyPortFlags PortFlags);
+	static void PropertyToTextHelper(FString& OutString, const FPropertyNode* InPropertyNode, const FProperty* Property, const FObjectBaseAddress& ObjectAddress, EPropertyPortFlags PortFlags);
+	static void TextToPropertyHelper(const TCHAR* Buffer, const FPropertyNode* InPropertyNode, const FProperty* Property, uint8* ValueAddress, UObject* Object);
+	static void TextToPropertyHelper(const TCHAR* Buffer, const FPropertyNode* InPropertyNode, const FProperty* Property, const FObjectBaseAddress& ObjectAddress);
+};
\ No newline at end of file
diff --git a/Engine/Source/Editor/PropertyEditor/Private/StructurePropertyNode.h b/Engine/Source/Editor/PropertyEditor/Private/StructurePropertyNode.h
--- a/Engine/Source/Editor/PropertyEditor/Private/StructurePropertyNode.h
+++ b/Engine/Source/Editor/PropertyEditor/Private/StructurePropertyNode.h
@@ -36,14 +36,14 @@ public:
 		return StructData;
 	}
 
-	bool GetReadAddressUncached(FPropertyNode& InPropertyNode, FReadAddressListData& OutAddresses) const override
+	bool GetReadAddressUncached(const FPropertyNode& InPropertyNode, FReadAddressListData& OutAddresses) const override
 	{
 		if (!HasValidStructData())
 		{
 			return false;
 		}
 
-		FProperty* InItemProperty = InPropertyNode.GetProperty();
+		const FProperty* InItemProperty = InPropertyNode.GetProperty();
 		if (!InItemProperty)
 		{
 			return false;
@@ -62,7 +62,7 @@ public:
 		return true;
 	}
 
-	bool GetReadAddressUncached(FPropertyNode& InPropertyNode,
+	bool GetReadAddressUncached(const FPropertyNode& InPropertyNode,
 		bool InRequiresSingleSelection,
 		FReadAddressListData* OutAddresses,
 		bool bComparePropertyContents,
@@ -149,7 +149,7 @@ protected:
 	/** FPropertyNode interface */
 	virtual void InitChildNodes() override;
 
-	virtual uint8* GetValueBaseAddress(uint8* Base, bool bIsSparseData) override
+	virtual uint8* GetValueBaseAddress(uint8* Base, bool bIsSparseData) const override
 	{
 		check(bIsSparseData == false);
 		return HasValidStructData() ? StructData->GetStructMemory() : nullptr;
diff --git a/Engine/Source/Editor/PropertyEditor/Public/PropertyHandle.h b/Engine/Source/Editor/PropertyEditor/Public/PropertyHandle.h
--- a/Engine/Source/Editor/PropertyEditor/Public/PropertyHandle.h
+++ b/Engine/Source/Editor/PropertyEditor/Public/PropertyHandle.h
@@ -164,7 +164,7 @@ public:
 	*
 	* @return A pointer to a FProperty value or UObject.
 	*/
-	virtual uint8* GetValueBaseAddress( uint8* Base ) = 0;
+	virtual uint8* GetValueBaseAddress( uint8* Base ) const = 0;
 
 	/**
 	 * Gets the value formatted as a string.
diff --git a/Engine/Source/Runtime/CoreUObject/Private/UObject/PropertyArray.cpp b/Engine/Source/Runtime/CoreUObject/Private/UObject/PropertyArray.cpp
--- a/Engine/Source/Runtime/CoreUObject/Private/UObject/PropertyArray.cpp
+++ b/Engine/Source/Runtime/CoreUObject/Private/UObject/PropertyArray.cpp
@@ -416,12 +416,12 @@ void FArrayProperty::ExportTextItem( FString& ValueStr, const void* PropertyValu
 	ExportTextInnerItem(ValueStr, Inner, ArrayHelper.GetRawPtr(0), ArrayHelper.Num(), DefaultValue, DefaultSize, Parent, PortFlags, ExportRootScope);
 }
 
-void FArrayProperty::ExportTextInnerItem(FString& ValueStr, FProperty* Inner, const void* PropertyValue, int32 PropertySize, const void* DefaultValue, int32 DefaultSize, UObject* Parent, int32 PortFlags, UObject* ExportRootScope)
+void FArrayProperty::ExportTextInnerItem(FString& ValueStr, const FProperty* Inner, const void* PropertyValue, int32 PropertySize, const void* DefaultValue, int32 DefaultSize, UObject* Parent, int32 PortFlags, UObject* ExportRootScope)
 {
 	checkSlow(Inner);
 
 	uint8* StructDefaults = NULL;
-	FStructProperty* StructProperty = CastField<FStructProperty>(Inner);
+	const FStructProperty* StructProperty = CastField<FStructProperty>(Inner);
 
 	const bool bReadableForm = (0 != (PPF_BlueprintDebugView & PortFlags));
 	const bool bExternalEditor = (0 != (PPF_ExternalEditor & PortFlags));
@@ -502,7 +502,7 @@ const TCHAR* FArrayProperty::ImportText_Internal(const TCHAR* Buffer, void* Data
 	return ImportTextInnerItem(Buffer, Inner, Data, PortFlags, OwnerObject, &ArrayHelper, ErrorText);
 }
 
-const TCHAR* FArrayProperty::ImportTextInnerItem( const TCHAR* Buffer, FProperty* Inner, void* Data, int32 PortFlags, UObject* Parent, FScriptArrayHelper* ArrayHelper, FOutputDevice* ErrorText )
+const TCHAR* FArrayProperty::ImportTextInnerItem( const TCHAR* Buffer, const FProperty* Inner, void* Data, int32 PortFlags, UObject* Parent, FScriptArrayHelper* ArrayHelper, FOutputDevice* ErrorText )
 {
 	checkSlow(Inner);
 
diff --git a/Engine/Source/Runtime/CoreUObject/Public/UObject/UnrealType.h b/Engine/Source/Runtime/CoreUObject/Public/UObject/UnrealType.h
--- a/Engine/Source/Runtime/CoreUObject/Public/UObject/UnrealType.h
+++ b/Engine/Source/Runtime/CoreUObject/Public/UObject/UnrealType.h
@@ -2832,10 +2832,10 @@ public:
 	FString GetCPPTypeCustom(FString* ExtendedTypeText, uint32 CPPExportFlags, const FString& InnerTypeText, const FString& InInnerExtendedTypeText) const;
 
 	/** Called by ExportTextItem, but can also be used by a non-ArrayProperty whose ArrayDim is > 1. */
-	static void ExportTextInnerItem(FString& ValueStr, FProperty* Inner, const void* PropertyValue, int32 PropertySize, const void* DefaultValue, int32 DefaultSize, UObject* Parent = nullptr, int32 PortFlags = 0, UObject* ExportRootScope = nullptr);
+	static void ExportTextInnerItem(FString& ValueStr, const FProperty* Inner, const void* PropertyValue, int32 PropertySize, const void* DefaultValue, int32 DefaultSize, UObject* Parent = nullptr, int32 PortFlags = 0, UObject* ExportRootScope = nullptr);
 
 	/** Called by ImportTextItem, but can also be used by a non-ArrayProperty whose ArrayDim is > 1. ArrayHelper should be supplied by ArrayProperties and nullptr for fixed-size arrays. */
-	static const TCHAR* ImportTextInnerItem(const TCHAR* Buffer, FProperty* Inner, void* Data, int32 PortFlags, UObject* OwnerObject, FScriptArrayHelper* ArrayHelper = nullptr, FOutputDevice* ErrorText = (FOutputDevice*)GWarn);
+	static const TCHAR* ImportTextInnerItem(const TCHAR* Buffer, const FProperty* Inner, void* Data, int32 PortFlags, UObject* OwnerObject, FScriptArrayHelper* ArrayHelper = nullptr, FOutputDevice* ErrorText = (FOutputDevice*)GWarn);
 
 private:
 	FORCEINLINE void SetElementSize()
