Fixing some monolithic build errors.

diff --git a/Engine/Source/Editor/PropertyEditor/Private/PropertyNode.h b/Engine/Source/Editor/PropertyEditor/Private/PropertyNode.h
--- a/Engine/Source/Editor/PropertyEditor/Private/PropertyNode.h
+++ b/Engine/Source/Editor/PropertyEditor/Private/PropertyNode.h
@@ -6,6 +6,7 @@
 #include "UObject/Object.h"
 #include "UObject/UnrealType.h"
 #include "PropertyPath.h"
+#include "PropertyEditorModule.h"
 
 class FComplexPropertyNode;
 class FNotifyHook;
diff --git a/Engine/Source/Editor/PropertyEditor/Private/PropertyTextUtilities.h b/Engine/Source/Editor/PropertyEditor/Private/PropertyTextUtilities.h
--- a/Engine/Source/Editor/PropertyEditor/Private/PropertyTextUtilities.h
+++ b/Engine/Source/Editor/PropertyEditor/Private/PropertyTextUtilities.h
@@ -2,6 +2,9 @@
 
 #pragma once
 
+#include "CoreMinimal.h"
+#include "UObject/Object.h"
+#include "UObject/UnrealType.h"
 #include "UObject/PropertyPortFlags.h"
 
 class FPropertyNode;
