More mono build includes.

diff --git a/Engine/Source/Editor/PropertyEditor/Private/PropertyNode.cpp b/Engine/Source/Editor/PropertyEditor/Private/PropertyNode.cpp
--- a/Engine/Source/Editor/PropertyEditor/Private/PropertyNode.cpp
+++ b/Engine/Source/Editor/PropertyEditor/Private/PropertyNode.cpp
@@ -23,6 +23,9 @@
 #include "Framework/Notifications/NotificationManager.h"
 #include "Widgets/Notifications/SNotificationList.h"
 #include "PropertyTextUtilities.h"
+#include "UObject/TextProperty.h"
+#include "UObject/EnumProperty.h"
+#include "UObject/UnrealType.h"
 
 #define LOCTEXT_NAMESPACE "PropertyNode"
 
