Improved the details panel root object customization API to support multiple objects in the "root object set".  Deprecated old functionality and removed pass through customizations that are no longer needed

diff --git a/Engine/Plugins/Enterprise/DatasmithImporter/Source/DatasmithImporter/Private/UI/DatasmithImportOptionsWindow.cpp b/Engine/Plugins/Enterprise/DatasmithImporter/Source/DatasmithImporter/Private/UI/DatasmithImportOptionsWindow.cpp
--- a/Engine/Plugins/Enterprise/DatasmithImporter/Source/DatasmithImporter/Private/UI/DatasmithImportOptionsWindow.cpp
+++ b/Engine/Plugins/Enterprise/DatasmithImporter/Source/DatasmithImporter/Private/UI/DatasmithImportOptionsWindow.cpp
@@ -150,16 +150,6 @@ void SDatasmithOptionsWindow::Construct(const FArguments& InArgs)
 	DetailsViewArgs.NameAreaSettings = FDetailsViewArgs::HideNameArea;
 	TSharedPtr<IDetailsView> DetailsView = PropertyEditorModule.CreateDetailView(DetailsViewArgs);
 
-	// customized IDetailRootObjectCustomization in order to hide labels from the various Options Objects displayed
-	class FDatasmithRootObjectCustomization : public IDetailRootObjectCustomization
-	{
-		virtual TSharedPtr<SWidget> CustomizeObjectHeader(const UObject* InRootObject) { return SNullWidget::NullWidget; }
-		virtual bool IsObjectVisible(const UObject* InRootObject) const override { return true; }
-		virtual bool ShouldDisplayHeader(const UObject* InRootObject) const override { return false; }
-	};
-	TSharedPtr<IDetailRootObjectCustomization> RootObjectCustomization = MakeShareable(new FDatasmithRootObjectCustomization());
-	DetailsView->SetRootObjectCustomizationInstance(RootObjectCustomization);
-
 	DetailsViewBox->SetContent(DetailsView.ToSharedRef());
 	// @todo: Find a way to aggregate properties of all UObjects into one set
 	DetailsView->SetObjects(ImportOptions);
diff --git a/Engine/Plugins/Experimental/ChaosEditor/Source/FractureEditor/Private/FractureEditorModeToolkit.cpp b/Engine/Plugins/Experimental/ChaosEditor/Source/FractureEditor/Private/FractureEditorModeToolkit.cpp
--- a/Engine/Plugins/Experimental/ChaosEditor/Source/FractureEditor/Private/FractureEditorModeToolkit.cpp
+++ b/Engine/Plugins/Experimental/ChaosEditor/Source/FractureEditor/Private/FractureEditorModeToolkit.cpp
@@ -56,15 +56,6 @@
 
 #define LOCTEXT_NAMESPACE "FFractureEditorModeToolkit"
 
-class FFractureRootObjectCustomization : public IDetailRootObjectCustomization
-{
-public:
-	/** IDetailRootObjectCustomization interface */
-	virtual TSharedPtr<SWidget> CustomizeObjectHeader(const UObject* InRootObject) override { return SNullWidget::NullWidget; }
-	virtual bool IsObjectVisible(const UObject* InRootObject) const override { return true; }
-	virtual bool ShouldDisplayHeader(const UObject* InRootObject) const override { return false; }
-};
-
 TArray<UClass*> FindFractureToolClasses()
 {
 	TArray<UClass*> Classes;
diff --git a/Engine/Plugins/Experimental/ControlRig/Source/ControlRigEditor/Private/EditMode/SControlRigEditModeTools.cpp b/Engine/Plugins/Experimental/ControlRig/Source/ControlRigEditor/Private/EditMode/SControlRigEditModeTools.cpp
--- a/Engine/Plugins/Experimental/ControlRig/Source/ControlRigEditor/Private/EditMode/SControlRigEditModeTools.cpp
+++ b/Engine/Plugins/Experimental/ControlRig/Source/ControlRigEditor/Private/EditMode/SControlRigEditModeTools.cpp
@@ -21,26 +21,6 @@
 
 #define LOCTEXT_NAMESPACE "ControlRigRootCustomization"
 
-class FControlRigRootCustomization : public IDetailRootObjectCustomization
-{
-public:
-	// IDetailRootObjectCustomization interface
-	virtual TSharedPtr<SWidget> CustomizeObjectHeader(const UObject* InRootObject) override
-	{
-		return SNullWidget::NullWidget;
-	}
-
-	virtual bool IsObjectVisible(const UObject* InRootObject) const override
-	{
-		return true;
-	}
-
-	virtual bool ShouldDisplayHeader(const UObject* InRootObAject) const override
-	{
-		return false;
-	}
-};
-
 void SControlRigEditModeTools::SetControlRig(UControlRig* ControlRig)
 {
 	ControlHierarchy->SetControlRig(ControlRig);
@@ -70,7 +50,6 @@ void SControlRigEditModeTools::Construct(const FArguments& InArgs, FControlRigEd
 	DetailsView->SetKeyframeHandler(SharedThis(this));
 	DetailsView->SetIsPropertyVisibleDelegate(FIsPropertyVisible::CreateSP(this, &SControlRigEditModeTools::ShouldShowPropertyOnDetailCustomization));
 	DetailsView->SetIsPropertyReadOnlyDelegate(FIsPropertyReadOnly::CreateSP(this, &SControlRigEditModeTools::IsReadOnlyPropertyOnDetailCustomization));
-	DetailsView->SetRootObjectCustomizationInstance(MakeShareable(new FControlRigRootCustomization));
 
 	ChildSlot
 	[
diff --git a/Engine/Plugins/Experimental/ModelingToolsEditorMode/Source/ModelingToolsEditorMode/Private/ModelingToolsEditorModeToolkit.cpp b/Engine/Plugins/Experimental/ModelingToolsEditorMode/Source/ModelingToolsEditorMode/Private/ModelingToolsEditorModeToolkit.cpp
--- a/Engine/Plugins/Experimental/ModelingToolsEditorMode/Source/ModelingToolsEditorMode/Private/ModelingToolsEditorModeToolkit.cpp
+++ b/Engine/Plugins/Experimental/ModelingToolsEditorMode/Source/ModelingToolsEditorMode/Private/ModelingToolsEditorModeToolkit.cpp
@@ -34,41 +34,13 @@ FModelingToolsEditorModeToolkit::FModelingToolsEditorModeToolkit()
 {
 }
 
+
 FModelingToolsEditorModeToolkit::~FModelingToolsEditorModeToolkit()
 {
 	UModelingToolsEditorModeSettings* Settings = GetMutableDefault<UModelingToolsEditorModeSettings>();
 	Settings->OnModified.Remove(AssetSettingsModifiedHandle);
 }
 
-
-/**
- * Customization for Tool properties multi-object DetailsView, that just hides per-object header
- */
-class FModelingToolsDetailRootObjectCustomization : public IDetailRootObjectCustomization
-{
-public:
-	FModelingToolsDetailRootObjectCustomization()
-	{
-	}
-
-	virtual TSharedPtr<SWidget> CustomizeObjectHeader(const UObject* InRootObject) override
-	{
-		return SNew(STextBlock).Text( FText::FromString(InRootObject->GetName())  );
-	}
-
-	virtual bool IsObjectVisible(const UObject* InRootObject) const override
-	{
-		return true;
-	}
-
-	virtual bool ShouldDisplayHeader(const UObject* InRootObject) const override
-	{
-		return false;
-	}
-};
-
-
-
 void FModelingToolsEditorModeToolkit::Init(const TSharedPtr<IToolkitHost>& InitToolkitHost)
 {
 	FPropertyEditorModule& PropertyEditorModule = FModuleManager::GetModuleChecked<FPropertyEditorModule>("PropertyEditor");
@@ -88,11 +60,6 @@ void FModelingToolsEditorModeToolkit::Init(const TSharedPtr<IToolkitHost>& InitT
 
 	DetailsView = PropertyEditorModule.CreateDetailView(DetailsViewArgs);
 
-	// add customization that hides header labels
-	TSharedPtr<FModelingToolsDetailRootObjectCustomization> RootObjectCustomization 
-		= MakeShared<FModelingToolsDetailRootObjectCustomization>();
-	DetailsView->SetRootObjectCustomizationInstance(RootObjectCustomization);
-
 	ModeWarningArea = SNew(STextBlock)
 		.AutoWrapText(true)
 		.Font(FCoreStyle::GetDefaultFontStyle("Bold", 9))
diff --git a/Engine/Source/Developer/MaterialBaking/Private/MaterialOptionsWindow.cpp b/Engine/Source/Developer/MaterialBaking/Private/MaterialOptionsWindow.cpp
--- a/Engine/Source/Developer/MaterialBaking/Private/MaterialOptionsWindow.cpp
+++ b/Engine/Source/Developer/MaterialBaking/Private/MaterialOptionsWindow.cpp
@@ -34,9 +34,7 @@ void SMaterialOptions::Construct(const FArguments& InArgs)
 	DetailsView = PropertyEditorModule.CreateDetailView(DetailsViewArgs);
 	// Register instance property customization
 	DetailsView->RegisterInstancedCustomPropertyLayout(UMaterialOptions::StaticClass(), FOnGetDetailCustomizationInstance::CreateLambda([=]() { return FMaterialOptionsCustomization::MakeInstance(InArgs._NumLODs); }));
-	// Set up root object customization to get desired layout
-	DetailsView->SetRootObjectCustomizationInstance(MakeShareable(new FSimpleRootObjectCustomization));
-	
+
 	// Set provided objects on SDetailsView
 	DetailsView->SetObjects(InArgs._SettingsObjects, true);
 	
@@ -121,10 +119,5 @@ bool SMaterialOptions::WasUserCancelled()
 	return bUserCancelled;
 }
 
-TSharedPtr<SWidget> FSimpleRootObjectCustomization::CustomizeObjectHeader(const UObject* InRootObject)
-{
-	return SNullWidget::NullWidget;
-}
-
 #undef LOCTEXT_NAMESPACE //"SMaterialOptions"
 
diff --git a/Engine/Source/Developer/MaterialBaking/Private/MaterialOptionsWindow.h b/Engine/Source/Developer/MaterialBaking/Private/MaterialOptionsWindow.h
--- a/Engine/Source/Developer/MaterialBaking/Private/MaterialOptionsWindow.h
+++ b/Engine/Source/Developer/MaterialBaking/Private/MaterialOptionsWindow.h
@@ -10,16 +10,6 @@
 struct FMaterialData;
 class SButton;
 
-/** Required to ensure desired details view layout */
-class FSimpleRootObjectCustomization : public IDetailRootObjectCustomization
-{
-public:
-	/** IDetailRootObjectCustomization interface */
-	virtual TSharedPtr<SWidget> CustomizeObjectHeader(const UObject* InRootObject) override;
-	virtual bool IsObjectVisible(const UObject* InRootObject) const override { return true; }
-	virtual bool ShouldDisplayHeader(const UObject* InRootObject) const override { return false; }
-};
-
 /** Options window used to populate provided settings objects */
 class SMaterialOptions : public SCompoundWidget
 {
diff --git a/Engine/Source/Developer/SettingsEditor/Private/Widgets/SSettingsSectionHeader.cpp b/Engine/Source/Developer/SettingsEditor/Private/Widgets/SSettingsSectionHeader.cpp
--- a/Engine/Source/Developer/SettingsEditor/Private/Widgets/SSettingsSectionHeader.cpp
+++ b/Engine/Source/Developer/SettingsEditor/Private/Widgets/SSettingsSectionHeader.cpp
@@ -528,17 +528,17 @@ void FSettingsDetailRootObjectCustomization::Initialize()
 	OnSelectedSectionChanged();
 }
 
-TSharedPtr<SWidget> FSettingsDetailRootObjectCustomization::CustomizeObjectHeader(const UObject* InRootObject, const TSharedPtr<ITableRow>& InTableRow)
+TSharedPtr<SWidget> FSettingsDetailRootObjectCustomization::CustomizeObjectHeader(const FDetailsObjectSet& InRootObjectSet, const TSharedPtr<ITableRow>& InTableRow)
 {
-	return SNew(SSettingsSectionHeader, InRootObject, Model, DetailsView.Pin(), InTableRow);
+	return SNew(SSettingsSectionHeader, InRootObjectSet.RootObjects[0], Model, DetailsView.Pin(), InTableRow);
 }
 
-bool FSettingsDetailRootObjectCustomization::IsObjectVisible(const UObject* InRootObject) const
+bool FSettingsDetailRootObjectCustomization::AreObjectsVisible(const FDetailsObjectSet& InRootObjectSet) const
 {
-	return SelectedSectionObject == nullptr || SelectedSectionObject == InRootObject || DetailsView.Pin()->HasActiveSearch();
+	return SelectedSectionObject == nullptr || SelectedSectionObject == InRootObjectSet.RootObjects[0] || DetailsView.Pin()->HasActiveSearch();
 }
 
-bool FSettingsDetailRootObjectCustomization::ShouldDisplayHeader(const UObject* InRootObject) const
+bool FSettingsDetailRootObjectCustomization::ShouldDisplayHeader(const FDetailsObjectSet& InRootObjects) const
 {
 	return true;
 }
diff --git a/Engine/Source/Developer/SettingsEditor/Private/Widgets/SSettingsSectionHeader.h b/Engine/Source/Developer/SettingsEditor/Private/Widgets/SSettingsSectionHeader.h
--- a/Engine/Source/Developer/SettingsEditor/Private/Widgets/SSettingsSectionHeader.h
+++ b/Engine/Source/Developer/SettingsEditor/Private/Widgets/SSettingsSectionHeader.h
@@ -118,9 +118,9 @@ public:
 	void Initialize();
 
 	/** IDetailRootObjectCustomization interface */
-	virtual TSharedPtr<SWidget> CustomizeObjectHeader(const UObject* InRootObject, const TSharedPtr<ITableRow>& InTableRow) override;
-	virtual bool IsObjectVisible(const UObject* InRootObject) const override;
-	virtual bool ShouldDisplayHeader(const UObject* InRootObject) const override;
+	virtual TSharedPtr<SWidget> CustomizeObjectHeader(const FDetailsObjectSet& InRootObjectSet, const TSharedPtr<ITableRow>& InTableRow) override;
+	virtual bool AreObjectsVisible(const FDetailsObjectSet& InRootObjectSet) const override;
+	virtual bool ShouldDisplayHeader(const FDetailsObjectSet& InRootObjectSet) const override;
 	virtual EExpansionArrowUsage GetExpansionArrowUsage() const override { return EExpansionArrowUsage::Custom; }
 private:
 	void OnSelectedSectionChanged();
diff --git a/Engine/Source/Editor/AdvancedPreviewScene/Private/SAdvancedPreviewDetailsTab.cpp b/Engine/Source/Editor/AdvancedPreviewScene/Private/SAdvancedPreviewDetailsTab.cpp
--- a/Engine/Source/Editor/AdvancedPreviewScene/Private/SAdvancedPreviewDetailsTab.cpp
+++ b/Engine/Source/Editor/AdvancedPreviewScene/Private/SAdvancedPreviewDetailsTab.cpp
@@ -271,15 +271,6 @@ void SAdvancedPreviewDetailsTab::CreateSettingsView()
 		DetailDelegate.OnPreviewSceneChangedDelegate.AddSP(this, &SAdvancedPreviewDetailsTab::OnPreviewSceneChanged);
 	}
 
-	class FDetailRootObjectCustomization : public IDetailRootObjectCustomization
-	{
-		virtual TSharedPtr<SWidget> CustomizeObjectHeader(const UObject* InRootObject) { return SNullWidget::NullWidget; }
-		virtual bool IsObjectVisible(const UObject* InRootObject) const { return true; }
-		virtual bool ShouldDisplayHeader(const UObject* InRootObject) const { return false; }
-	};
-
-	SettingsView->SetRootObjectCustomizationInstance(MakeShareable(new FDetailRootObjectCustomization));
-
 	UpdateSettingsView();
 }
 
diff --git a/Engine/Source/Editor/ClothPainter/Private/ClothPaintSettingsCustomization.h b/Engine/Source/Editor/ClothPainter/Private/ClothPaintSettingsCustomization.h
--- a/Engine/Source/Editor/ClothPainter/Private/ClothPaintSettingsCustomization.h
+++ b/Engine/Source/Editor/ClothPainter/Private/ClothPaintSettingsCustomization.h
@@ -41,17 +41,6 @@ private:
 };
 
 
-class FClothPaintSettingsRootObjectCustomization : public IDetailRootObjectCustomization
-{
-public:
-
-	/** IDetailRootObjectCustomization interface */
-	virtual TSharedPtr<SWidget> CustomizeObjectHeader(const UObject* InRootObject) override { return SNullWidget::NullWidget; }
-	virtual bool IsObjectVisible(const UObject* InRootObject) const override { return true; }
-	virtual bool ShouldDisplayHeader(const UObject* InRootObject) const override { return false; }
-};
-
-
 class FClothPaintBrushSettingsCustomization : public IDetailCustomization
 {
 public:
diff --git a/Engine/Source/Editor/ClothPainter/Private/SClothPaintWidget.cpp b/Engine/Source/Editor/ClothPainter/Private/SClothPaintWidget.cpp
--- a/Engine/Source/Editor/ClothPainter/Private/SClothPaintWidget.cpp
+++ b/Engine/Source/Editor/ClothPainter/Private/SClothPaintWidget.cpp
@@ -87,7 +87,6 @@ void SClothPaintWidget::CreateDetailsView(FClothPainter* InPainter)
 	DetailsViewArgs.bAllowMultipleTopLevelObjects = true;
 	
 	DetailsView = EditModule.CreateDetailView(DetailsViewArgs);
-	DetailsView->SetRootObjectCustomizationInstance(MakeShareable(new FClothPaintSettingsRootObjectCustomization));
 	DetailsView->RegisterInstancedCustomPropertyLayout(UClothPainterSettings::StaticClass(), FOnGetDetailCustomizationInstance::CreateStatic(&FClothPaintSettingsCustomization::MakeInstance, InPainter));
 	DetailsView->RegisterInstancedCustomPropertyLayout(UPaintBrushSettings::StaticClass(), FOnGetDetailCustomizationInstance::CreateStatic(&FClothPaintBrushSettingsCustomization::MakeInstance));
 
diff --git a/Engine/Source/Editor/MeshPaintMode/Private/PaintModeSettingsCustomization.cpp b/Engine/Source/Editor/MeshPaintMode/Private/PaintModeSettingsCustomization.cpp
--- a/Engine/Source/Editor/MeshPaintMode/Private/PaintModeSettingsCustomization.cpp
+++ b/Engine/Source/Editor/MeshPaintMode/Private/PaintModeSettingsCustomization.cpp
@@ -453,11 +453,6 @@ void FPaintModeSettingsCustomization::OnPaintTypeChanged(IDetailLayoutBuilder* L
 	LayoutBuilder->ForceRefreshDetails();
 }
 
-TSharedPtr<SWidget> FPaintModeSettingsRootObjectCustomization::CustomizeObjectHeader(const UObject* InRootObject)
-{
-	return SNullWidget::NullWidget;
-}
-
 TSharedRef<SHorizontalBox> CreateColorChannelWidget(TSharedRef<IPropertyHandle> ChannelProperty)
 {
 	return SNew(SHorizontalBox)
diff --git a/Engine/Source/Editor/MeshPaintMode/Private/PaintModeSettingsCustomization.h b/Engine/Source/Editor/MeshPaintMode/Private/PaintModeSettingsCustomization.h
--- a/Engine/Source/Editor/MeshPaintMode/Private/PaintModeSettingsCustomization.h
+++ b/Engine/Source/Editor/MeshPaintMode/Private/PaintModeSettingsCustomization.h
@@ -24,14 +24,6 @@ public:
 	void OnPaintTypeChanged(IDetailLayoutBuilder* LayoutBuilder);
 };
 
-class FPaintModeSettingsRootObjectCustomization : public IDetailRootObjectCustomization
-{
-public:
-	/** IDetailRootObjectCustomization interface */
-	virtual TSharedPtr<SWidget> CustomizeObjectHeader(const UObject* InRootObject) override;
-	virtual bool IsObjectVisible(const UObject* InRootObject) const override { return true; }
-	virtual bool ShouldDisplayHeader(const UObject* InRootObject) const override { return false; }
-};
 
 class FVertexPaintSettingsCustomization : public IPropertyTypeCustomization
 {
diff --git a/Engine/Source/Editor/MeshPaintMode/Private/SPaintModeWidget.cpp b/Engine/Source/Editor/MeshPaintMode/Private/SPaintModeWidget.cpp
--- a/Engine/Source/Editor/MeshPaintMode/Private/SPaintModeWidget.cpp
+++ b/Engine/Source/Editor/MeshPaintMode/Private/SPaintModeWidget.cpp
@@ -122,7 +122,6 @@ void SPaintModeWidget::CreateDetailsView()
 	DetailsViewArgs.bAllowMultipleTopLevelObjects = true;
 
 	SettingsDetailsView = EditModule.CreateDetailView(DetailsViewArgs);
-	SettingsDetailsView->SetRootObjectCustomizationInstance(MakeShareable(new FPaintModeSettingsRootObjectCustomization));
 	SettingsDetailsView->SetObjects(SettingsObjects);
 }
 
diff --git a/Engine/Source/Editor/PropertyEditor/Private/DetailLayoutBuilderImpl.cpp b/Engine/Source/Editor/PropertyEditor/Private/DetailLayoutBuilderImpl.cpp
--- a/Engine/Source/Editor/PropertyEditor/Private/DetailLayoutBuilderImpl.cpp
+++ b/Engine/Source/Editor/PropertyEditor/Private/DetailLayoutBuilderImpl.cpp
@@ -329,14 +329,12 @@ void FDetailLayoutBuilderImpl::GenerateDetailLayout()
 	TSharedPtr<FComplexPropertyNode> RootNodePinned = RootNode.Pin();
 	if(DetailsView && DetailsView->GetRootObjectCustomization() && RootNodePinned->GetInstancesNum())
 	{
-		// This should always exist here
-		UObject* RootObject = RootNodePinned->AsObjectNode()->GetUObject(0);
-		check(RootObject);
+		FObjectPropertyNode* ObjectNode = RootNodePinned->AsObjectNode();
 
 		TSharedPtr<IDetailRootObjectCustomization> RootObjectCustomization = DetailsView->GetRootObjectCustomization();
 
 		// there are multiple objects in the details panel.  Separate each one with a unique object name node to differentiate them
-		AllRootTreeNodes.Add( MakeShareable( new FDetailMultiTopLevelObjectRootNode( CategoryNodes, RootObjectCustomization, DetailsView, *RootObject) ) );
+		AllRootTreeNodes.Add(MakeShared<FDetailMultiTopLevelObjectRootNode>(CategoryNodes, RootObjectCustomization, DetailsView, ObjectNode));
 	}
 	else
 	{
diff --git a/Engine/Source/Editor/PropertyEditor/Private/DetailMultiTopLevelObjectRootNode.cpp b/Engine/Source/Editor/PropertyEditor/Private/DetailMultiTopLevelObjectRootNode.cpp
--- a/Engine/Source/Editor/PropertyEditor/Private/DetailMultiTopLevelObjectRootNode.cpp
+++ b/Engine/Source/Editor/PropertyEditor/Private/DetailMultiTopLevelObjectRootNode.cpp
@@ -3,6 +3,7 @@
 #include "DetailMultiTopLevelObjectRootNode.h"
 #include "IDetailRootObjectCustomization.h"
 #include "DetailWidgetRow.h"
+#include "ObjectPropertyNode.h"
 
 void SDetailMultiTopLevelObjectTableRow::Construct( const FArguments& InArgs, TSharedRef<FDetailTreeNode> InOwnerTreeNode, const TSharedRef<STableViewBase>& InOwnerTableView )
 {
@@ -79,20 +80,50 @@ FReply SDetailMultiTopLevelObjectTableRow::OnMouseButtonDoubleClick( const FGeom
 }
 
 
-FDetailMultiTopLevelObjectRootNode::FDetailMultiTopLevelObjectRootNode( const FDetailNodeList& InChildNodes, const TSharedPtr<IDetailRootObjectCustomization>& InRootObjectCustomization, IDetailsViewPrivate* InDetailsView, const UObject& InRootObject )
+FDetailMultiTopLevelObjectRootNode::FDetailMultiTopLevelObjectRootNode( const FDetailNodeList& InChildNodes, const TSharedPtr<IDetailRootObjectCustomization>& InRootObjectCustomization, IDetailsViewPrivate* InDetailsView, const FObjectPropertyNode* RootNode)
 	: ChildNodes(InChildNodes)
 	, DetailsView(InDetailsView)
 	, RootObjectCustomization(InRootObjectCustomization)
-	, RootObject(const_cast<UObject*>(&InRootObject))
-	, NodeName(InRootObject.GetFName())
 	, bShouldBeVisible(false)
+	, bHasFilterStrings(false)
 {
+	RootObjectSet.RootObjects.Reserve(RootNode->GetNumObjects());
+	for (int32 ObjectIndex = 0; ObjectIndex < RootNode->GetNumObjects(); ++ObjectIndex)
+	{
+		RootObjectSet.RootObjects.Add(RootNode->GetUObject(ObjectIndex));
+	}
+
+	RootObjectSet.CommonBaseClass = RootNode->GetObjectBaseClass();
+
+	NodeName = RootObjectSet.CommonBaseClass->GetFName();
+}
+
+void FDetailMultiTopLevelObjectRootNode::OnItemExpansionChanged(bool bIsExpanded, bool bShouldSaveState)
+{
+	if (bShouldSaveState)
+	{
+		GConfig->SetBool(TEXT("DetailMultiObjectNodeExpansion"), *NodeName.ToString(), bIsExpanded, GEditorPerProjectIni);
+	}
+}
+
+bool FDetailMultiTopLevelObjectRootNode::ShouldBeExpanded() const
+{
+	if (bHasFilterStrings)
+	{
+		return true;
+	}
+	else
+	{
+		bool bShouldBeExpanded = true;
+		GConfig->GetBool(TEXT("DetailMultiObjectNodeExpansion"), *NodeName.ToString(), bShouldBeExpanded, GEditorPerProjectIni);
+		return bShouldBeExpanded;
+	}
 }
 
 ENodeVisibility FDetailMultiTopLevelObjectRootNode::GetVisibility() const
 {
 	ENodeVisibility FinalVisibility = ENodeVisibility::Visible;
-	if(RootObjectCustomization.IsValid() && RootObject.IsValid() && !RootObjectCustomization.Pin()->IsObjectVisible(RootObject.Get()))
+	if(RootObjectCustomization.IsValid() && !RootObjectCustomization.Pin()->AreObjectsVisible(RootObjectSet))
 	{
 		FinalVisibility = ENodeVisibility::ForcedHidden;
 	}
@@ -156,6 +187,8 @@ void FDetailMultiTopLevelObjectRootNode::GetChildren(FDetailNodeList& OutChildre
 void FDetailMultiTopLevelObjectRootNode::FilterNode( const FDetailFilter& InFilter )
 {
 	bShouldBeVisible = false;
+	bHasFilterStrings = InFilter.FilterStrings.Num() > 0;
+
 	for( int32 ChildIndex = 0; ChildIndex < ChildNodes.Num(); ++ChildIndex )
 	{
 		TSharedRef<FDetailTreeNode>& Child = ChildNodes[ChildIndex];
@@ -176,25 +209,15 @@ void FDetailMultiTopLevelObjectRootNode::FilterNode( const FDetailFilter& InFilt
 
 bool FDetailMultiTopLevelObjectRootNode::ShouldShowOnlyChildren() const
 {
-	return RootObjectCustomization.IsValid() && RootObject.IsValid() ? !RootObjectCustomization.Pin()->ShouldDisplayHeader(RootObject.Get()) : bShouldShowOnlyChildren;
+	return RootObjectCustomization.IsValid() && RootObjectSet.RootObjects.Num() ? !RootObjectCustomization.Pin()->ShouldDisplayHeader(RootObjectSet) : bShouldShowOnlyChildren;
 }
 
 void FDetailMultiTopLevelObjectRootNode::GenerateWidget_Internal(FDetailWidgetRow& OutRow, TSharedPtr<SDetailMultiTopLevelObjectTableRow> TableRowWidget) const
 {
 	TSharedPtr<SWidget> HeaderWidget;
-	if (RootObjectCustomization.IsValid() && RootObject.IsValid())
-	{
-		HeaderWidget = RootObjectCustomization.Pin()->CustomizeObjectHeader(RootObject.Get(), TableRowWidget);
-	}
-
-	if (!HeaderWidget.IsValid())
+	if (RootObjectCustomization.IsValid() && RootObjectSet.RootObjects.Num())
 	{
-		// no customization was supplied or was passed back from the interface as invalid
-		// just make a text block with the name
-		HeaderWidget =
-			SNew(STextBlock)
-			.Font(FEditorStyle::GetFontStyle("DetailsView.CategoryFontStyle"))
-			.Text(FText::FromName(NodeName));
+		HeaderWidget = RootObjectCustomization.Pin()->CustomizeObjectHeader(RootObjectSet, TableRowWidget);
 	}
 
 	OutRow.NameContent()
diff --git a/Engine/Source/Editor/PropertyEditor/Private/DetailMultiTopLevelObjectRootNode.h b/Engine/Source/Editor/PropertyEditor/Private/DetailMultiTopLevelObjectRootNode.h
--- a/Engine/Source/Editor/PropertyEditor/Private/DetailMultiTopLevelObjectRootNode.h
+++ b/Engine/Source/Editor/PropertyEditor/Private/DetailMultiTopLevelObjectRootNode.h
@@ -33,8 +33,8 @@ private:
 	const FSlateBrush* GetBackgroundImage() const;
 
 private:
-	virtual FReply OnMouseButtonDown( const FGeometry& MyGeometry, const FPointerEvent& MouseEvent ) override;
-	virtual FReply OnMouseButtonDoubleClick( const FGeometry& InMyGeometry, const FPointerEvent& InMouseEvent ) override;
+	virtual FReply OnMouseButtonDown(const FGeometry& MyGeometry, const FPointerEvent& MouseEvent) override;
+	virtual FReply OnMouseButtonDoubleClick(const FGeometry& InMyGeometry, const FPointerEvent& InMouseEvent) override;
 private:
 	EExpansionArrowUsage ExpansionArrowUsage;
 	SHorizontalBox::FSlot* ContentSlot = nullptr;
@@ -44,17 +44,17 @@ private:
 class FDetailMultiTopLevelObjectRootNode : public FDetailTreeNode, public TSharedFromThis<FDetailMultiTopLevelObjectRootNode>
 {
 public:
-	FDetailMultiTopLevelObjectRootNode( const FDetailNodeList& InChildNodes, const TSharedPtr<IDetailRootObjectCustomization>& RootObjectCustomization, IDetailsViewPrivate* InDetailsView, const UObject& InRootObject );
+	FDetailMultiTopLevelObjectRootNode(const FDetailNodeList& InChildNodes, const TSharedPtr<IDetailRootObjectCustomization>& RootObjectCustomization, IDetailsViewPrivate* InDetailsView, const FObjectPropertyNode* RootNode);
 private:
-	virtual IDetailsViewPrivate* GetDetailsView() const override{ return DetailsView; }
-	virtual void OnItemExpansionChanged( bool bIsExpanded, bool bShouldSaveState ) override {}
-	virtual bool ShouldBeExpanded() const override { return true; }
+	virtual IDetailsViewPrivate* GetDetailsView() const override { return DetailsView; }
+	virtual void OnItemExpansionChanged(bool bIsExpanded, bool bShouldSaveState) override;
+	virtual bool ShouldBeExpanded() const override;
 	virtual ENodeVisibility GetVisibility() const override;
-	virtual TSharedRef< ITableRow > GenerateWidgetForTableView(const TSharedRef<STableViewBase>& OwnerTable, const FDetailColumnSizeData& ColumnSizeData, bool bAllowFavoriteSystem) override;
+	virtual TSharedRef<ITableRow> GenerateWidgetForTableView(const TSharedRef<STableViewBase>& OwnerTable, const FDetailColumnSizeData& ColumnSizeData, bool bAllowFavoriteSystem) override;
 	virtual bool GenerateStandaloneWidget(FDetailWidgetRow& OutRow) const override;
-	virtual void GetChildren(FDetailNodeList& OutChildren )  override;
-	virtual void FilterNode( const FDetailFilter& InFilter ) override;
-	virtual void Tick( float DeltaTime ) override {}
+	virtual void GetChildren(FDetailNodeList& OutChildren)  override;
+	virtual void FilterNode(const FDetailFilter& InFilter) override;
+	virtual void Tick(float DeltaTime) override {}
 	virtual bool ShouldShowOnlyChildren() const override;
 	virtual FName GetNodeName() const override { return NodeName; }
 	virtual EDetailNodeType GetNodeType() const override { return EDetailNodeType::Object; }
@@ -65,8 +65,10 @@ private:
 	FDetailNodeList ChildNodes;
 	IDetailsViewPrivate* DetailsView;
 	TWeakPtr<IDetailRootObjectCustomization> RootObjectCustomization;
-	const TWeakObjectPtr<UObject> RootObject;
+	FDetailsObjectSet RootObjectSet;
+	const UClass* CommonBaseClass;
 	FName NodeName;
 	bool bShouldBeVisible;
+	bool bHasFilterStrings;
 	bool bShouldShowOnlyChildren;
 };
diff --git a/Engine/Source/Editor/PropertyEditor/Public/IDetailRootObjectCustomization.h b/Engine/Source/Editor/PropertyEditor/Public/IDetailRootObjectCustomization.h
--- a/Engine/Source/Editor/PropertyEditor/Public/IDetailRootObjectCustomization.h
+++ b/Engine/Source/Editor/PropertyEditor/Public/IDetailRootObjectCustomization.h
@@ -6,6 +6,16 @@
 class SWidget;
 class ITableRow;
 
+/** Defines a set of objects and a common base class between them for a root object customization */
+struct FDetailsObjectSet
+{
+	/** List of objects that represent the root in the details panel */
+	TArray<const UObject*> RootObjects;
+
+	/** The most common base class between all root objects */
+	const UClass* CommonBaseClass;
+};
+
 /** 
  * Interface for any class that lays out details for a specific class
  */
@@ -21,35 +31,75 @@ public:
 
 	virtual ~IDetailRootObjectCustomization() {}
 
-	/** 
+	/**
 	 * Called when the details panel wants to display an object header widget for a given object
 	 *
-	 * @param	InRootObject	The object whose header is being customized
- 	 */
-	virtual TSharedPtr<SWidget> CustomizeObjectHeader(const UObject* InRootObject) { return nullptr; }
+	 * @param	InRootObjects	The object set whose header is being customized
+	 */
+	virtual TSharedPtr<SWidget> CustomizeObjectHeader(const FDetailsObjectSet& InRootObjectSet)
+	{ 
+		PRAGMA_DISABLE_DEPRECATION_WARNINGS
+		return CustomizeObjectHeader(InRootObjectSet.RootObjects[0]);
+		PRAGMA_ENABLE_DEPRECATION_WARNINGS
+	}
 
 	/** 
 	 * Called when the details panel wants to display an object header widget for a given object
 	 *
-	 * @param	InRootObject	The object whose header is being customized
+	 * @param	InRootObjects	The object set whose header is being customized
 	 * @param	InTableRow		The ITableRow object (table views to talk to their rows) that will use the current IDetailRootObjectCustomization element.  This may be null if the customization is not being shown in a table view
  	 */
-	virtual TSharedPtr<SWidget> CustomizeObjectHeader(const UObject* InRootObject, const TSharedPtr<ITableRow>& InTableRow)
+	virtual TSharedPtr<SWidget> CustomizeObjectHeader(const FDetailsObjectSet& InRootObjectSet, const TSharedPtr<ITableRow>& InTableRow)
 	{
-		return CustomizeObjectHeader(InRootObject);
+		PRAGMA_DISABLE_DEPRECATION_WARNINGS
+		return CustomizeObjectHeader(InRootObjectSet.RootObjects[0], InTableRow);
+		PRAGMA_ENABLE_DEPRECATION_WARNINGS
 	}
 
 	/**
-	 * Whether or not the object and all of its children should be visible in the details panel
+	 * Whether or not the objects and all of its children should be visible in the details panel
 	 */
-	virtual bool IsObjectVisible(const UObject* InRootObject) const = 0;
+	virtual bool AreObjectsVisible(const FDetailsObjectSet& InRootObjectSet) const
+	{
+		PRAGMA_DISABLE_DEPRECATION_WARNINGS
+		return IsObjectVisible(InRootObjectSet.RootObjects[0]);
+		PRAGMA_ENABLE_DEPRECATION_WARNINGS
+	}
 
 	/**
 	 * Whether or not the object should have a header displayed or just show the children directly
 	 *
 	 * @return true if the this customization should be displayed or false to just show the children directly
 	 */
-	virtual bool ShouldDisplayHeader(const UObject* InRootObject) const = 0;
+	virtual bool ShouldDisplayHeader(const FDetailsObjectSet& InRootObjectSet) const
+	{
+		PRAGMA_DISABLE_DEPRECATION_WARNINGS
+		return ShouldDisplayHeader(InRootObjectSet.RootObjects[0]);
+		PRAGMA_ENABLE_DEPRECATION_WARNINGS
+	}
+
+	UE_DEPRECATED(4.25, "Please use the CustomizeObjectHeader version which takes in an array of root objects")
+	virtual TSharedPtr<SWidget> CustomizeObjectHeader(const UObject* InRootObject)
+	{
+		return nullptr;
+	}
+
+	UE_DEPRECATED(4.25, "Please use the CustomizeObjectHeader version which takes in an array of root objects")
+	virtual TSharedPtr<SWidget> CustomizeObjectHeader(const UObject* InRootObject, const TSharedPtr<ITableRow>& InTableRow)
+	{
+		PRAGMA_DISABLE_DEPRECATION_WARNINGS
+		return CustomizeObjectHeader(InRootObject);
+		PRAGMA_ENABLE_DEPRECATION_WARNINGS
+	}
+
+	UE_DEPRECATED(4.25, "Please use the AreObjectsVisible instead")
+	virtual bool IsObjectVisible(const UObject* InRootObject) const { return true; }
+
+	UE_DEPRECATED(4.25, "Please use the ShouldDisplayHeader version which takes in an array of objects instead")
+	virtual bool ShouldDisplayHeader(const UObject* InRootObject) const
+	{
+		return true;
+	}
 
 	/**
 	 * Gets the setup for expansion arrows in this customization
