Fix non-unity

#rnx

diff --git a/Engine/Source/Editor/PropertyEditor/Private/DetailMultiTopLevelObjectRootNode.cpp b/Engine/Source/Editor/PropertyEditor/Private/DetailMultiTopLevelObjectRootNode.cpp
--- a/Engine/Source/Editor/PropertyEditor/Private/DetailMultiTopLevelObjectRootNode.cpp
+++ b/Engine/Source/Editor/PropertyEditor/Private/DetailMultiTopLevelObjectRootNode.cpp
@@ -4,6 +4,8 @@
 #include "IDetailRootObjectCustomization.h"
 #include "DetailWidgetRow.h"
 #include "ObjectPropertyNode.h"
+#include "Misc/ConfigCacheIni.h"
+
 
 void SDetailMultiTopLevelObjectTableRow::Construct( const FArguments& InArgs, TSharedRef<FDetailTreeNode> InOwnerTreeNode, const TSharedRef<STableViewBase>& InOwnerTableView )
 {
