Allow "NoCategory" to be used on properties to show them at the top of the details panel root node for the object with the "nocategory" property

diff --git a/Engine/Source/Editor/PropertyEditor/Private/DetailCategoryBuilderImpl.cpp b/Engine/Source/Editor/PropertyEditor/Private/DetailCategoryBuilderImpl.cpp
--- a/Engine/Source/Editor/PropertyEditor/Private/DetailCategoryBuilderImpl.cpp
+++ b/Engine/Source/Editor/PropertyEditor/Private/DetailCategoryBuilderImpl.cpp
@@ -122,7 +122,8 @@ FDetailCategoryImpl::FDetailCategoryImpl(FName InCategoryName, TSharedRef<FDetai
 {
 	const UStruct* BaseStruct = InDetailLayout->GetRootNode()->GetBaseStructure();
 
-	bShowOnlyChildren = InDetailLayout->IsLayoutForExternalRoot() && !InDetailLayout->GetRootNode()->HasNodeFlags(EPropertyNodeFlags::ShowCategories);
+	static const FName NoCategoryName = TEXT("NoCategory");
+	bShowOnlyChildren = (InDetailLayout->IsLayoutForExternalRoot() && !InDetailLayout->GetRootNode()->HasNodeFlags(EPropertyNodeFlags::ShowCategories)) || CategoryName == NoCategoryName;
 
 	// Use the base class name if there is one otherwise this is a generic category not specific to a class
 	FName BaseStructName = BaseStruct ? BaseStruct->GetFName() : FName("Generic");
