Fix crash with SPropertyEditorInteractiveActorPicker.
Was deactivating a mode that was not active.

diff --git a/Engine/Source/Editor/PropertyEditor/Private/UserInterface/PropertyEditor/SPropertyEditorInteractiveActorPicker.cpp b/Engine/Source/Editor/PropertyEditor/Private/UserInterface/PropertyEditor/SPropertyEditorInteractiveActorPicker.cpp
--- a/Engine/Source/Editor/PropertyEditor/Private/UserInterface/PropertyEditor/SPropertyEditorInteractiveActorPicker.cpp
+++ b/Engine/Source/Editor/PropertyEditor/Private/UserInterface/PropertyEditor/SPropertyEditorInteractiveActorPicker.cpp
@@ -11,7 +11,10 @@ SPropertyEditorInteractiveActorPicker::~SPropertyEditorInteractiveActorPicker()
 	FActorPickerModeModule& ActorPickerMode = FModuleManager::Get().GetModuleChecked<FActorPickerModeModule>("ActorPickerMode");
 
 	// make sure we are unregistered when this widget goes away
-	ActorPickerMode.EndActorPickingMode();
+	if (ActorPickerMode.IsInActorPickingMode())
+	{
+		ActorPickerMode.EndActorPickingMode();
+	}
 }
 
 void SPropertyEditorInteractiveActorPicker::Construct( const FArguments& InArgs )
