Expose a way to set the root object customization on actor details panels from outside the level editor

diff --git a/Engine/Source/Editor/LevelEditor/Private/SActorDetails.cpp b/Engine/Source/Editor/LevelEditor/Private/SActorDetails.cpp
--- a/Engine/Source/Editor/LevelEditor/Private/SActorDetails.cpp
+++ b/Engine/Source/Editor/LevelEditor/Private/SActorDetails.cpp
@@ -319,9 +319,11 @@ void SActorDetails::PostRedo(bool bSuccess)
 	PostUndo(bSuccess);
 }
 
-void SActorDetails::SetActorDetailsFilter(TSharedPtr<FDetailsViewObjectFilter> InFilter)
+void SActorDetails::SetActorDetailsRootCustomization(TSharedPtr<FDetailsViewObjectFilter> ActorDetailsObjectFilter, TSharedPtr<IDetailRootObjectCustomization> ActorDetailsRootCustomization)
 {
-	DetailsView->SetObjectFilter(InFilter);
+	DetailsView->SetObjectFilter(ActorDetailsObjectFilter);
+	DetailsView->SetRootObjectCustomizationInstance(ActorDetailsRootCustomization);
+	DetailsView->ForceRefresh();
 }
 
 void SActorDetails::OnComponentsEditedInWorld()
diff --git a/Engine/Source/Editor/LevelEditor/Private/SActorDetails.h b/Engine/Source/Editor/LevelEditor/Private/SActorDetails.h
--- a/Engine/Source/Editor/LevelEditor/Private/SActorDetails.h
+++ b/Engine/Source/Editor/LevelEditor/Private/SActorDetails.h
@@ -51,7 +51,7 @@ public:
 	 *
 	 * @param InFilter	The filter to use or nullptr to remove the active filter
 	 */
-	void SetActorDetailsFilter(TSharedPtr<FDetailsViewObjectFilter> InFilter);
+	void SetActorDetailsRootCustomization(TSharedPtr<FDetailsViewObjectFilter> ActorDetailsObjectFilter, TSharedPtr<IDetailRootObjectCustomization> ActorDetailsRootCustomization);
 private:
 	AActor* GetSelectedActorInEditor() const;
 	AActor* GetActorContext() const;
diff --git a/Engine/Source/Editor/LevelEditor/Private/SLevelEditor.cpp b/Engine/Source/Editor/LevelEditor/Private/SLevelEditor.cpp
--- a/Engine/Source/Editor/LevelEditor/Private/SLevelEditor.cpp
+++ b/Engine/Source/Editor/LevelEditor/Private/SLevelEditor.cpp
@@ -1702,13 +1702,13 @@ TSharedRef<SWidget> SLevelEditor::CreateActorDetails( const FName TabIdentifier
 }
 
 
-void SLevelEditor::SetActorDetailsFilter(TSharedPtr<FDetailsViewObjectFilter> ActorDetailsFilter)
+void SLevelEditor::SetActorDetailsRootCustomization(TSharedPtr<FDetailsViewObjectFilter> ActorDetailsObjectFilter, TSharedPtr<IDetailRootObjectCustomization> ActorDetailsRootCustomization)
 {
 	for (TWeakPtr<SActorDetails> Details : AllActorDetailPanels)
 	{
-		if (Details.IsValid())
+		if (TSharedPtr<SActorDetails> DetailsPinned = Details.Pin())
 		{
-			Details.Pin()->SetActorDetailsFilter(ActorDetailsFilter);
+			DetailsPinned->SetActorDetailsRootCustomization(ActorDetailsObjectFilter, ActorDetailsRootCustomization);
 		}
 	}
 }
diff --git a/Engine/Source/Editor/LevelEditor/Private/SLevelEditor.h b/Engine/Source/Editor/LevelEditor/Private/SLevelEditor.h
--- a/Engine/Source/Editor/LevelEditor/Private/SLevelEditor.h
+++ b/Engine/Source/Editor/LevelEditor/Private/SLevelEditor.h
@@ -121,7 +121,7 @@ public:
 	virtual void OnToolkitHostingFinished( const TSharedRef< class IToolkit >& Toolkit ) override;
 	virtual UWorld* GetWorld() const override;
 	virtual TSharedRef<SWidget> CreateActorDetails( const FName TabIdentifier ) override;
-	virtual void SetActorDetailsFilter(TSharedPtr<FDetailsViewObjectFilter> ActorDetailsFilter) override;
+	virtual void SetActorDetailsRootCustomization(TSharedPtr<FDetailsViewObjectFilter> ActorDetailsObjectFilter, TSharedPtr<IDetailRootObjectCustomization> ActorDetailsRootCustomization) override;
 	virtual TSharedRef<SWidget> CreateToolBox() override;
 
 	/** SWidget overrides */
diff --git a/Engine/Source/Editor/LevelEditor/Public/ILevelEditor.h b/Engine/Source/Editor/LevelEditor/Public/ILevelEditor.h
--- a/Engine/Source/Editor/LevelEditor/Public/ILevelEditor.h
+++ b/Engine/Source/Editor/LevelEditor/Public/ILevelEditor.h
@@ -13,6 +13,7 @@
 #include "Toolkits/IToolkitHost.h"
 
 class FDetailsViewObjectFilter;
+class IDetailRootObjectCustomization;
 class ISceneOutliner;
 class IAssetViewport;
 class SLevelViewport;
@@ -60,7 +61,7 @@ public:
 	virtual TSharedRef<SWidget> CreateActorDetails( const FName TabIdentifier ) = 0;
 
 	/** Set the filter that should be used to determine the set of objects that should be shown in a details panel when an actor in the level editor is selected */
-	virtual void SetActorDetailsFilter(TSharedPtr<FDetailsViewObjectFilter> ActorDetailsFilter) = 0;
+	virtual void SetActorDetailsRootCustomization(TSharedPtr<FDetailsViewObjectFilter> ActorDetailsObjectFilter, TSharedPtr<IDetailRootObjectCustomization> ActorDetailsRootCustomization) = 0;
 
 	/** Spawns a level editor ToolBox widget (aka. "Modes") */
 	virtual TSharedRef<SWidget> CreateToolBox() = 0;
diff --git a/Engine/Source/Editor/PropertyEditor/Private/DetailMultiTopLevelObjectRootNode.cpp b/Engine/Source/Editor/PropertyEditor/Private/DetailMultiTopLevelObjectRootNode.cpp
--- a/Engine/Source/Editor/PropertyEditor/Private/DetailMultiTopLevelObjectRootNode.cpp
+++ b/Engine/Source/Editor/PropertyEditor/Private/DetailMultiTopLevelObjectRootNode.cpp
@@ -12,27 +12,7 @@ void SDetailMultiTopLevelObjectTableRow::Construct( const FArguments& InArgs, TS
 	OwnerTreeNode = InOwnerTreeNode;
 	ExpansionArrowUsage = InArgs._ExpansionArrowUsage;
 
-	ChildSlot
-	[	
-		SNew( SBox )
-		.Padding( FMargin( 0.0f, 0.0f, SDetailTableRowBase::ScrollbarPaddingSize, 0.0f ) )
-		[
-			SNew( SHorizontalBox )
-			+SHorizontalBox::Slot()
-			.VAlign(VAlign_Center)
-			.Padding(2.0f, 2.0f, 2.0f, 2.0f)
-			.AutoWidth()
-			[
-				SNew( SExpanderArrow, SharedThis(this) )
-				.Visibility(ExpansionArrowUsage == EExpansionArrowUsage::Default ? EVisibility::Visible : EVisibility::Collapsed)
-			]
-			+SHorizontalBox::Slot()
-			.Expose(ContentSlot)
-			[
-				SNullWidget::NullWidget
-			]
-		]
-	];
+
 
 	STableRow< TSharedPtr< FDetailTreeNode > >::ConstructInternal(
 		STableRow::FArguments()
@@ -45,10 +25,38 @@ void SDetailMultiTopLevelObjectTableRow::Construct( const FArguments& InArgs, TS
 
 void SDetailMultiTopLevelObjectTableRow::SetContent(TSharedRef<SWidget> InContent)
 {
-	(*ContentSlot)
-	[
-		InContent
-	];
+	if (ExpansionArrowUsage == EExpansionArrowUsage::Default)
+	{
+		ChildSlot
+		[
+			SNew(SBox)
+			.Padding(FMargin(0.0f, 0.0f, SDetailTableRowBase::ScrollbarPaddingSize, 0.0f))
+			[
+				SNew(SHorizontalBox)
+				+ SHorizontalBox::Slot()
+				.VAlign(VAlign_Center)
+				.Padding(2.0f, 2.0f, 2.0f, 2.0f)
+				.AutoWidth()
+				[
+					SNew(SExpanderArrow, SharedThis(this))
+					.Visibility(ExpansionArrowUsage == EExpansionArrowUsage::Default ? EVisibility::Visible : EVisibility::Collapsed)
+				]
+				+ SHorizontalBox::Slot()
+				.Expose(ContentSlot)
+				[
+					InContent
+				]
+			]
+		];
+	}
+	else
+	{
+		ChildSlot
+		[
+			InContent
+		];
+	}
+	
 }
 
 const FSlateBrush* SDetailMultiTopLevelObjectTableRow::GetBackgroundImage() const
@@ -216,7 +224,7 @@ bool FDetailMultiTopLevelObjectRootNode::ShouldShowOnlyChildren() const
 
 void FDetailMultiTopLevelObjectRootNode::GenerateWidget_Internal(FDetailWidgetRow& OutRow, TSharedPtr<SDetailMultiTopLevelObjectTableRow> TableRowWidget) const
 {
-	TSharedPtr<SWidget> HeaderWidget;
+	TSharedPtr<SWidget> HeaderWidget = SNullWidget::NullWidget;
 	if (RootObjectCustomization.IsValid() && RootObjectSet.RootObjects.Num())
 	{
 		HeaderWidget = RootObjectCustomization.Pin()->CustomizeObjectHeader(RootObjectSet, TableRowWidget);
diff --git a/Engine/Source/Editor/PropertyEditor/Public/DetailsViewObjectFilter.h b/Engine/Source/Editor/PropertyEditor/Public/DetailsViewObjectFilter.h
--- a/Engine/Source/Editor/PropertyEditor/Public/DetailsViewObjectFilter.h
+++ b/Engine/Source/Editor/PropertyEditor/Public/DetailsViewObjectFilter.h
@@ -12,18 +12,15 @@
  */
 struct FDetailsViewObjectRoot
 {
-	FDetailsViewObjectRoot(bool bInShowRootHeader = true)
-		: bShowRootHeader(bInShowRootHeader)
+	FDetailsViewObjectRoot()
 	{}
 
-	FDetailsViewObjectRoot(UObject* InObject, bool bInShowRootHeader=true)
-		: bShowRootHeader(bInShowRootHeader)
+	FDetailsViewObjectRoot(UObject* InObject)
 	{
 		Objects.Add(InObject);
 	}
 
-	FDetailsViewObjectRoot(const TArray<UObject*>& InObjects, bool bInShowRootHeader=true)
-		: bShowRootHeader(bInShowRootHeader)
+	FDetailsViewObjectRoot(const TArray<UObject*>& InObjects)
 	{
 		Objects.Reserve(InObjects.Num());
 		for (UObject* Object : InObjects)
@@ -33,7 +30,6 @@ struct FDetailsViewObjectRoot
 	}
 
 	TArray<UObject*> Objects;
-	bool bShowRootHeader;
 };
 
 /**
