Integrate CL 11418990 from Dev-Editor

PostEditChange should not be called on unchanged archetype instances

diff --git a/Engine/Source/Editor/PropertyEditor/Private/PropertyHandleImpl.cpp b/Engine/Source/Editor/PropertyEditor/Private/PropertyHandleImpl.cpp
--- a/Engine/Source/Editor/PropertyEditor/Private/PropertyHandleImpl.cpp
+++ b/Engine/Source/Editor/PropertyEditor/Private/PropertyHandleImpl.cpp
@@ -31,6 +31,13 @@
 
 #define LOCTEXT_NAMESPACE "PropertyHandleImplementation"
 
+bool IsTemplate(UObject* Obj)
+{
+	return Obj &&
+		(Obj->HasAnyFlags(RF_ClassDefaultObject | RF_ArchetypeObject) ||
+		(Obj->HasAnyFlags(RF_DefaultSubObject) && Obj->GetOuter()->HasAnyFlags(RF_ClassDefaultObject | RF_ArchetypeObject)));
+}
+
 FPropertyValueImpl::FPropertyValueImpl( TSharedPtr<FPropertyNode> InPropertyNode, FNotifyHook* InNotifyHook, TSharedPtr<IPropertyUtilities> InPropertyUtilities )
 	: PropertyNode( InPropertyNode )
 	, PropertyUtilities( InPropertyUtilities )
@@ -924,13 +931,14 @@ void FPropertyValueImpl::AddChild()
 		{
 			// determines whether we actually changed any values (if the user clicks the "empty" button when the array is already empty,
 			// we don't want the objects to be marked dirty)
-			bool bNotifiedPreChange = false;
+			bool bHasChanges = false;
 
 			// If we added a new Map entry we need to rebuild children
 			bool bAddedMapEntry = false;
 
 			TArray< TMap<FString, int32> > ArrayIndicesPerObject;
-			TArray< TMap<UObject*, bool> > PropagationResultPerObject;
+			TArray< TArray< UObject* > > AffectedInstancesPerObject;
+			AffectedInstancesPerObject.SetNum(ReadAddresses.Num());
 
 			// List of top level objects sent to the PropertyChangedEvent
 			TArray<const UObject*> TopLevelObjects;
@@ -950,80 +958,93 @@ void FPropertyValueImpl::AddChild()
 				void* Addr = ReadAddresses.GetAddress(i);
 				if ( Addr )
 				{
-					if ( !bNotifiedPreChange )
-					{
-						bNotifiedPreChange = true;
+					bHasChanges = true;
 
-						// send the PreEditChange notification to all selected objects
-						PropertyNodePin->NotifyPreChange( NodeProperty, NotifyHook );
+					if (!FApp::IsGame())
+					{
+						UObject* Obj = ObjectNode ? ObjectNode->GetUObject(i) : nullptr;
+						if (IsTemplate(Obj))
+						{
+							PropertyNodePin->GatherInstancesAffectedByContainerPropertyChange(Obj, Addr, EPropertyArrayChangeType::Add, AffectedInstancesPerObject[i]);
+						}
 					}
+				}
+			}
 
-					//add on array index so we can tell which entry just changed
-					ArrayIndicesPerObject.Add(TMap<FString,int32>());
-					FPropertyValueImpl::GenerateArrayIndexMapToObjectNode(ArrayIndicesPerObject[i], PropertyNodePin.Get());
+			if ( bHasChanges )
+			{
+				TSet< UObject* > AllAffectedInstances;
+				for (const TArray<UObject*>& AffectedInstances : AffectedInstancesPerObject)
+				{
+					AllAffectedInstances.Append(AffectedInstances);
+				}
 
-					UObject* Obj = ObjectNode ? ObjectNode->GetUObject(i) : nullptr;
-					if (Obj)
+				// send the PreEditChange notification to all selected objects
+				PropertyNodePin->NotifyPreChange(NodeProperty, NotifyHook, AllAffectedInstances);
+
+				for (int32 i = 0 ; i < ReadAddresses.Num() ; ++i )
+				{
+					void* Addr = ReadAddresses.GetAddress(i);
+					if ( Addr )
 					{
-						if ((Obj->HasAnyFlags(RF_ClassDefaultObject | RF_ArchetypeObject) ||
-							(Obj->HasAnyFlags(RF_DefaultSubObject) && Obj->GetOuter()->HasAnyFlags(RF_ClassDefaultObject | RF_ArchetypeObject))) &&
-							!FApp::IsGame())
+						//add on array index so we can tell which entry just changed
+						ArrayIndicesPerObject.Add(TMap<FString, int32>());
+						FPropertyValueImpl::GenerateArrayIndexMapToObjectNode(ArrayIndicesPerObject[i], PropertyNodePin.Get());
+
+						UObject* Obj = ObjectNode ? ObjectNode->GetUObject(i) : nullptr;
+						if (Obj)
 						{
-							TMap<UObject*, bool> PropagationResult;
-							PropertyNodePin->PropagateContainerPropertyChange(Obj, Addr, EPropertyArrayChangeType::Add, -1, &PropagationResult);
-							PropagationResultPerObject.Add(MoveTemp(PropagationResult));
-						}
+							if (IsTemplate(Obj) && !FApp::IsGame())
+							{
+								PropertyNodePin->PropagateContainerPropertyChange(Obj, Addr, AffectedInstancesPerObject[i], EPropertyArrayChangeType::Add, -1);
+							}
 
-						TopLevelObjects.Add(Obj);
-					}
+							TopLevelObjects.Add(Obj);
+						}
 
-					int32 Index = INDEX_NONE;
+						int32 Index = INDEX_NONE;
 
-					if (Array)
-					{
-						FScriptArrayHelper	ArrayHelper(Array, Addr);
-						Index = ArrayHelper.AddValue();
-					}
-					else if (Set)
-					{
-						FScriptSetHelper	SetHelper(Set, Addr);
-						Index = SetHelper.AddDefaultValue_Invalid_NeedsRehash();
-						SetHelper.Rehash();
+						if (Array)
+						{
+							FScriptArrayHelper	ArrayHelper(Array, Addr);
+							Index = ArrayHelper.AddValue();
+						}
+						else if (Set)
+						{
+							FScriptSetHelper	SetHelper(Set, Addr);
+							Index = SetHelper.AddDefaultValue_Invalid_NeedsRehash();
+							SetHelper.Rehash();
 
-					}
-					else if (Map)
-					{
-						FScriptMapHelper	MapHelper(Map, Addr);
-						Index = MapHelper.AddDefaultValue_Invalid_NeedsRehash();
-						MapHelper.Rehash();
+						}
+						else if (Map)
+						{
+							FScriptMapHelper	MapHelper(Map, Addr);
+							Index = MapHelper.AddDefaultValue_Invalid_NeedsRehash();
+							MapHelper.Rehash();
+							bAddedMapEntry = true;
+						}
 
-						uint8* PairPtr = MapHelper.GetPairPtr(Index);
-						bAddedMapEntry = true;
+						ArrayIndicesPerObject[i].Add(NodeProperty->GetName(), Index);
 					}
-
-					ArrayIndicesPerObject[i].Add(NodeProperty->GetName(), Index);
 				}
-			}
 
-			FPropertyChangedEvent ChangeEvent(NodeProperty, EPropertyChangeType::ArrayAdd, MakeArrayView(TopLevelObjects));
-			ChangeEvent.SetArrayIndexPerObject(ArrayIndicesPerObject);
-			ChangeEvent.SetInstancesChangedResultPerArchetype(PropagationResultPerObject);
+				FPropertyChangedEvent ChangeEvent(NodeProperty, EPropertyChangeType::ArrayAdd, MakeArrayView(TopLevelObjects));
+				ChangeEvent.SetArrayIndexPerObject(ArrayIndicesPerObject);
+				ChangeEvent.SetInstancesChanged(MoveTemp(AllAffectedInstances));
 
-			if ( bNotifiedPreChange )
-			{
 				// send the PostEditChange notification; it will be propagated to all selected objects
 				PropertyNodePin->NotifyPostChange(ChangeEvent, NotifyHook);
-			}
 
-			if (PropertyUtilities.IsValid())
-			{
-				PropertyNodePin->FixPropertiesInEvent(ChangeEvent);
-				PropertyUtilities.Pin()->NotifyFinishedChangingProperties(ChangeEvent);
-			}
+				if (PropertyUtilities.IsValid())
+				{
+					PropertyNodePin->FixPropertiesInEvent(ChangeEvent);
+					PropertyUtilities.Pin()->NotifyFinishedChangingProperties(ChangeEvent);
+				}
 
-			if (bAddedMapEntry)
-			{
-				PropertyNodePin->RebuildChildren();
+				if (bAddedMapEntry)
+				{
+					PropertyNodePin->RebuildChildren();
+				}
 			}
 		}
 	}
@@ -1046,12 +1067,15 @@ void FPropertyValueImpl::ClearChildren()
 		{
 			// determines whether we actually changed any values (if the user clicks the "emtpy" button when the array is already empty,
 			// we don't want the objects to be marked dirty)
-			bool bNotifiedPreChange = false;
+			bool bHasChanges = false;
 
 			// List of top level objects sent to the PropertyChangedEvent
 			TArray<const UObject*> TopLevelObjects;
 			TopLevelObjects.Reserve(ReadAddresses.Num());
 
+			TArray< TArray< UObject* > > AffectedInstancesPerObject;
+			AffectedInstancesPerObject.SetNum(ReadAddresses.Num());
+
 			// Begin a property edit transaction.
 			FScopedTransaction Transaction( NSLOCTEXT("UnrealEd", "ClearChildren", "Clear Children") );
 			FObjectPropertyNode* ObjectNode = PropertyNodePin->FindObjectItemParent();
@@ -1066,120 +1090,138 @@ void FPropertyValueImpl::ClearChildren()
 				void* Addr = ReadAddresses.GetAddress(i);
 				if ( Addr )
 				{
-					if ( !bNotifiedPreChange )
-					{
-						bNotifiedPreChange = true;
+					bHasChanges = true;
 
-						// send the PreEditChange notification to all selected objects
-						PropertyNodePin->NotifyPreChange( NodeProperty, NotifyHook );
-					}
-
-					UObject* Obj = ObjectNode ? ObjectNode->GetUObject(i) : nullptr;
-					if (Obj)
+					if (!FApp::IsGame())
 					{
-						if ((Obj->HasAnyFlags(RF_ClassDefaultObject | RF_ArchetypeObject) ||
-							(Obj->HasAnyFlags(RF_DefaultSubObject) && Obj->GetOuter()->HasAnyFlags(RF_ClassDefaultObject | RF_ArchetypeObject))) &&
-							!FApp::IsGame())
+						UObject* Obj = ObjectNode ? ObjectNode->GetUObject(i) : nullptr;
+						if (IsTemplate(Obj))
 						{
-							PropertyNodePin->PropagateContainerPropertyChange(Obj, Addr, EPropertyArrayChangeType::Clear, -1);
+							PropertyNodePin->GatherInstancesAffectedByContainerPropertyChange(Obj, Addr, EPropertyArrayChangeType::Clear, AffectedInstancesPerObject[i]);
 						}
-
-						TopLevelObjects.Add(Obj);
 					}
+				}
+			}
 
-					if (ArrayProperty)
-					{
-						FScriptArrayHelper ArrayHelper(ArrayProperty, Addr);
+			if (bHasChanges)
+			{
+				TSet< UObject* > AllAffectedInstances;
+				for (const TArray<UObject*>& AffectedInstances : AffectedInstancesPerObject)
+				{
+					AllAffectedInstances.Append(AffectedInstances);
+				}
 
-						// If the inner property is an instanced component property we must move the old components to the 
-						// transient package so resetting owned components on the parent doesn't find them
-						FObjectProperty* InnerObjectProperty = CastField<FObjectProperty>(ArrayProperty->Inner);
-						if (InnerObjectProperty && InnerObjectProperty->HasAnyPropertyFlags(CPF_InstancedReference) && InnerObjectProperty->PropertyClass->IsChildOf(UActorComponent::StaticClass()))
+				// Send the PreEditChange notification to all selected objects
+				PropertyNodePin->NotifyPreChange(NodeProperty, NotifyHook, AllAffectedInstances);
+
+				for ( int32 i = 0 ; i < ReadAddresses.Num() ; ++i )
+				{
+					void* Addr = ReadAddresses.GetAddress(i);
+					if ( Addr )
+					{
+						UObject* Obj = ObjectNode ? ObjectNode->GetUObject(i) : nullptr;
+						if (Obj)
 						{
-							const int32 ArraySize = ArrayHelper.Num();
-							for (int32 Index = 0; Index < ArraySize; ++Index)
+							if (IsTemplate(Obj) && !FApp::IsGame())
 							{
-								if (UActorComponent* Component = *reinterpret_cast<UActorComponent**>(ArrayHelper.GetRawPtr(Index)))
-								{
-									Component->Modify();
-									Component->Rename(nullptr, GetTransientPackage(), REN_DontCreateRedirectors);
-								}
+								PropertyNodePin->PropagateContainerPropertyChange(Obj, Addr, AffectedInstancesPerObject[i], EPropertyArrayChangeType::Clear, -1);
 							}
-						}
 
-						ArrayHelper.EmptyValues();
-					}
-					else if (SetProperty)
-					{
-						FScriptSetHelper SetHelper(SetProperty, Addr);
+							TopLevelObjects.Add(Obj);
+						}
 
-						// If the element property is an instanced component property we must move the old components to the 
-						// transient package so resetting owned components on the parent doesn't find them
-						FObjectProperty* ElementObjectProperty = CastField<FObjectProperty>(SetProperty->ElementProp);
-						if (ElementObjectProperty && ElementObjectProperty->HasAnyPropertyFlags(CPF_InstancedReference) && ElementObjectProperty->PropertyClass->IsChildOf(UActorComponent::StaticClass()))
+						if (ArrayProperty)
 						{
-							int32 ElementsToRemove = SetHelper.Num();
-							int32 Index = 0;
-							while (ElementsToRemove > 0)
+							FScriptArrayHelper ArrayHelper(ArrayProperty, Addr);
+
+							// If the inner property is an instanced component property we must move the old components to the 
+							// transient package so resetting owned components on the parent doesn't find them
+							FObjectProperty* InnerObjectProperty = CastField<FObjectProperty>(ArrayProperty->Inner);
+							if (InnerObjectProperty && InnerObjectProperty->HasAnyPropertyFlags(CPF_InstancedReference) && InnerObjectProperty->PropertyClass->IsChildOf(UActorComponent::StaticClass()))
 							{
-								if (SetHelper.IsValidIndex(Index))
+								const int32 ArraySize = ArrayHelper.Num();
+								for (int32 Index = 0; Index < ArraySize; ++Index)
 								{
-									if (UActorComponent* Component = *reinterpret_cast<UActorComponent**>(SetHelper.GetElementPtr(Index)))
+									if (UActorComponent* Component = *reinterpret_cast<UActorComponent**>(ArrayHelper.GetRawPtr(Index)))
 									{
 										Component->Modify();
 										Component->Rename(nullptr, GetTransientPackage(), REN_DontCreateRedirectors);
 									}
-									--ElementsToRemove;
 								}
-								++Index;
 							}
+
+							ArrayHelper.EmptyValues();
 						}
+						else if (SetProperty)
+						{
+							FScriptSetHelper SetHelper(SetProperty, Addr);
 
-						SetHelper.EmptyElements();
-					}
-					else if (MapProperty)
-					{
-						FScriptMapHelper MapHelper(MapProperty, Addr);
+							// If the element property is an instanced component property we must move the old components to the 
+							// transient package so resetting owned components on the parent doesn't find them
+							FObjectProperty* ElementObjectProperty = CastField<FObjectProperty>(SetProperty->ElementProp);
+							if (ElementObjectProperty && ElementObjectProperty->HasAnyPropertyFlags(CPF_InstancedReference) && ElementObjectProperty->PropertyClass->IsChildOf(UActorComponent::StaticClass()))
+							{
+								int32 ElementsToRemove = SetHelper.Num();
+								int32 Index = 0;
+								while (ElementsToRemove > 0)
+								{
+									if (SetHelper.IsValidIndex(Index))
+									{
+										if (UActorComponent* Component = *reinterpret_cast<UActorComponent**>(SetHelper.GetElementPtr(Index)))
+										{
+											Component->Modify();
+											Component->Rename(nullptr, GetTransientPackage(), REN_DontCreateRedirectors);
+										}
+										--ElementsToRemove;
+									}
+									++Index;
+								}
+							}
 
-						// If the map's value property is an instanced component property we must move the old components to the 
-						// transient package so resetting owned components on the parent doesn't find them
-						FObjectProperty* ValueObjectProperty = CastField<FObjectProperty>(MapProperty->ValueProp);
-						if (ValueObjectProperty && ValueObjectProperty->HasAnyPropertyFlags(CPF_InstancedReference) && ValueObjectProperty->PropertyClass->IsChildOf(UActorComponent::StaticClass()))
+							SetHelper.EmptyElements();
+						}
+						else if (MapProperty)
 						{
-							int32 ElementsToRemove = MapHelper.Num();
-							int32 Index = 0;
-							while (ElementsToRemove > 0)
+							FScriptMapHelper MapHelper(MapProperty, Addr);
+
+							// If the map's value property is an instanced component property we must move the old components to the 
+							// transient package so resetting owned components on the parent doesn't find them
+							FObjectProperty* ValueObjectProperty = CastField<FObjectProperty>(MapProperty->ValueProp);
+							if (ValueObjectProperty && ValueObjectProperty->HasAnyPropertyFlags(CPF_InstancedReference) && ValueObjectProperty->PropertyClass->IsChildOf(UActorComponent::StaticClass()))
 							{
-								if (MapHelper.IsValidIndex(Index))
+								int32 ElementsToRemove = MapHelper.Num();
+								int32 Index = 0;
+								while (ElementsToRemove > 0)
 								{
-									if (UActorComponent* Component = *reinterpret_cast<UActorComponent**>(MapHelper.GetValuePtr(Index)))
+									if (MapHelper.IsValidIndex(Index))
 									{
-										Component->Modify();
-										Component->Rename(nullptr, GetTransientPackage(), REN_DontCreateRedirectors);
+										if (UActorComponent* Component = *reinterpret_cast<UActorComponent**>(MapHelper.GetValuePtr(Index)))
+										{
+											Component->Modify();
+											Component->Rename(nullptr, GetTransientPackage(), REN_DontCreateRedirectors);
+										}
+										--ElementsToRemove;
 									}
-									--ElementsToRemove;
+									++Index;
 								}
-								++Index;
 							}
-						}
 
-						MapHelper.EmptyValues();
+							MapHelper.EmptyValues();
+						}
 					}
 				}
-			}
 
-			FPropertyChangedEvent ChangeEvent(NodeProperty, EPropertyChangeType::ArrayClear, MakeArrayView(TopLevelObjects));
+				FPropertyChangedEvent ChangeEvent(NodeProperty, EPropertyChangeType::ArrayClear, MakeArrayView(TopLevelObjects));
+				ChangeEvent.SetInstancesChanged(MoveTemp(AllAffectedInstances));
 
-			if ( bNotifiedPreChange )
-			{
-				// send the PostEditChange notification; it will be propagated to all selected objects
+				// Send the PostEditChange notification; it will be propagated to all selected objects
 				PropertyNodePin->NotifyPostChange( ChangeEvent, NotifyHook );
-			}
 
-			if (PropertyUtilities.IsValid())
-			{
-				PropertyNodePin->FixPropertiesInEvent(ChangeEvent);
-				PropertyUtilities.Pin()->NotifyFinishedChangingProperties(ChangeEvent);
+				if (PropertyUtilities.IsValid())
+				{
+					PropertyNodePin->FixPropertiesInEvent(ChangeEvent);
+					PropertyUtilities.Pin()->NotifyFinishedChangingProperties(ChangeEvent);
+				}
 			}
 		}
 	}
@@ -1214,29 +1256,37 @@ void FPropertyValueImpl::InsertChild( TSharedPtr<FPropertyNode> ChildNodeToInser
 
 	if( Addr )
 	{
+		TArray< UObject* > AffectedInstances;
+
 		// Begin a property edit transaction.
 		FScopedTransaction Transaction( NSLOCTEXT("UnrealEd", "InsertChild", "Insert Child") );
 
-		ChildNodePtr->NotifyPreChange( ParentNode->GetProperty(), NotifyHook );
+		if (!FApp::IsGame())
+		{
+			UObject* Obj = ObjectNode ? ObjectNode->GetUObject(0) : nullptr;
+			if (IsTemplate(Obj))
+			{
+				ChildNodePtr->GatherInstancesAffectedByContainerPropertyChange(Obj, Addr, EPropertyArrayChangeType::Add, AffectedInstances);
+			}
+		}
+
+		TSet< UObject* > AllAffectedInstances;
+		AllAffectedInstances.Append(AffectedInstances);
+
+		ChildNodePtr->NotifyPreChange( ParentNode->GetProperty(), NotifyHook, AllAffectedInstances );
 
 		FScriptArrayHelper	ArrayHelper(ArrayProperty,Addr);
 		int32 Index = ChildNodePtr->GetArrayIndex();
 
-		TArray< TMap<UObject*, bool> > PropagationResultPerObject;
-
 		// List of top level objects sent to the PropertyChangedEvent
 		TArray<const UObject*> TopLevelObjects;
 
 		UObject* Obj = ObjectNode ? ObjectNode->GetUObject(0) : nullptr;
 		if (Obj)
 		{
-			if ((Obj->HasAnyFlags(RF_ClassDefaultObject | RF_ArchetypeObject) ||
-				(Obj->HasAnyFlags(RF_DefaultSubObject) && Obj->GetOuter()->HasAnyFlags(RF_ClassDefaultObject | RF_ArchetypeObject))) &&
-				!FApp::IsGame())
+			if (IsTemplate(Obj) && !FApp::IsGame())
 			{
-				TMap<UObject*, bool> PropagationResult;
-				ChildNodePtr->PropagateContainerPropertyChange(Obj, Addr, EPropertyArrayChangeType::Insert, Index, &PropagationResult);
-				PropagationResultPerObject.Add(MoveTemp(PropagationResult));
+				ChildNodePtr->PropagateContainerPropertyChange(Obj, Addr, AffectedInstances, EPropertyArrayChangeType::Insert, Index);
 			}
 
 			TopLevelObjects.Add(Obj);
@@ -1255,7 +1305,7 @@ void FPropertyValueImpl::InsertChild( TSharedPtr<FPropertyNode> ChildNodeToInser
 
 		FPropertyChangedEvent ChangeEvent(ParentNode->GetProperty(), EPropertyChangeType::ArrayAdd, MakeArrayView(TopLevelObjects));
 		ChangeEvent.SetArrayIndexPerObject(ArrayIndicesPerObject);
-		ChangeEvent.SetInstancesChangedResultPerArchetype(PropagationResultPerObject);
+		ChangeEvent.SetInstancesChanged(MoveTemp(AllAffectedInstances));
 
 		PropertyNode.Pin()->NotifyPostChange(ChangeEvent, NotifyHook);
 
@@ -1292,7 +1342,6 @@ void FPropertyValueImpl::DeleteChild( TSharedPtr<FPropertyNode> ChildNodeToDelet
 	FMapProperty* MapProperty = NodeProperty->GetOwner<FMapProperty>();
 
 	TArray< TMap<FString, int32> > ArrayIndicesPerObject;
-	TArray< TMap<UObject*, bool> > PropagationResultPerObject;
 
 	check(ArrayProperty || SetProperty || MapProperty);
 
@@ -1300,119 +1349,150 @@ void FPropertyValueImpl::DeleteChild( TSharedPtr<FPropertyNode> ChildNodeToDelet
 	ParentNode->GetReadAddress( !!ParentNode->HasNodeFlags(EPropertyNodeFlags::SingleSelectOnly), ReadAddresses ); 
 	if ( ReadAddresses.Num() )
 	{
-		FScopedTransaction Transaction( NSLOCTEXT("UnrealEd", "DeleteChild", "Delete Child") );
+		// determines whether we actually changed any values (if the user clicks the "empty" button when the array is already empty,
+		// we don't want the objects to be marked dirty)
+		bool bHasChanges = false;
 
-		ChildNodePtr->NotifyPreChange( NodeProperty, NotifyHook );
+		TArray< TArray< UObject* > > AffectedInstancesPerObject;
+		AffectedInstancesPerObject.SetNum(ReadAddresses.Num());
+
+		FScopedTransaction Transaction(NSLOCTEXT("UnrealEd", "DeleteChild", "Delete Child"));
 
 		// List of top level objects sent to the PropertyChangedEvent
 		TArray<const UObject*> TopLevelObjects;
 		TopLevelObjects.Reserve(ReadAddresses.Num());
 
-		// perform the operation on the array for all selected objects
 		for ( int32 i = 0 ; i < ReadAddresses.Num() ; ++i )
 		{
 			uint8* Address = ReadAddresses.GetAddress(i);
 
 			if( Address ) 
 			{
-				int32 Index = ChildNodePtr->GetArrayIndex();
-
-				//add on array index so we can tell which entry just changed
-				ArrayIndicesPerObject.Add(TMap<FString, int32>());
-				FPropertyValueImpl::GenerateArrayIndexMapToObjectNode(ArrayIndicesPerObject[i], ChildNodePtr);
+				bHasChanges = true;
 
-				UObject* Obj = ObjectNode ? ObjectNode->GetUObject(i) : nullptr;
-				if (Obj)
+				if (!FApp::IsGame())
 				{
-					if ((Obj->HasAnyFlags(RF_ClassDefaultObject | RF_ArchetypeObject) ||
-						(Obj->HasAnyFlags(RF_DefaultSubObject) && Obj->GetOuter()->HasAnyFlags(RF_ClassDefaultObject | RF_ArchetypeObject))) &&
-						!FApp::IsGame())
+					UObject* Obj = ObjectNode ? ObjectNode->GetUObject(i) : nullptr;
+					if (IsTemplate(Obj))
 					{
-						TMap<UObject*, bool> PropagationResult;
-						ChildNodePtr->PropagateContainerPropertyChange(Obj, Address, EPropertyArrayChangeType::Delete, Index, &PropagationResult);
-
-						PropagationResultPerObject.Add(MoveTemp(PropagationResult));
+						ChildNodePtr->GatherInstancesAffectedByContainerPropertyChange(Obj, Address, EPropertyArrayChangeType::Delete, AffectedInstancesPerObject[i]);
 					}
-
-					TopLevelObjects.Add(Obj);
 				}
+			}
+		}
 
-				if (ArrayProperty)
+		if (bHasChanges)
+		{
+			TSet< UObject* > AllAffectedInstances;
+			for (const TArray<UObject*>& AffectedInstances : AffectedInstancesPerObject)
+			{
+				AllAffectedInstances.Append(AffectedInstances);
+			}
+
+			// Send the PreEditChange notification to all selected objects
+			ChildNodePtr->NotifyPreChange(NodeProperty, NotifyHook, AllAffectedInstances);
+
+			// Perform the operation on the array for all selected objects
+			for (int32 i = 0; i < ReadAddresses.Num(); ++i)
+			{
+				uint8* Address = ReadAddresses.GetAddress(i);
+
+				if ( Address )
 				{
-					FScriptArrayHelper ArrayHelper(ArrayProperty, Address);
+					int32 Index = ChildNodePtr->GetArrayIndex();
 
-					// If the inner property is an instanced component property we must move the old component to the 
-					// transient package so resetting owned components on the parent doesn't find it
-					FObjectProperty* InnerObjectProperty = CastField<FObjectProperty>(ArrayProperty->Inner);
-					if (InnerObjectProperty && InnerObjectProperty->HasAnyPropertyFlags(CPF_InstancedReference) && InnerObjectProperty->PropertyClass->IsChildOf(UActorComponent::StaticClass()))
+					// Add on array index so we can tell which entry just changed
+					ArrayIndicesPerObject.Add(TMap<FString, int32>());
+					FPropertyValueImpl::GenerateArrayIndexMapToObjectNode(ArrayIndicesPerObject[i], ChildNodePtr);
+
+					UObject* Obj = ObjectNode ? ObjectNode->GetUObject(i) : nullptr;
+					if (Obj)
 					{
-						if (UActorComponent* Component = *reinterpret_cast<UActorComponent**>(ArrayHelper.GetRawPtr(ChildNodePtr->GetArrayIndex())))
+						if (IsTemplate(Obj) && !FApp::IsGame())
 						{
-							Component->Modify();
-							Component->Rename(nullptr, GetTransientPackage(), REN_DontCreateRedirectors);
+							ChildNodePtr->PropagateContainerPropertyChange(Obj, Address, AffectedInstancesPerObject[i], EPropertyArrayChangeType::Delete, Index);
 						}
-					}
 
-					ArrayHelper.RemoveValues(ChildNodePtr->GetArrayIndex());
-				}
-				else if (SetProperty)
-				{
-					FScriptSetHelper SetHelper(SetProperty, Address);
+						TopLevelObjects.Add(Obj);
+					}
 
-					// If the element property is an instanced component property we must move the old component to the 
-					// transient package so resetting owned components on the parent doesn't find it
-					FObjectProperty* ElementObjectProperty = CastField<FObjectProperty>(SetProperty->ElementProp);
-					if (ElementObjectProperty && ElementObjectProperty->HasAnyPropertyFlags(CPF_InstancedReference) && ElementObjectProperty->PropertyClass->IsChildOf(UActorComponent::StaticClass()))
+					if (ArrayProperty)
 					{
-						if (UActorComponent* Component = *reinterpret_cast<UActorComponent**>(SetHelper.GetElementPtr(ChildNodePtr->GetArrayIndex())))
+						FScriptArrayHelper ArrayHelper(ArrayProperty, Address);
+
+						// If the inner property is an instanced component property we must move the old component to the 
+						// transient package so resetting owned components on the parent doesn't find it
+						FObjectProperty* InnerObjectProperty = CastField<FObjectProperty>(ArrayProperty->Inner);
+						if (InnerObjectProperty && InnerObjectProperty->HasAnyPropertyFlags(CPF_InstancedReference) && InnerObjectProperty->PropertyClass->IsChildOf(UActorComponent::StaticClass()))
 						{
-							Component->Modify();
-							Component->Rename(nullptr, GetTransientPackage(), REN_DontCreateRedirectors);
+							if (UActorComponent* Component = *reinterpret_cast<UActorComponent**>(ArrayHelper.GetRawPtr(ChildNodePtr->GetArrayIndex())))
+							{
+								Component->Modify();
+								Component->Rename(nullptr, GetTransientPackage(), REN_DontCreateRedirectors);
+							}
 						}
+
+						ArrayHelper.RemoveValues(ChildNodePtr->GetArrayIndex());
 					}
+					else if (SetProperty)
+					{
+						FScriptSetHelper SetHelper(SetProperty, Address);
 
-					int32 InternalIndex = SetHelper.FindInternalIndex(ChildNodePtr->GetArrayIndex());
-					SetHelper.RemoveAt(InternalIndex);
-					SetHelper.Rehash();
-				}
-				else if (MapProperty)
-				{
-					FScriptMapHelper MapHelper(MapProperty, Address);
+						// If the element property is an instanced component property we must move the old component to the 
+						// transient package so resetting owned components on the parent doesn't find it
+						FObjectProperty* ElementObjectProperty = CastField<FObjectProperty>(SetProperty->ElementProp);
+						if (ElementObjectProperty && ElementObjectProperty->HasAnyPropertyFlags(CPF_InstancedReference) && ElementObjectProperty->PropertyClass->IsChildOf(UActorComponent::StaticClass()))
+						{
+							if (UActorComponent* Component = *reinterpret_cast<UActorComponent**>(SetHelper.GetElementPtr(ChildNodePtr->GetArrayIndex())))
+							{
+								Component->Modify();
+								Component->Rename(nullptr, GetTransientPackage(), REN_DontCreateRedirectors);
+							}
+						}
 
-					// If the map's value property is an instanced component property we must move the old component to the 
-					// transient package so resetting owned components on the parent doesn't find it
-					FObjectProperty* ValueObjectProperty = CastField<FObjectProperty>(MapProperty->ValueProp);
-					if (ValueObjectProperty && ValueObjectProperty->HasAnyPropertyFlags(CPF_InstancedReference) && ValueObjectProperty->PropertyClass->IsChildOf(UActorComponent::StaticClass()))
+						int32 InternalIndex = SetHelper.FindInternalIndex(ChildNodePtr->GetArrayIndex());
+						SetHelper.RemoveAt(InternalIndex);
+						SetHelper.Rehash();
+					}
+					else if (MapProperty)
 					{
-						if (UActorComponent* Component = *reinterpret_cast<UActorComponent**>(MapHelper.GetValuePtr(ChildNodePtr->GetArrayIndex())))
+						FScriptMapHelper MapHelper(MapProperty, Address);
+
+						// If the map's value property is an instanced component property we must move the old component to the 
+						// transient package so resetting owned components on the parent doesn't find it
+						FObjectProperty* ValueObjectProperty = CastField<FObjectProperty>(MapProperty->ValueProp);
+						if (ValueObjectProperty && ValueObjectProperty->HasAnyPropertyFlags(CPF_InstancedReference) && ValueObjectProperty->PropertyClass->IsChildOf(UActorComponent::StaticClass()))
 						{
-							Component->Modify();
-							Component->Rename(nullptr, GetTransientPackage(), REN_DontCreateRedirectors);
+							if (UActorComponent* Component = *reinterpret_cast<UActorComponent**>(MapHelper.GetValuePtr(ChildNodePtr->GetArrayIndex())))
+							{
+								Component->Modify();
+								Component->Rename(nullptr, GetTransientPackage(), REN_DontCreateRedirectors);
+							}
 						}
+
+						int32 InternalIndex = MapHelper.FindInternalIndex(ChildNodePtr->GetArrayIndex());
+						MapHelper.RemoveAt(InternalIndex);
+						MapHelper.Rehash();
 					}
 
-					int32 InternalIndex = MapHelper.FindInternalIndex(ChildNodePtr->GetArrayIndex());
-					MapHelper.RemoveAt(InternalIndex);
-					MapHelper.Rehash();
+					ArrayIndicesPerObject[i].Add(NodeProperty->GetName(), Index);
 				}
-
-				ArrayIndicesPerObject[i].Add(NodeProperty->GetName(), Index);
 			}
-		}
 
-		FPropertyChangedEvent ChangeEvent(ParentNode->GetProperty(), EPropertyChangeType::ArrayRemove, MakeArrayView(TopLevelObjects));
-		ChangeEvent.SetArrayIndexPerObject(ArrayIndicesPerObject);
-		ChangeEvent.SetInstancesChangedResultPerArchetype(PropagationResultPerObject);
+			FPropertyChangedEvent ChangeEvent(ParentNode->GetProperty(), EPropertyChangeType::ArrayRemove, MakeArrayView(TopLevelObjects));
+			ChangeEvent.SetArrayIndexPerObject(ArrayIndicesPerObject);
+			ChangeEvent.SetInstancesChanged(MoveTemp(AllAffectedInstances));
 
-		PropertyNode.Pin()->NotifyPostChange(ChangeEvent, NotifyHook);
+			PropertyNode.Pin()->NotifyPostChange(ChangeEvent, NotifyHook);
 
-		if (PropertyUtilities.IsValid())
-		{
-			ChildNodePtr->FixPropertiesInEvent(ChangeEvent);
-			PropertyUtilities.Pin()->NotifyFinishedChangingProperties(ChangeEvent);
-		}
+			if (PropertyUtilities.IsValid())
+			{
+				ChildNodePtr->FixPropertiesInEvent(ChangeEvent);
+				PropertyUtilities.Pin()->NotifyFinishedChangingProperties(ChangeEvent);
+			}
 
-		PropertyNode.Pin()->RebuildChildren();
+			PropertyNode.Pin()->RebuildChildren();
+		}
 	}
 }
 
@@ -1468,7 +1548,7 @@ void FPropertyValueImpl::SwapChildren( TSharedPtr<FPropertyNode> FirstChildNode,
 					if ((Obj->HasAnyFlags(RF_ClassDefaultObject | RF_ArchetypeObject) || (Obj->HasAnyFlags(RF_DefaultSubObject) && Obj->GetOuter()->HasAnyFlags(RF_ClassDefaultObject | RF_ArchetypeObject))) 
 						&& !FApp::IsGame())
 					{
-						FirstChildNodePtr->PropagateContainerPropertyChange(Obj, Address, EPropertyArrayChangeType::Swap, FirstIndex, nullptr, SecondIndex);
+						FirstChildNodePtr->PropagateContainerPropertyChange(Obj, Address, EPropertyArrayChangeType::Swap, FirstIndex, SecondIndex);
 					}
 
 					TopLevelObjects.Add(Obj);
@@ -1691,7 +1771,7 @@ void FPropertyValueImpl::MoveElementTo(int32 OriginalIndex, int32 NewIndex)
 						if ((Obj->HasAnyFlags(RF_ClassDefaultObject | RF_ArchetypeObject) ||
 							(Obj->HasAnyFlags(RF_DefaultSubObject) && Obj->GetOuter()->HasAnyFlags(RF_ClassDefaultObject | RF_ArchetypeObject))))
 						{
-							FirstChildNodePtr->PropagateContainerPropertyChange(Obj, Address, EPropertyArrayChangeType::Swap, FirstIndex, nullptr, SecondIndex);
+							FirstChildNodePtr->PropagateContainerPropertyChange(Obj, Address, EPropertyArrayChangeType::Swap, FirstIndex, SecondIndex);
 						}
 
 						TopLevelObjects.Add(Obj);
@@ -1831,28 +1911,37 @@ void FPropertyValueImpl::DuplicateChild( TSharedPtr<FPropertyNode> ChildNodeToDu
 
 	if( Addr )
 	{
+		TArray< UObject* > AffectedInstances;
+
 		// List of top level objects sent to the PropertyChangedEvent
 		TArray<const UObject*> TopLevelObjects;
 
 		FScopedTransaction Transaction( NSLOCTEXT("UnrealEd", "DuplicateChild", "Duplicate Child") );
 
-		ChildNodePtr->NotifyPreChange( ParentNode->GetProperty(), NotifyHook );
+		if (!FApp::IsGame())
+		{
+			UObject* Obj = ObjectNode ? ObjectNode->GetUObject(0) : nullptr;
+			if (IsTemplate(Obj))
+			{
+				ChildNodePtr->GatherInstancesAffectedByContainerPropertyChange(Obj, Addr, EPropertyArrayChangeType::Duplicate, AffectedInstances);
+			}
+		}
+
+		TSet< UObject* > AllAffectedInstances;
+		AllAffectedInstances.Append(AffectedInstances);
+
+		ChildNodePtr->NotifyPreChange( ParentNode->GetProperty(), NotifyHook, AllAffectedInstances );
 
 		int32 Index = ChildNodePtr->GetArrayIndex();
 		UObject* Obj = ObjectNode ? ObjectNode->GetUObject(0) : nullptr;
 
 		TArray< TMap<FString, int32> > ArrayIndicesPerObject;
-		TArray< TMap<UObject*, bool> > PropagationResultPerObject;
 
 		if (Obj)
 		{
-			if ((Obj->HasAnyFlags(RF_ClassDefaultObject | RF_ArchetypeObject) ||
-				(Obj->HasAnyFlags(RF_DefaultSubObject) && Obj->GetOuter()->HasAnyFlags(RF_ClassDefaultObject | RF_ArchetypeObject))) &&
-				!FApp::IsGame())
+			if (IsTemplate(Obj) && !FApp::IsGame())
 			{
-				TMap<UObject*, bool> PropagationResult;
-				ChildNodePtr->PropagateContainerPropertyChange(Obj, Addr, EPropertyArrayChangeType::Duplicate, Index, &PropagationResult);
-				PropagationResultPerObject.Add(MoveTemp(PropagationResult));
+				ChildNodePtr->PropagateContainerPropertyChange(Obj, Addr, AffectedInstances, EPropertyArrayChangeType::Duplicate, Index);
 			}
 
 			TopLevelObjects.Add(Obj);
@@ -1905,7 +1994,7 @@ void FPropertyValueImpl::DuplicateChild( TSharedPtr<FPropertyNode> ChildNodeToDu
 
 		FPropertyChangedEvent ChangeEvent(ParentNode->GetProperty(), EPropertyChangeType::Duplicate, MakeArrayView(TopLevelObjects));
 		ChangeEvent.SetArrayIndexPerObject(ArrayIndicesPerObject);
-		ChangeEvent.SetInstancesChangedResultPerArchetype(PropagationResultPerObject);
+		ChangeEvent.SetInstancesChanged(MoveTemp(AllAffectedInstances));
 
 		PropertyNode.Pin()->NotifyPostChange(ChangeEvent, NotifyHook);
 
diff --git a/Engine/Source/Editor/PropertyEditor/Private/PropertyNode.cpp b/Engine/Source/Editor/PropertyEditor/Private/PropertyNode.cpp
--- a/Engine/Source/Editor/PropertyEditor/Private/PropertyNode.cpp
+++ b/Engine/Source/Editor/PropertyEditor/Private/PropertyNode.cpp
@@ -2182,7 +2182,23 @@ void FPropertyNode::ProcessSeenFlagsForFavorites(void)
 void FPropertyNode::NotifyPreChange( FProperty* PropertyAboutToChange, FNotifyHook* InNotifyHook )
 {
 	TSharedRef<FEditPropertyChain> PropertyChain = BuildPropertyChain( PropertyAboutToChange );
+	NotifyPreChangeInternal(PropertyChain, PropertyAboutToChange, InNotifyHook);
+}
+
+void FPropertyNode::NotifyPreChange( FProperty* PropertyAboutToChange, FNotifyHook* InNotifyHook, const TSet<UObject*>& AffectedInstances )
+{
+	TSharedRef<FEditPropertyChain> PropertyChain = BuildPropertyChain( PropertyAboutToChange, AffectedInstances );
+	NotifyPreChangeInternal(PropertyChain, PropertyAboutToChange, InNotifyHook);
+}
+
+void FPropertyNode::NotifyPreChange( FProperty* PropertyAboutToChange, FNotifyHook* InNotifyHook, TSet<UObject*>&& AffectedInstances )
+{
+	TSharedRef<FEditPropertyChain> PropertyChain = BuildPropertyChain( PropertyAboutToChange, MoveTemp(AffectedInstances) );
+	NotifyPreChangeInternal(PropertyChain, PropertyAboutToChange, InNotifyHook);
+}
 
+void FPropertyNode::NotifyPreChangeInternal( TSharedRef<FEditPropertyChain> PropertyChain, FProperty* PropertyAboutToChange, FNotifyHook* InNotifyHook )
+{
 	// Call through to the property window's notify hook.
 	if( InNotifyHook )
 	{
@@ -2524,6 +2540,20 @@ TSharedRef<FEditPropertyChain> FPropertyNode::BuildPropertyChain( FProperty* InP
 	return PropertyChain;
 }
 
+TSharedRef<FEditPropertyChain> FPropertyNode::BuildPropertyChain( FProperty* InProperty, const TSet<UObject*>& InAffectedArchetypeInstances )
+{
+	TSharedRef<FEditPropertyChain> PropertyChain = BuildPropertyChain(InProperty);
+	PropertyChain->SetAffectedArchetypeInstances(InAffectedArchetypeInstances);
+	return PropertyChain;
+}
+
+TSharedRef<FEditPropertyChain> FPropertyNode::BuildPropertyChain( FProperty* InProperty, TSet<UObject*>&& InAffectedArchetypeInstances )
+{
+	TSharedRef<FEditPropertyChain> PropertyChain = BuildPropertyChain(InProperty);
+	PropertyChain->SetAffectedArchetypeInstances(MoveTemp(InAffectedArchetypeInstances));
+	return PropertyChain;
+}
+
 FPropertyChangedEvent& FPropertyNode::FixPropertiesInEvent(FPropertyChangedEvent& Event)
 {
 	ensure(Event.Property);
@@ -2601,14 +2631,18 @@ bool FPropertyNode::IsFilterAcceptable(const TArray<FString>& InAcceptableNames,
 	return bCompleteMatchFound;
 }
 
-void FPropertyNode::PropagateContainerPropertyChange( UObject* ModifiedObject, const void* OriginalContainerAddr, EPropertyArrayChangeType::Type ChangeType, int32 Index, TMap<UObject*, bool>* PropagationResult, int32 SwapIndex /*= INDEX_NONE*/)
+void FPropertyNode::PropagateContainerPropertyChange( UObject* ModifiedObject, const void* OriginalContainerAddr, EPropertyArrayChangeType::Type ChangeType, int32 Index, int32 SwapIndex /*= INDEX_NONE*/)
+{
+	TArray<UObject*> AffectedInstances;
+	GatherInstancesAffectedByContainerPropertyChange(ModifiedObject, OriginalContainerAddr, ChangeType, AffectedInstances);
+	PropagateContainerPropertyChange(ModifiedObject, OriginalContainerAddr, AffectedInstances, ChangeType, Index, SwapIndex);
+}
+
+void FPropertyNode::GatherInstancesAffectedByContainerPropertyChange(UObject* ModifiedObject, const void* OriginalContainerAddr, EPropertyArrayChangeType::Type ChangeType, TArray<UObject*>& OutAffectedInstances)
 {
 	check(OriginalContainerAddr);
 
 	FProperty* NodeProperty = GetProperty();
-	FArrayProperty* ArrayProperty = NULL;
-	FSetProperty* SetProperty = NULL;
-	FMapProperty* MapProperty = NULL;
 
 	FPropertyNode* ParentPropertyNode = GetParentNode();
 	
@@ -2623,19 +2657,13 @@ void FPropertyNode::PropagateContainerPropertyChange( UObject* ModifiedObject, c
 		ConvertedProperty = NodeProperty->GetOwner<FProperty>();
 	}
 
-	ArrayProperty = CastField<FArrayProperty>(ConvertedProperty);
-	SetProperty = CastField<FSetProperty>(ConvertedProperty);
-	MapProperty = CastField<FMapProperty>(ConvertedProperty);
-
-	check(ArrayProperty || SetProperty || MapProperty);
-
 	TArray<UObject*> ArchetypeInstances, ObjectsToChange;
 	FPropertyNode* SubobjectPropertyNode = NULL;
 	UObject* Object = ModifiedObject;
 
 	if (Object->HasAnyFlags(RF_ClassDefaultObject|RF_ArchetypeObject))
 	{
-		// Object is a default suobject, collect all instances.
+		// Object is a default subobject, collect all instances.
 		Object->GetArchetypeInstances(ArchetypeInstances);
 	}
 	else if (Object->HasAnyFlags(RF_DefaultSubObject) && Object->GetOuter()->HasAnyFlags(RF_ClassDefaultObject|RF_ArchetypeObject))
@@ -2696,114 +2724,12 @@ void FPropertyNode::PropagateContainerPropertyChange( UObject* ModifiedObject, c
 
 					checkf(false, TEXT("PropagateContainerPropertyChange tried to propagate a change onto itself!"));
 				}
-				
-				const bool bIsDefaultContainerContent = ConvertedProperty->Identical(OriginalContainerAddr, Addr);
 
-				// Return instance changes result to caller
-				if (PropagationResult != nullptr)
+				const bool bIsDefaultContainerContent = ConvertedProperty->Identical(OriginalContainerAddr, Addr);
+				if (bIsDefaultContainerContent)
 				{
-					PropagationResult->Add(ActualObjToChange, bIsDefaultContainerContent);
+					OutAffectedInstances.Add(ActualObjToChange);
 				}
-
-				if (ArrayProperty)
-				{
-					FScriptArrayHelper ArrayHelper(ArrayProperty, Addr);
-
-					// Check if the original value was the default value and change it only then
-					if (bIsDefaultContainerContent)
-					{
-						int32 ElementToInitialize = -1;
-						switch (ChangeType)
-						{
-						case EPropertyArrayChangeType::Add:
-							ElementToInitialize = ArrayHelper.AddValue();
-							break;
-						case EPropertyArrayChangeType::Clear:
-							ArrayHelper.EmptyValues();
-							break;
-						case EPropertyArrayChangeType::Insert:
-							ArrayHelper.InsertValues(ArrayIndex, 1);
-							ElementToInitialize = ArrayIndex;
-							break;
-						case EPropertyArrayChangeType::Delete:
-							ArrayHelper.RemoveValues(ArrayIndex, 1);
-							break;
-						case EPropertyArrayChangeType::Duplicate:
-							ArrayHelper.InsertValues(ArrayIndex, 1);
-							// Copy the selected item's value to the new item.
-							NodeProperty->CopyCompleteValue(ArrayHelper.GetRawPtr(ArrayIndex), ArrayHelper.GetRawPtr(ArrayIndex + 1));
-							Object->InstanceSubobjectTemplates();
-							break;
-						case EPropertyArrayChangeType::Swap:
-							if (SwapIndex != INDEX_NONE)
-							{
-								ArrayHelper.SwapValues(Index, SwapIndex);
-							}
-							break;
-						}
-					}
-				}	// End Array
-
-				else if (SetProperty)
-				{
-					FScriptSetHelper SetHelper(SetProperty, Addr);
-
-					// Check if the original value was the default value and change it only then
-					if (bIsDefaultContainerContent)
-					{
-						int32 ElementToInitialize = -1;
-						switch (ChangeType)
-						{
-						case EPropertyArrayChangeType::Add:
-							ElementToInitialize = SetHelper.AddDefaultValue_Invalid_NeedsRehash();
-							SetHelper.Rehash();
-							break;
-						case EPropertyArrayChangeType::Clear:
-							SetHelper.EmptyElements();
-							break;
-						case EPropertyArrayChangeType::Insert:
-							check(false);	// Insert is not supported for sets
-							break;
-						case EPropertyArrayChangeType::Delete:
-							SetHelper.RemoveAt(SetHelper.FindInternalIndex(ArrayIndex));
-							SetHelper.Rehash();
-							break;
-						case EPropertyArrayChangeType::Duplicate:
-							check(false);	// Duplicate not supported on sets
-							break;
-						}
-					}
-				}	// End Set
-				else if (MapProperty)
-				{
-					FScriptMapHelper MapHelper(MapProperty, Addr);
-
-					// Check if the original value was the default value and change it only then
-					if (bIsDefaultContainerContent)
-					{
-						int32 ElementToInitialize = -1;
-						switch (ChangeType)
-						{
-						case EPropertyArrayChangeType::Add:
-							ElementToInitialize = MapHelper.AddDefaultValue_Invalid_NeedsRehash();
-							MapHelper.Rehash();
-							break;
-						case EPropertyArrayChangeType::Clear:
-							MapHelper.EmptyValues();
-							break;
-						case EPropertyArrayChangeType::Insert:
-							check(false);	// Insert is not supported for maps
-							break;
-						case EPropertyArrayChangeType::Delete:
-							MapHelper.RemoveAt(MapHelper.FindInternalIndex(ArrayIndex));
-							MapHelper.Rehash();
-							break;
-						case EPropertyArrayChangeType::Duplicate:
-							check(false);	// Duplicate is not supported for maps
-							break;
-						}
-					}
-				}	// End Map
 			}
 		}
 
@@ -2820,6 +2746,154 @@ void FPropertyNode::PropagateContainerPropertyChange( UObject* ModifiedObject, c
 	}
 }
 
+void FPropertyNode::PropagateContainerPropertyChange( UObject* ModifiedObject, const void* OriginalContainerAddr, const TArray<UObject*>& AffectedInstances, EPropertyArrayChangeType::Type ChangeType, int32 Index, int32 SwapIndex /*= INDEX_NONE*/)
+{
+	check(OriginalContainerAddr);
+
+	FProperty* NodeProperty = GetProperty();
+
+	FPropertyNode* ParentPropertyNode = GetParentNode();
+	
+	FProperty* ConvertedProperty = NULL;
+
+	if (ChangeType == EPropertyArrayChangeType::Add || ChangeType == EPropertyArrayChangeType::Clear)
+	{
+		ConvertedProperty = NodeProperty;
+	}
+	else
+	{
+		ConvertedProperty = NodeProperty->GetOwner<FProperty>();
+	}
+
+	FArrayProperty* ArrayProperty = CastField<FArrayProperty>(ConvertedProperty);
+	FSetProperty* SetProperty = CastField<FSetProperty>(ConvertedProperty);
+	FMapProperty* MapProperty = CastField<FMapProperty>(ConvertedProperty);
+
+	check(ArrayProperty || SetProperty || MapProperty);
+
+	FPropertyNode* SubobjectPropertyNode = NULL;
+
+	UObject* Object = ModifiedObject;
+
+	if (Object->HasAnyFlags(RF_ClassDefaultObject|RF_ArchetypeObject))
+	{
+		// Object is a default subobject
+	}
+	else if (Object->HasAnyFlags(RF_DefaultSubObject) && Object->GetOuter()->HasAnyFlags(RF_ClassDefaultObject|RF_ArchetypeObject))
+	{
+		// Object is a default subobject of a default object. Get the subobject property node and use its owner instead.
+		for (SubobjectPropertyNode = FindObjectItemParent(); SubobjectPropertyNode && !SubobjectPropertyNode->GetProperty(); SubobjectPropertyNode = SubobjectPropertyNode->GetParentNode());
+		if (SubobjectPropertyNode != NULL)
+		{
+			// Switch the object to the owner default object
+			Object = Object->GetOuter();	
+		}
+	}
+	
+	for (UObject* InstanceToChange : AffectedInstances)
+	{
+		uint8* Addr = NULL;
+
+		if (ChangeType == EPropertyArrayChangeType::Add || ChangeType == EPropertyArrayChangeType::Clear)
+		{
+			Addr = GetValueBaseAddressFromObject(InstanceToChange);
+		}
+		else
+		{
+			Addr = ParentPropertyNode->GetValueBaseAddressFromObject(InstanceToChange);
+		}
+
+		if (ArrayProperty)
+		{
+			FScriptArrayHelper ArrayHelper(ArrayProperty, Addr);
+
+			int32 ElementToInitialize = -1;
+			switch (ChangeType)
+			{
+			case EPropertyArrayChangeType::Add:
+				ElementToInitialize = ArrayHelper.AddValue();
+				break;
+			case EPropertyArrayChangeType::Clear:
+				ArrayHelper.EmptyValues();
+				break;
+			case EPropertyArrayChangeType::Insert:
+				ArrayHelper.InsertValues(ArrayIndex, 1);
+				ElementToInitialize = ArrayIndex;
+				break;
+			case EPropertyArrayChangeType::Delete:
+				ArrayHelper.RemoveValues(ArrayIndex, 1);
+				break;
+			case EPropertyArrayChangeType::Duplicate:
+				ArrayHelper.InsertValues(ArrayIndex, 1);
+				// Copy the selected item's value to the new item.
+				NodeProperty->CopyCompleteValue(ArrayHelper.GetRawPtr(ArrayIndex), ArrayHelper.GetRawPtr(ArrayIndex + 1));
+				Object->InstanceSubobjectTemplates();
+				break;
+			case EPropertyArrayChangeType::Swap:
+				if (SwapIndex != INDEX_NONE)
+				{
+					ArrayHelper.SwapValues(Index, SwapIndex);
+				}
+				break;
+			}
+		}	// End Array
+
+		else if (SetProperty)
+		{
+			FScriptSetHelper SetHelper(SetProperty, Addr);
+
+			int32 ElementToInitialize = -1;
+			switch (ChangeType)
+			{
+			case EPropertyArrayChangeType::Add:
+				ElementToInitialize = SetHelper.AddDefaultValue_Invalid_NeedsRehash();
+				SetHelper.Rehash();
+				break;
+			case EPropertyArrayChangeType::Clear:
+				SetHelper.EmptyElements();
+				break;
+			case EPropertyArrayChangeType::Insert:
+				check(false);	// Insert is not supported for sets
+				break;
+			case EPropertyArrayChangeType::Delete:
+				SetHelper.RemoveAt(SetHelper.FindInternalIndex(ArrayIndex));
+				SetHelper.Rehash();
+				break;
+			case EPropertyArrayChangeType::Duplicate:
+				check(false);	// Duplicate not supported on sets
+				break;
+			}
+		}	// End Set
+		else if (MapProperty)
+		{
+			FScriptMapHelper MapHelper(MapProperty, Addr);
+
+			// Check if the original value was the default value and change it only then
+			int32 ElementToInitialize = -1;
+			switch (ChangeType)
+			{
+			case EPropertyArrayChangeType::Add:
+				ElementToInitialize = MapHelper.AddDefaultValue_Invalid_NeedsRehash();
+				MapHelper.Rehash();
+				break;
+			case EPropertyArrayChangeType::Clear:
+				MapHelper.EmptyValues();
+				break;
+			case EPropertyArrayChangeType::Insert:
+				check(false);	// Insert is not supported for maps
+				break;
+			case EPropertyArrayChangeType::Delete:
+				MapHelper.RemoveAt(MapHelper.FindInternalIndex(ArrayIndex));
+				MapHelper.Rehash();
+				break;
+			case EPropertyArrayChangeType::Duplicate:
+				check(false);	// Duplicate is not supported for maps
+				break;
+			}
+		}	// End Map
+	}
+}
+
 void FPropertyNode::PropagatePropertyChange( UObject* ModifiedObject, const TCHAR* NewValue, const FString& PreviousValue )
 {
 	TArray<UObject*> ArchetypeInstances, ObjectsToChange;
diff --git a/Engine/Source/Editor/PropertyEditor/Private/PropertyNode.h b/Engine/Source/Editor/PropertyEditor/Private/PropertyNode.h
--- a/Engine/Source/Editor/PropertyEditor/Private/PropertyNode.h
+++ b/Engine/Source/Editor/PropertyEditor/Private/PropertyNode.h
@@ -577,6 +577,8 @@ public:
 	bool IsChildOfFavorite(void) const;
 
 	void NotifyPreChange(FProperty* PropertyAboutToChange, class FNotifyHook* InNotifyHook);
+	void NotifyPreChange(FProperty* PropertyAboutToChange, class FNotifyHook* InNotifyHook, const TSet<UObject*>& AffectedInstances);
+	void NotifyPreChange(FProperty* PropertyAboutToChange, class FNotifyHook* InNotifyHook, TSet<UObject*>&& AffectedInstances);
 
 	void NotifyPostChange(FPropertyChangedEvent& InPropertyChangedEvent, class FNotifyHook* InNotifyHook);
 
@@ -600,7 +602,29 @@ public:
 	 * @param	Index						Index of the modified item
 	 */
 	void PropagateContainerPropertyChange(UObject* ModifiedObject, const void* OriginalContainerAddr,
-		EPropertyArrayChangeType::Type ChangeType, int32 Index, TMap<UObject*, bool>* PropagationResult = nullptr, int32 SwapIndex = INDEX_NONE);
+		EPropertyArrayChangeType::Type ChangeType, int32 Index, int32 SwapIndex = INDEX_NONE);
+
+	/**
+	 * Gather the list of all instances that will be affected by a container property change
+	 *
+	 * @param	ModifiedObject				Object which property has been modified
+	 * @param	OriginalContainerAddr		Original address holding the container value before the modification
+	 * @param	ChangeType					In which way is the container modified
+	 * @param	OutAffectedInstances		Instances affected by the property change
+	 */
+	void GatherInstancesAffectedByContainerPropertyChange(UObject* ModifiedObject, const void* OriginalContainerAddr, EPropertyArrayChangeType::Type ChangeType, TArray<UObject*>& OutAffectedInstances);
+
+	/**
+	 * Propagates the property change of a container property to the provided archetype instances
+	 *
+	 * @param	ModifiedObject				Object which property has been modified
+	 * @param	OriginalContainerAddr		Original address holding the container value before the modification
+	 * @param	AffectedInstances			Instances affected by the property change
+	 * @param	ChangeType					In which way was the container modified
+	 * @param	Index						Index of the modified item
+	 */
+	void PropagateContainerPropertyChange(UObject* ModifiedObject, const void* OriginalContainerAddr, const TArray<UObject*>& AffectedInstances,
+		EPropertyArrayChangeType::Type ChangeType, int32 Index, int32 SwapIndex = INDEX_NONE);
 
 	/** Broadcasts when a property value changes */
 	DECLARE_EVENT(FPropertyNode, FPropertyValueChangedEvent);
@@ -892,6 +916,10 @@ protected:
 	uint8* GetStartAddress(const UObject* Obj) const;
 
 	TSharedRef<FEditPropertyChain> BuildPropertyChain( FProperty* PropertyAboutToChange );
+	TSharedRef<FEditPropertyChain> BuildPropertyChain( FProperty* PropertyAboutToChange, const TSet<UObject*>& InAffectedArchetypeInstances );
+	TSharedRef<FEditPropertyChain> BuildPropertyChain( FProperty* PropertyAboutToChange, TSet<UObject*>&& InAffectedArchetypeInstances );
+
+	void NotifyPreChangeInternal(TSharedRef<FEditPropertyChain> PropertyChain, FProperty* PropertyAboutToChange, FNotifyHook* InNotifyHook);
 
 	/**
 	 * Destroys all node within the hierarchy
@@ -899,7 +927,7 @@ protected:
 	void DestroyTree(const bool bInDestroySelf=true);
 
 	/**
-	 * Interface function for Custom Setup of Node (priot to node flags being set)
+	 * Interface function for Custom Setup of Node (prior to node flags being set)
 	 */
 	virtual void InitBeforeNodeFlags () {};
 
diff --git a/Engine/Source/Runtime/CoreUObject/Private/UObject/EditPropertyChain.cpp b/Engine/Source/Runtime/CoreUObject/Private/UObject/EditPropertyChain.cpp
--- a/Engine/Source/Runtime/CoreUObject/Private/UObject/EditPropertyChain.cpp
+++ b/Engine/Source/Runtime/CoreUObject/Private/UObject/EditPropertyChain.cpp
@@ -53,6 +53,11 @@ bool FEditPropertyChain::SetActiveMemberPropertyNode( FProperty* NewActiveMember
 	return bResult;
 }
 
+bool FEditPropertyChain::IsArchetypeInstanceAffected( UObject* InInstance ) const
+{
+	return AffectedInstances.Contains(InInstance);
+}
+
 /**
  * Returns the node corresponding to the currently active property.
  */
diff --git a/Engine/Source/Runtime/CoreUObject/Private/UObject/Obj.cpp b/Engine/Source/Runtime/CoreUObject/Private/UObject/Obj.cpp
--- a/Engine/Source/Runtime/CoreUObject/Private/UObject/Obj.cpp
+++ b/Engine/Source/Runtime/CoreUObject/Private/UObject/Obj.cpp
@@ -466,7 +466,7 @@ void UObject::PropagatePreEditChange( TArray<UObject*>& AffectedObjects, FEditPr
 
 		// in order to ensure that all objects are saved properly, only process the objects which have this object as their
 		// ObjectArchetype since we are going to call Pre/PostEditChange on each object (which could potentially affect which data is serialized
-		if ( Obj->GetArchetype() == this )
+		if ( Obj->GetArchetype() == this || Obj->GetOuter()->GetArchetype() == this )
 		{
 			// add this object to the list that we're going to process
 			Instances.Add(Obj);
@@ -480,14 +480,17 @@ void UObject::PropagatePreEditChange( TArray<UObject*>& AffectedObjects, FEditPr
 	{
 		UObject* Obj = Instances[i];
 
-		// this object must now be included in any undo/redo operations
-		Obj->SetFlags(RF_Transactional);
+		if ( PropertyAboutToChange.IsArchetypeInstanceAffected(Obj) )
+		{
+			// this object must now be included in any undo/redo operations
+			Obj->SetFlags(RF_Transactional);
 
-		// This will call ClearComponents in the Actor case, so that we do not serialize more stuff than we need to.
-		Obj->PreEditChange(PropertyAboutToChange);
+			// This will call ClearComponents in the Actor case, so that we do not serialize more stuff than we need to.
+			Obj->PreEditChange(PropertyAboutToChange);
 
-		// now recurse into this object, saving its instances
-		Obj->PropagatePreEditChange(AffectedObjects, PropertyAboutToChange);
+			// now recurse into this object, saving its instances
+			Obj->PropagatePreEditChange(AffectedObjects, PropertyAboutToChange);
+		}
 	}
 }
 
@@ -517,11 +520,14 @@ void UObject::PropagatePostEditChange( TArray<UObject*>& AffectedObjects, FPrope
 	{
 		UObject* Obj = Instances[i];
 
-		// notify the object that all changes are complete
-		Obj->PostEditChangeChainProperty(PropertyChangedEvent);
+		if ( PropertyChangedEvent.HasArchetypeInstanceChanged(Obj) )
+		{
+			// notify the object that all changes are complete
+			Obj->PostEditChangeChainProperty(PropertyChangedEvent);
 
-		// now recurse into this object, loading its instances
-		Obj->PropagatePostEditChange(AffectedObjects, PropertyChangedEvent);
+			// now recurse into this object, loading its instances
+			Obj->PropagatePostEditChange(AffectedObjects, PropertyChangedEvent);
+		}
 	}
 }
 
diff --git a/Engine/Source/Runtime/CoreUObject/Public/UObject/UnrealType.h b/Engine/Source/Runtime/CoreUObject/Public/UObject/UnrealType.h
--- a/Engine/Source/Runtime/CoreUObject/Public/UObject/UnrealType.h
+++ b/Engine/Source/Runtime/CoreUObject/Public/UObject/UnrealType.h
@@ -5096,6 +5096,20 @@ public:
 	 */
 	bool SetActiveMemberPropertyNode( FProperty* NewActiveMemberProperty );
 
+	/**
+	 * Specify the set of archetype instances that will be affected by the property change.
+	 */
+	template<typename T>
+	void SetAffectedArchetypeInstances( T&& InAffectedInstances )
+	{
+		AffectedInstances = Forward<T>(InAffectedInstances);
+	}
+	
+	/**
+	 * Returns whether the specified archetype instance will be affected by the property change.
+	 */
+	bool IsArchetypeInstanceAffected( UObject* InInstance ) const;
+
 	/**
 	 * Returns the node corresponding to the currently active property.
 	 */
@@ -5123,6 +5137,10 @@ protected:
 	 */
 	TDoubleLinkedListNode* ActiveMemberPropertyNode;
 
+	/**
+	 * Archetype instances that will be affected by the property change.
+	 */
+	TSet<UObject*> AffectedInstances;
 
 	/** TDoubleLinkedList interface */
 	/**
@@ -5203,11 +5221,12 @@ struct FPropertyChangedEvent
 	}
 
 	/**
-	 * Saves off map of instance changed result per archetype being set.
+	 * Specify the set of archetype instances that were modified by the property change.
 	 */
-	void SetInstancesChangedResultPerArchetype(TArrayView<const TMap<UObject*, bool>> InInstancesChangedResultPerArchetype)
+	template<typename T>
+	void SetInstancesChanged(T&& InInstancesChanged)
 	{
-		InstancesChangedResultPerArchetype = InInstancesChangedResultPerArchetype;
+		InstancesChanged = Forward<T>(InInstancesChanged);
 	}
 
 	/**
@@ -5230,23 +5249,12 @@ struct FPropertyChangedEvent
 	}
 
 	/**
-	 * Gets the instance changed status of the "current object" based on the instance provided
+	 * Test whether an archetype instance was modified.
 	 * InInstance - The instance we want to know the status.
 	 */
 	bool HasArchetypeInstanceChanged(UObject* InInstance) const
 	{
-		bool Retval = true;
-
-		if (InstancesChangedResultPerArchetype.IsValidIndex(ObjectIteratorIndex))
-		{
-			const bool* ValuePtr = InstancesChangedResultPerArchetype[ObjectIteratorIndex].Find(InInstance);
-			if (ValuePtr)
-			{
-				Retval = *ValuePtr;
-			}
-		}
-
-		return Retval;
+		return InstancesChanged.Contains(InInstance);
 	}
 
 	/**
@@ -5291,7 +5299,7 @@ private:
 	TArrayView<const TMap<FString, int32>> ArrayIndicesPerObject;
 	
 	//In the property window, multiple objects can be selected at once. In this case we want to know if an instance was updated for this operation (used in array/set/map context)
-	TArrayView<const TMap<UObject*, bool>> InstancesChangedResultPerArchetype;
+	TSet<UObject*> InstancesChanged;
 
 	/** List of top level objects being changed */
 	TArrayView<const UObject* const> TopLevelObjects;
diff --git a/Engine/Source/Runtime/Engine/Private/InstancedStaticMesh.cpp b/Engine/Source/Runtime/Engine/Private/InstancedStaticMesh.cpp
--- a/Engine/Source/Runtime/Engine/Private/InstancedStaticMesh.cpp
+++ b/Engine/Source/Runtime/Engine/Private/InstancedStaticMesh.cpp
@@ -2986,8 +2986,7 @@ void UInstancedStaticMeshComponent::PostEditChangeChainProperty(FPropertyChanged
 {
 	if(PropertyChangedEvent.Property != NULL)
 	{
-		// Only permit editing archetype or instance if instance was changed by an archetype
-		if (PropertyChangedEvent.Property->GetFName() == "PerInstanceSMData" && (HasAnyFlags(RF_ArchetypeObject|RF_ClassDefaultObject) ||  PropertyChangedEvent.HasArchetypeInstanceChanged(this)))
+		if (PropertyChangedEvent.Property->GetFName() == GET_MEMBER_NAME_CHECKED(UInstancedStaticMeshComponent, PerInstanceSMData))
 		{
 			if (PropertyChangedEvent.ChangeType == EPropertyChangeType::ArrayAdd
 				|| PropertyChangedEvent.ChangeType == EPropertyChangeType::Duplicate)
@@ -3014,7 +3013,7 @@ void UInstancedStaticMeshComponent::PostEditChangeChainProperty(FPropertyChanged
 			
 			MarkRenderStateDirty();
 		}
-		else if (PropertyChangedEvent.Property->GetFName() == "Transform")
+		else if (PropertyChangedEvent.Property->GetFName() == GET_MEMBER_NAME_CHECKED(FInstancedStaticMeshInstanceData, Transform))
 		{
 			PartialNavigationUpdate(-1);
 			// Force recreation of the render data
