PR #6704: Properties that have subcategories specified for them and have the ShowOnlyInnerProperties metadata flag now display correctly.  (Contributed by [REDACTED_PERSON])

diff --git a/Engine/Source/Editor/PropertyEditor/Private/ObjectPropertyNode.cpp b/Engine/Source/Editor/PropertyEditor/Private/ObjectPropertyNode.cpp
--- a/Engine/Source/Editor/PropertyEditor/Private/ObjectPropertyNode.cpp
+++ b/Engine/Source/Editor/PropertyEditor/Private/ObjectPropertyNode.cpp
@@ -473,6 +473,7 @@ void FObjectPropertyNode::GetCategoryProperties(const TSet<UClass*>& ClassesToCo
 	const TSet<FName>& CategoriesFromBlueprints, TSet<FName>& CategoriesFromProperties, TArray<FName>& SortedCategories)
 {
 	static const FName Name_bShowOnlyWhenTrue("bShowOnlyWhenTrue");
+	static const FName Name_ShowOnlyInnerProperties("ShowOnlyInnerProperties");
 
 	FName CategoryName = FObjectEditorUtils::GetCategoryFName(CurrentProperty);
 
@@ -504,6 +505,18 @@ void FObjectPropertyNode::GetCategoryProperties(const TSet<UClass*>& ClassesToCo
 			CategoriesFromProperties.Add(CategoryName);
 		}
 	}
+
+	if (CurrentProperty->HasMetaData(Name_ShowOnlyInnerProperties))
+	{
+		const FStructProperty* StructProperty = CastField<const FStructProperty>(CurrentProperty);
+		if (StructProperty)
+		{
+			for (TFieldIterator<FProperty> It(StructProperty->Struct); It; ++It)
+			{
+				GetCategoryProperties(ClassesToConsider, *It, bShouldShowDisableEditOnInstance, bShouldShowHiddenProperties, CategoriesFromBlueprints, CategoriesFromProperties, SortedCategories);
+			}
+		}
+	}
 }
 
 void FObjectPropertyNode::InternalInitChildNodes( FName SinglePropertyName )
