Optimize getting the material parameters for the UI.

- Cuts down the speed from ~15-20s on my machine to 4-5s for the master materials.

diff --git a/Engine/Source/Editor/MaterialEditor/Private/SMaterialParametersOverviewWidget.cpp b/Engine/Source/Editor/MaterialEditor/Private/SMaterialParametersOverviewWidget.cpp
--- a/Engine/Source/Editor/MaterialEditor/Private/SMaterialParametersOverviewWidget.cpp
+++ b/Engine/Source/Editor/MaterialEditor/Private/SMaterialParametersOverviewWidget.cpp
@@ -639,6 +639,9 @@ void SMaterialParametersOverviewTree::CreateGroupsWidget()
 
 	Children.Empty();
 	ParameterGroups->GetChildren(Children);
+
+	// the order should correspond to UnsortedParameters exactly
+	TArray<TSharedPtr<IPropertyHandle>> DeferredSearches;
 	for (int32 GroupIdx = 0; GroupIdx < Children.Num(); ++GroupIdx)
 	{
 		bHasAnyParameters = true;
@@ -694,14 +697,28 @@ void SMaterialParametersOverviewTree::CreateGroupsWidget()
 
 				NonLayerProperty.Parameter = Parameter;
 				NonLayerProperty.ParameterGroup = ParameterGroup;
-				NonLayerProperty.ParameterNode = Generator->FindTreeNode(ParameterValueProperty);
-				NonLayerProperty.ParameterHandle = NonLayerProperty.ParameterNode->CreatePropertyHandle();
 				NonLayerProperty.UnsortedName = Parameter->ParameterInfo.Name;
 
+				DeferredSearches.Add(ParameterValueProperty);
 				UnsortedParameters.Add(NonLayerProperty);
 			}
 		}
 	}
+
+	checkf(UnsortedParameters.Num() == DeferredSearches.Num(), TEXT("Internal inconsistency: number of node searches does not match the number of properties"));
+	TArray<TSharedPtr<IDetailTreeNode>> DeferredResults = Generator->FindTreeNodes(DeferredSearches);
+	checkf(UnsortedParameters.Num() == DeferredResults.Num(), TEXT("Internal inconsistency: number of node search results does not match the number of properties"));
+
+	for (int Idx = 0, NumUnsorted = UnsortedParameters.Num(); Idx < NumUnsorted; ++Idx)
+	{
+		FUnsortedParamData& NonLayerProperty = UnsortedParameters[Idx];
+		NonLayerProperty.ParameterNode = DeferredResults[Idx];
+		NonLayerProperty.ParameterHandle = NonLayerProperty.ParameterNode->CreatePropertyHandle();
+	}
+
+	DeferredResults.Empty();
+	DeferredSearches.Empty();
+
 	ShowSubParameters();
 	RequestTreeRefresh();
 	SetParentsExpansionState();
diff --git a/Engine/Source/Editor/PropertyEditor/Private/PropertyRowGenerator.cpp b/Engine/Source/Editor/PropertyEditor/Private/PropertyRowGenerator.cpp
--- a/Engine/Source/Editor/PropertyEditor/Private/PropertyRowGenerator.cpp
+++ b/Engine/Source/Editor/PropertyEditor/Private/PropertyRowGenerator.cpp
@@ -188,19 +188,81 @@ TSharedPtr<IDetailTreeNode> FPropertyRowGenerator::FindTreeNode(TSharedPtr<IProp
 {
 	if (PropertyHandle.IsValid() && PropertyHandle->IsValidHandle())
 	{
-		for (const TSharedPtr<IDetailTreeNode>& RootNode : RootTreeNodes)
+		TArray<TSharedPtr<IDetailTreeNode>> NodesToCheck;
+		NodesToCheck.Append(RootTreeNodes);
+		TSharedPtr<FPropertyNode> PropertyNodeWeAreAfter = StaticCastSharedPtr<FPropertyHandleBase>(PropertyHandle)->GetPropertyNode();
+
+		TArray<TSharedRef<IDetailTreeNode>> Children;
+		while(NodesToCheck.Num())
 		{
-			TSharedPtr<IDetailTreeNode> Node = FindTreeNodeRecursive(RootNode, PropertyHandle);
-			if (Node.IsValid())
+			TSharedPtr<IDetailTreeNode> Node = NodesToCheck.Pop(false);
+			TSharedPtr<FDetailTreeNode> TreeNodeImpl = StaticCastSharedPtr<FDetailTreeNode>(Node);
+			TSharedPtr<FPropertyNode> PropertyNode = TreeNodeImpl->GetPropertyNode();
+
+			if (UNLIKELY(PropertyNode.IsValid() && PropertyNode == PropertyNodeWeAreAfter))
 			{
 				return Node;
 			}
+
+			Node->GetChildren(Children);	// will nix the existing Children contents
+			NodesToCheck.Append(Children);
 		}
 	}
-
 	return nullptr;
 }
 
+TArray<TSharedPtr<IDetailTreeNode>> FPropertyRowGenerator::FindTreeNodes(TArray<TSharedPtr<IPropertyHandle>> PropertyHandles) const
+{
+	TArray<TSharedPtr<IDetailTreeNode>> NodesToCheck;
+	NodesToCheck.Append(RootTreeNodes);
+	TArray<TSharedPtr<FPropertyNode>> PropertyNodesWeAreAfter;
+
+	for(TSharedPtr<IPropertyHandle>& Handle: PropertyHandles)
+	{
+		PropertyNodesWeAreAfter.Add(StaticCastSharedPtr<FPropertyHandleBase>(Handle)->GetPropertyNode());
+	}
+	TArray<TSharedPtr<IDetailTreeNode>> Results;
+	Results.AddDefaulted(PropertyHandles.Num());
+
+	int32 NumNotFound = Results.Num();
+
+	TArray<TSharedRef<IDetailTreeNode>> Children;
+	while (NodesToCheck.Num())
+	{
+		TSharedPtr<IDetailTreeNode> Node = NodesToCheck.Pop(false);
+		TSharedPtr<FDetailTreeNode> TreeNodeImpl = StaticCastSharedPtr<FDetailTreeNode>(Node);
+		TSharedPtr<FPropertyNode> PropertyNode = TreeNodeImpl->GetPropertyNode();
+
+		if (PropertyNode.IsValid())
+		{
+			// check if any is one that we're after
+			for (int Idx = 0, NumResults = Results.Num(); Idx < NumResults; ++Idx)
+			{
+				if (!Results[Idx].IsValid())
+				{
+					if (UNLIKELY(PropertyNode == PropertyNodesWeAreAfter[Idx]))
+					{
+						Results[Idx] = Node;
+						--NumNotFound;
+						// assume no duplicates and break?
+					}
+				}
+			}
+
+			check(NumNotFound >= 0);
+			if (NumNotFound == 0)
+			{
+				break;
+			}
+		}
+
+		Node->GetChildren(Children);	// will nix the existing Children contents
+		NodesToCheck.Append(Children);
+	}
+
+	return Results;
+}
+
 void FPropertyRowGenerator::RegisterInstancedCustomPropertyLayout(UStruct* Class, FOnGetDetailCustomizationInstance DetailLayoutDelegate)
 {
 	check(Class);
diff --git a/Engine/Source/Editor/PropertyEditor/Private/PropertyRowGenerator.h b/Engine/Source/Editor/PropertyEditor/Private/PropertyRowGenerator.h
--- a/Engine/Source/Editor/PropertyEditor/Private/PropertyRowGenerator.h
+++ b/Engine/Source/Editor/PropertyEditor/Private/PropertyRowGenerator.h
@@ -82,6 +82,7 @@ public:
 	virtual void SetObjects(const TArray<UObject*>& InObjects) override;
 	virtual const TArray<TSharedRef<IDetailTreeNode>>& GetRootTreeNodes() const override;
 	virtual TSharedPtr<IDetailTreeNode> FindTreeNode(TSharedPtr<IPropertyHandle> PropertyHandle) const override;
+	virtual TArray<TSharedPtr<IDetailTreeNode>> FindTreeNodes(TArray<TSharedPtr<IPropertyHandle>> PropertyHandles) const override;
 	virtual FOnRowsRefreshed& OnRowsRefreshed() override { return RowsRefreshedDelegate; }
 	virtual void RegisterInstancedCustomPropertyLayout(UStruct* Class, FOnGetDetailCustomizationInstance DetailLayoutDelegate) override;
 	virtual void RegisterInstancedCustomPropertyTypeLayout(FName PropertyTypeName, FOnGetPropertyTypeCustomizationInstance PropertyTypeLayoutDelegate, TSharedPtr<IPropertyTypeIdentifier> Identifier = nullptr) override;
diff --git a/Engine/Source/Editor/PropertyEditor/Public/IPropertyRowGenerator.h b/Engine/Source/Editor/PropertyEditor/Public/IPropertyRowGenerator.h
--- a/Engine/Source/Editor/PropertyEditor/Public/IPropertyRowGenerator.h
+++ b/Engine/Source/Editor/PropertyEditor/Public/IPropertyRowGenerator.h
@@ -53,6 +53,13 @@ public:
 	 */
 	virtual TSharedPtr<IDetailTreeNode> FindTreeNode(TSharedPtr<IPropertyHandle> PropertyHandle) const = 0;
 
+	/**
+	* Finds tree nodes by property handles.
+	* 
+	* @return The found tree nodes (in the same order, a pointer will be null if the node cannot be found)
+	*/
+	virtual TArray<TSharedPtr<IDetailTreeNode>> FindTreeNodes(TArray<TSharedPtr<IPropertyHandle>> PropertyHandles) const = 0;
+
 	/**
 	 * Registers a custom detail layout delegate for a specific class in this instance of the generator only
 	 *
