Fix issues on losing focus during the rename (UE-89767).

- Redoing the fix for UE-89605 in particular widgets.

(Edigrating 12021273 to Release-4.25)

diff --git a/Engine/Source/Editor/PropertyEditor/Private/UserInterface/PropertyEditor/SPropertyEditorNumeric.h b/Engine/Source/Editor/PropertyEditor/Private/UserInterface/PropertyEditor/SPropertyEditorNumeric.h
--- a/Engine/Source/Editor/PropertyEditor/Private/UserInterface/PropertyEditor/SPropertyEditorNumeric.h
+++ b/Engine/Source/Editor/PropertyEditor/Private/UserInterface/PropertyEditor/SPropertyEditorNumeric.h
@@ -509,7 +509,11 @@ private:
 	void OnValueCommitted( NumericType NewValue, ETextCommit::Type CommitInfo )
 	{
 		const TSharedRef< IPropertyHandle > PropertyHandle = PropertyEditor->GetPropertyHandle();
-		PropertyHandle->SetValue( NewValue );
+		NumericType OrgValue(0);
+		if (PropertyHandle->GetValue(OrgValue) == FPropertyAccess::Fail || OrgValue != NewValue)
+		{
+			PropertyHandle->SetValue(NewValue);
+		}
 
 		if (TypeInterface.IsValid() && !TypeInterface->FixedDisplayUnits.IsSet())
 		{
diff --git a/Engine/Source/Editor/PropertyEditor/Private/UserInterface/PropertyEditor/SPropertyEditorText.cpp b/Engine/Source/Editor/PropertyEditor/Private/UserInterface/PropertyEditor/SPropertyEditorText.cpp
--- a/Engine/Source/Editor/PropertyEditor/Private/UserInterface/PropertyEditor/SPropertyEditorText.cpp
+++ b/Engine/Source/Editor/PropertyEditor/Private/UserInterface/PropertyEditor/SPropertyEditorText.cpp
@@ -111,7 +111,8 @@ void SPropertyEditorText::OnTextCommitted( const FText& NewText, ETextCommit::Ty
 	const TSharedRef< IPropertyHandle > PropertyHandle = PropertyEditor->GetPropertyHandle();
 
 	FText CurrentText;
-	if( PropertyHandle->GetValueAsFormattedText( CurrentText ) != FPropertyAccess::MultipleValues || NewText.ToString() != FPropertyEditor::MultipleValuesDisplayName )
+	if( (PropertyHandle->GetValueAsFormattedText( CurrentText ) != FPropertyAccess::MultipleValues || NewText.ToString() != FPropertyEditor::MultipleValuesDisplayName)
+		&& !NewText.ToString().Equals(CurrentText.ToString(), ESearchCase::CaseSensitive))
 	{
 		PropertyHandle->SetValueFromFormattedString( NewText.ToString() );
 	}
diff --git a/Engine/Source/Runtime/Slate/Private/Widgets/Text/SlateEditableTextLayout.cpp b/Engine/Source/Runtime/Slate/Private/Widgets/Text/SlateEditableTextLayout.cpp
--- a/Engine/Source/Runtime/Slate/Private/Widgets/Text/SlateEditableTextLayout.cpp
+++ b/Engine/Source/Runtime/Slate/Private/Widgets/Text/SlateEditableTextLayout.cpp
@@ -802,15 +802,11 @@ bool FSlateEditableTextLayout::HandleFocusLost(const FFocusEvent& InFocusEvent)
 
 	const FText EditedText = GetEditableText();
 
-	// do not commit if the edited text is the same as original
-	if (!EditedText.ToString().Equals(OriginalText.Text.ToString(), ESearchCase::CaseSensitive))
-	{
-		OwnerWidget->OnTextCommitted(EditedText, TextAction);
+	OwnerWidget->OnTextCommitted(EditedText, TextAction);
 
-		// Reload underlying value now it is committed  (commit may alter the value) 
-		// so it can be re-displayed in the edit box
-		LoadText();
-	}
+	// Reload underlying value now it is committed  (commit may alter the value) 
+	// so it can be re-displayed in the edit box
+	LoadText();
 
 	UpdateCursorHighlight();
 
