Reordered instanced components are no longer set to transient and disassociated from the current world.

diff --git a/Engine/Source/Editor/PropertyEditor/Private/PropertyHandleImpl.cpp b/Engine/Source/Editor/PropertyEditor/Private/PropertyHandleImpl.cpp
--- a/Engine/Source/Editor/PropertyEditor/Private/PropertyHandleImpl.cpp
+++ b/Engine/Source/Editor/PropertyEditor/Private/PropertyHandleImpl.cpp
@@ -1598,7 +1598,45 @@ void FPropertyValueImpl::SwapChildren( TSharedPtr<FPropertyNode> FirstChildNode,
 
 void FPropertyValueImpl::MoveElementTo(int32 OriginalIndex, int32 NewIndex)
 {
-	// Insert into the middle or add to the end
+	if (OriginalIndex == NewIndex)
+	{
+		return;
+	}
+
+	// We don't do a true swap of the old index and new index here
+	// Instead we insert a new item after the existing item at the desired NewIndex, swap the Original after the "NewIndex" node, and delete the item at OriginalIndex
+	// To try to describe with array examples..
+
+	// Case A NewIndex > OriginalIndex
+	// Input Array: A B C D, where OriginalIndex = 1 and NewIndex = 2
+	// We add a new node (?) at NewIndex + 1
+	// A B C ? D
+	// Now swap OriginalIndex (B) and NewIndex + 1 (?)
+	// A ? C B D
+	// Delete OriginalIndex (?)
+	// A C B D
+
+	// Or -
+	// Input Array: A B C D, where OriginalIndex = 2 and NewIndex = 1
+	// Insert at NewIndex (?), which increments OriginalIndex (C)
+	// A ? B C D
+	// Swap OriginalIndex + 1 (C) and NewIndex (?)
+	// A C B ? D
+	// Delete OriginalIndex + 1 (?)
+	// A C B D
+
+	// If we're in case A, increment the "NewIndex" for swap target
+	if (NewIndex > OriginalIndex)
+	{
+		NewIndex += 1;
+	}
+	// Otherwise increment OriginalIndex
+	else
+	{
+		OriginalIndex += 1;
+	}
+
+	// Insert into the middle of the array
 	if (NewIndex < GetPropertyNode()->GetNumChildNodes())
 	{
 		TSharedPtr<FPropertyNode> InsertAfterChild = PropertyNode.Pin()->GetChildNode(NewIndex);
@@ -1657,7 +1695,7 @@ void FPropertyValueImpl::MoveElementTo(int32 OriginalIndex, int32 NewIndex)
 			}
 		}
 	}
-	else
+	else // or add to the end of the array
 	{
 		TSharedPtr<FPropertyNode> PropertyNodePin = PropertyNode.Pin();
 		if (PropertyNodePin.IsValid())
@@ -1723,12 +1761,6 @@ void FPropertyValueImpl::MoveElementTo(int32 OriginalIndex, int32 NewIndex)
 			}
 		}
 	}
-	// We inserted an element above our original index
-	if (NewIndex < OriginalIndex)
-	{
-		OriginalIndex += 1;
-	}
-
 
 	// Both Insert and Add are deferred so you need to rebuild the parent node's children
 	GetPropertyNode()->RebuildChildren();
@@ -1777,30 +1809,8 @@ void FPropertyValueImpl::MoveElementTo(int32 OriginalIndex, int32 NewIndex)
 						TopLevelObjects.Add(Obj);
 					}
 
-					if (ArrayProperty)
-					{
-						FScriptArrayHelper ArrayHelper(ArrayProperty, Address);
-
-						// If the inner property is an instanced component property we must move the old component to the 
-						// transient package so resetting owned components on the parent doesn't find it
-						FObjectProperty* InnerObjectProperty = CastField<FObjectProperty>(ArrayProperty->Inner);
-						if (InnerObjectProperty && InnerObjectProperty->HasAnyPropertyFlags(CPF_InstancedReference) && InnerObjectProperty->PropertyClass->IsChildOf(UActorComponent::StaticClass()))
-						{
-							if (UActorComponent* Component = *reinterpret_cast<UActorComponent**>(ArrayHelper.GetRawPtr(FirstIndex)))
-							{
-								Component->Modify();
-								Component->Rename(nullptr, GetTransientPackage(), REN_DontCreateRedirectors);
-							}
-
-							if (UActorComponent* Component = *reinterpret_cast<UActorComponent**>(ArrayHelper.GetRawPtr(SecondIndex)))
-							{
-								Component->Modify();
-								Component->Rename(nullptr, GetTransientPackage(), REN_DontCreateRedirectors);
-							}
-						}
-
-						ArrayHelper.SwapValues(FirstIndex, SecondIndex);
-					}
+					FScriptArrayHelper ArrayHelper(ArrayProperty, Address);
+					ArrayHelper.SwapValues(FirstIndex, SecondIndex);
 				}
 			}
 
@@ -1814,7 +1824,7 @@ void FPropertyValueImpl::MoveElementTo(int32 OriginalIndex, int32 NewIndex)
 		}
 	}
 
-	//Delete
+	// Delete the original index
 	{
 		FPropertyNode* ChildNodePtr = GetPropertyNode()->GetChildNode(OriginalIndex).Get();
 
@@ -1855,25 +1865,11 @@ void FPropertyValueImpl::MoveElementTo(int32 OriginalIndex, int32 NewIndex)
 					}
 
 					FScriptArrayHelper ArrayHelper(ArrayProperty, Address);
-
-					// If the inner property is an instanced component property we must move the old component to the 
-					// transient package so resetting owned components on the parent doesn't find it
-					FObjectProperty* InnerObjectProperty = CastField<FObjectProperty>(ArrayProperty->Inner);
-					if (InnerObjectProperty && InnerObjectProperty->HasAnyPropertyFlags(CPF_InstancedReference) && InnerObjectProperty->PropertyClass->IsChildOf(UActorComponent::StaticClass()))
-					{
-						if (UActorComponent* Component = *reinterpret_cast<UActorComponent**>(ArrayHelper.GetRawPtr(ChildNodePtr->GetArrayIndex())))
-						{
-							Component->Modify();
-							Component->Rename(nullptr, GetTransientPackage(), REN_DontCreateRedirectors);
-						}
-					}
-
 					ArrayHelper.RemoveValues(ChildNodePtr->GetArrayIndex());
-					
 				}
 			}
 
-			FPropertyChangedEvent ChangeEvent(ParentNode->GetProperty(), EPropertyChangeType::Unspecified, MakeArrayView(TopLevelObjects));
+			FPropertyChangedEvent ChangeEvent(ParentNode->GetProperty(), EPropertyChangeType::ArrayRemove, MakeArrayView(TopLevelObjects));
 			if (PropertyUtilities.IsValid())
 			{
 				ChildNodePtr->FixPropertiesInEvent(ChangeEvent);
diff --git a/Engine/Source/Editor/PropertyEditor/Private/SDetailSingleItemRow.cpp b/Engine/Source/Editor/PropertyEditor/Private/SDetailSingleItemRow.cpp
--- a/Engine/Source/Editor/PropertyEditor/Private/SDetailSingleItemRow.cpp
+++ b/Engine/Source/Editor/PropertyEditor/Private/SDetailSingleItemRow.cpp
@@ -186,10 +186,6 @@ FReply SDetailSingleItemRow::OnArrayDrop(const FDragDropEvent& DragDropEvent)
 		{
 			int32 OriginalIndex = SwappingPropertyNode->GetArrayIndex();
 			int32 NewIndex = SwappablePropertyNode->GetArrayIndex();
-			if (NewIndex > OriginalIndex)
-			{
-				NewIndex += 1;
-			}
 			TSharedPtr<IPropertyHandle> SwappingHandle = PropertyEditorHelpers::GetPropertyHandle(SwappingPropertyNode.ToSharedRef(), OwnerTreeNode.Pin()->GetDetailsView()->GetNotifyHook(), OwnerTreeNode.Pin()->GetDetailsView()->GetPropertyUtilities());
 			TSharedPtr<IPropertyHandleArray> ParentHandle = SwappingHandle->GetParentHandle()->AsArray();
 			if (ParentHandle.IsValid() && SwappablePropertyNode->GetParentNode() == SwappingPropertyNode->GetParentNode())
