Fixing Expression Parser token conflict:
- Moving primitive types token declarations inside ExpressionParserTypes.h to avoid defining them multiple times and have conflicting GUID.

diff --git a/Engine/Source/Editor/PropertyEditor/Private/EditConditionParser.h b/Engine/Source/Editor/PropertyEditor/Private/EditConditionParser.h
--- a/Engine/Source/Editor/PropertyEditor/Private/EditConditionParser.h
+++ b/Engine/Source/Editor/PropertyEditor/Private/EditConditionParser.h
@@ -74,7 +74,6 @@ namespace EditConditionParserTokens
 	DEFINE_EDIT_CONDITION_NODE(FSubExpressionStart, 0x78CE8411, 0x34076241, 0xA2BD0548, 0x54458A3A)
 	DEFINE_EDIT_CONDITION_NODE(FSubExpressionEnd, 0x5CB25466, 0x780CAE4D, 0x8A88E000, 0xD3695885)
 
-	DEFINE_EXPRESSION_NODE_TYPE(bool, 0xCACBC715, 0x505A6B4A, 0x8808809F, 0x897AA5F6)
 	DEFINE_EXPRESSION_NODE_TYPE(EditConditionParserTokens::FPropertyToken, 0x9A3FAF6F, 0xB2E45E4D, 0xA80A70C6, 0x47A89BD7)
 	DEFINE_EXPRESSION_NODE_TYPE(EditConditionParserTokens::FEnumToken, 0xC9A35C24, 0x21FC904B, 0x9F1B2B6A, 0xDF6F4BC4)
 	DEFINE_EXPRESSION_NODE_TYPE(EditConditionParserTokens::FNullPtrToken, 0xC3E91735, 0xB3B90C4D, 0x99BA3DCA, 0xC6CC6231)
diff --git a/Engine/Source/Runtime/Core/Public/Math/BasicMathExpressionEvaluator.h b/Engine/Source/Runtime/Core/Public/Math/BasicMathExpressionEvaluator.h
--- a/Engine/Source/Runtime/Core/Public/Math/BasicMathExpressionEvaluator.h
+++ b/Engine/Source/Runtime/Core/Public/Math/BasicMathExpressionEvaluator.h
@@ -15,8 +15,6 @@ namespace ExpressionParser {\
 DEFINE_EXPRESSION_NODE_TYPE(ExpressionParser::TYPE, __VA_ARGS__)
 
 /** Define some expression types for basic arithmetic */
-DEFINE_EXPRESSION_NODE_TYPE(double, 0x8444A8A3, 0x19AE4E13, 0xBCFA75EE, 0x39982BD6)
-
 DEFINE_EXPRESSION_OPERATOR_NODE(CORE_API, FSubExpressionStart, 0xCC40A083, 0xADBF46E2, 0xA93D12BB, 0x525D7417)
 DEFINE_EXPRESSION_OPERATOR_NODE(CORE_API, FSubExpressionEnd, 0x125E4C67, 0x96EB48C4, 0x8894E09C, 0xB3CD56BF)
 
diff --git a/Engine/Source/Runtime/Core/Public/Misc/ExpressionParserTypes.h b/Engine/Source/Runtime/Core/Public/Misc/ExpressionParserTypes.h
--- a/Engine/Source/Runtime/Core/Public/Misc/ExpressionParserTypes.h
+++ b/Engine/Source/Runtime/Core/Public/Misc/ExpressionParserTypes.h
@@ -167,6 +167,10 @@ template<> struct TGetExpressionNodeTypeId<TYPE>\
 };
 template<typename T> struct TGetExpressionNodeTypeId;
 
+/** Primitive types should only be declared once inside the codebase to avoid conflicts */
+DEFINE_EXPRESSION_NODE_TYPE(bool, 	0xCACBC715, 0x505A6B4A, 0x8808809F, 0x897AA5F6)
+DEFINE_EXPRESSION_NODE_TYPE(double, 0x8444A8A3, 0x19AE4E13, 0xBCFA75EE, 0x39982BD6)
+
 /**
  * A node in an expression.
  * 	Can be constructed from any C++ type that has a corresponding DEFINE_EXPRESSION_NODE_TYPE.
diff --git a/Engine/Source/Runtime/Core/Public/Misc/TextFilterExpressionEvaluator.h b/Engine/Source/Runtime/Core/Public/Misc/TextFilterExpressionEvaluator.h
--- a/Engine/Source/Runtime/Core/Public/Misc/TextFilterExpressionEvaluator.h
+++ b/Engine/Source/Runtime/Core/Public/Misc/TextFilterExpressionEvaluator.h
@@ -158,7 +158,8 @@ DEFINE_TEXT_EXPRESSION_OPERATOR_NODE(FNot,						2,	0x03D78990, 0x41D04E26, 0x8E9
 
 DEFINE_EXPRESSION_NODE_TYPE(TextFilterExpressionParser::FTextToken,	0x09E49538, 0x633545E3, 0x84B5644F, 0x1F11628F);
 DEFINE_EXPRESSION_NODE_TYPE(TextFilterExpressionParser::FFunction,	0x084E6214, 0x032FFA48, 0x9245ABF1, 0x9E248F1A);
-DEFINE_EXPRESSION_NODE_TYPE(bool,									0xC1CD5DCF, 0x2AB44958, 0xB3FF4F8F, 0xE665D121);
+
+#undef DEFINE_TEXT_EXPRESSION_OPERATOR_NODE
 
 DECLARE_DELEGATE_RetVal_OneParam(bool, FTokenFunctionHandler, const FTextFilterString&);
 
