With SPropertyEditorNumeric, if the value has changed using a slider set the value without flags to ensure that all the callback are called.

diff --git a/Engine/Source/Editor/PropertyEditor/Private/UserInterface/PropertyEditor/SPropertyEditorNumeric.h b/Engine/Source/Editor/PropertyEditor/Private/UserInterface/PropertyEditor/SPropertyEditorNumeric.h
--- a/Engine/Source/Editor/PropertyEditor/Private/UserInterface/PropertyEditor/SPropertyEditorNumeric.h
+++ b/Engine/Source/Editor/PropertyEditor/Private/UserInterface/PropertyEditor/SPropertyEditorNumeric.h
@@ -513,6 +513,7 @@ private:
 		if (PropertyHandle->GetValue(OrgValue) == FPropertyAccess::Fail || OrgValue != NewValue)
 		{
 			PropertyHandle->SetValue(NewValue);
+			LastSliderCommittedValue = NewValue;
 		}
 
 		if (TypeInterface.IsValid() && !TypeInterface->FixedDisplayUnits.IsSet())
@@ -551,6 +552,9 @@ private:
 	void OnBeginSliderMovement()
 	{
 		bIsUsingSlider = true;
+
+		const TSharedRef<IPropertyHandle> PropertyHandle = PropertyEditor->GetPropertyHandle();
+		PropertyHandle->GetValue(LastSliderCommittedValue);
 		
 		GEditor->BeginTransaction(TEXT("PropertyEditor"), FText::Format(NSLOCTEXT("PropertyEditor", "SetNumericPropertyTransaction", "Edit {0}"), PropertyEditor->GetDisplayName()), nullptr /** PropertyEditor->GetPropertyHandle()->GetProperty() */ );
 	}
@@ -563,7 +567,17 @@ private:
 	{
 		bIsUsingSlider = false;
 
-		GEditor->EndTransaction();
+		// When the slider end, we may have not called SetValue(NewValue) without the InteractiveChange|NotTransactable flags.
+		//That prevents some transaction and callback to be triggered like the NotifyHook.
+		if (LastSliderCommittedValue != NewValue)
+		{
+			const TSharedRef<IPropertyHandle> PropertyHandle = PropertyEditor->GetPropertyHandle();
+			PropertyHandle->SetValue(NewValue);
+		}
+		else
+		{
+			GEditor->EndTransaction();
+		}
 	}
 
 	/** @return True if the property can be edited */
@@ -617,6 +631,9 @@ private:
 
 	/** True if the slider is being used to change the value of the property */
 	bool bIsUsingSlider;
+
+	/** When using the slider, what was the last committed value */
+	NumericType LastSliderCommittedValue;
 };
 
 #undef LOCTEXT_NAMESPACE
