Remove a comment that got reinroduce by a merge.

diff --git a/Engine/Source/Editor/PropertyEditor/Private/DetailCategoryBuilderImpl.cpp b/Engine/Source/Editor/PropertyEditor/Private/DetailCategoryBuilderImpl.cpp
--- a/Engine/Source/Editor/PropertyEditor/Private/DetailCategoryBuilderImpl.cpp
+++ b/Engine/Source/Editor/PropertyEditor/Private/DetailCategoryBuilderImpl.cpp
@@ -75,8 +75,6 @@ bool FDetailLayoutCustomization::IsHidden() const
 	return !IsValidCustomization()
 		|| (HasCustomWidget() && WidgetDecl->VisibilityAttr.Get() != EVisibility::Visible)
 		|| (HasPropertyNode() && PropertyRow->GetPropertyVisibility() != EVisibility::Visible);
-	/** Partial revert of CL 7273612 (fix for UE-76064) that caused a bunch of regressions in the details panel (UE-77377,UE-77376,UE-77451). */
-	//|| (HasCustomBuilder() && CustomBuilderRow->AreChildCustomizationsHidden());
 }
 
 
