FSoftObjectProperty now correctly respects CanEditChange for properties when edited from the Details panel.

diff --git a/Engine/Source/Editor/PropertyEditor/Private/PropertyCustomizationHelpers.cpp b/Engine/Source/Editor/PropertyEditor/Private/PropertyCustomizationHelpers.cpp
--- a/Engine/Source/Editor/PropertyEditor/Private/PropertyCustomizationHelpers.cpp
+++ b/Engine/Source/Editor/PropertyEditor/Private/PropertyCustomizationHelpers.cpp
@@ -453,6 +453,7 @@ void SObjectPropertyEntryBox::Construct( const FArguments& InArgs )
 	ObjectPath = InArgs._ObjectPath;
 	OnObjectChanged = InArgs._OnObjectChanged;
 	OnShouldSetAsset = InArgs._OnShouldSetAsset;
+	OnIsEnabled = InArgs._OnIsEnabled;
 
 	const TArray<FAssetData>& OwnerAssetDataArray = InArgs._OwnerAssetDataArray;
 
@@ -463,7 +464,6 @@ void SObjectPropertyEntryBox::Construct( const FArguments& InArgs )
 		ThumbnailSize = InArgs._ThumbnailSizeOverride.Get();
 	}
 
-
 	if( InArgs._PropertyHandle.IsValid() && InArgs._PropertyHandle->IsValidHandle() )
 	{
 		PropertyHandle = InArgs._PropertyHandle;
@@ -500,7 +500,7 @@ void SObjectPropertyEntryBox::Construct( const FArguments& InArgs )
 	if (InArgs._CustomResetToDefault.IsSet() || (PropertyHandle.IsValid() && !PropertyHandle->HasMetaData(TEXT("NoResetToDefault")) && !PropertyHandle->IsResetToDefaultCustomized()))
 	{
 		SAssignNew(ResetButton, SResetToDefaultPropertyEditor, PropertyHandle)
-			.IsEnabled(true)
+			.IsEnabled(this, &SObjectPropertyEntryBox::IsEnabled)
 			.CustomResetToDefault(InArgs._CustomResetToDefault);		
 	};
 
@@ -517,6 +517,7 @@ void SObjectPropertyEntryBox::Construct( const FArguments& InArgs )
 				.ObjectPath( this, &SObjectPropertyEntryBox::OnGetObjectPath )
 				.Class( InArgs._AllowedClass )
 				.NewAssetFactories( InArgs._NewAssetFactories )
+				.IsEnabled(this, &SObjectPropertyEntryBox::IsEnabled)
 				.OnSetObject(this, &SObjectPropertyEntryBox::OnSetObject)
 				.ThumbnailPool(InArgs._ThumbnailPool)
 				.DisplayThumbnail(bDisplayThumbnail)
@@ -574,6 +575,22 @@ void SObjectPropertyEntryBox::OnSetObject(const FAssetData& AssetData)
 	OnObjectChanged.ExecuteIfBound(AssetData);
 }
 
+bool SObjectPropertyEntryBox::IsEnabled() const
+{
+	bool IsEnabled = true;
+	if (PropertyHandle.IsValid())
+	{
+		IsEnabled &= PropertyHandle->IsEditable();
+	}
+
+	if (OnIsEnabled.IsBound())
+	{
+		IsEnabled &= OnIsEnabled.Execute();
+	}
+
+	return IsEnabled;
+}
+
 void SClassPropertyEntryBox::Construct(const FArguments& InArgs)
 {
 	ChildSlot
diff --git a/Engine/Source/Editor/PropertyEditor/Public/PropertyCustomizationHelpers.h b/Engine/Source/Editor/PropertyEditor/Public/PropertyCustomizationHelpers.h
--- a/Engine/Source/Editor/PropertyEditor/Public/PropertyCustomizationHelpers.h
+++ b/Engine/Source/Editor/PropertyEditor/Public/PropertyCustomizationHelpers.h
@@ -51,6 +51,7 @@ DECLARE_DELEGATE_ThreeParams(FOnGetPropertyComboBoxStrings, TArray< TSharedPtr<F
 DECLARE_DELEGATE_RetVal(FString, FOnGetPropertyComboBoxValue);
 DECLARE_DELEGATE_OneParam(FOnPropertyComboBoxValueSelected, const FString&);
 DECLARE_DELEGATE_ThreeParams(FOnInstancedPropertyIteration, IDetailCategoryBuilder&, IDetailGroup*, TSharedRef<IPropertyHandle>&);
+DECLARE_DELEGATE_RetVal(bool, FOnIsEnabled);
 
 namespace PropertyCustomizationHelpers
 {
@@ -150,6 +151,8 @@ public:
 		SLATE_EVENT(FOnSetObject, OnObjectChanged)
 		/** Called to check if an asset is valid to use */
 		SLATE_EVENT(FOnShouldFilterAsset, OnShouldFilterAsset)
+		/** Called to check if the asset should be enabled. */
+		SLATE_EVENT(FOnIsEnabled, OnIsEnabled)
 		/** Whether the asset can be 'None' */
 		SLATE_ARGUMENT(bool, AllowClear)
 		/** Whether to show the 'Use Selected' button */
@@ -181,11 +184,16 @@ private:
 
 	/** @return the object path for the object we are viewing */
 	FString OnGetObjectPath() const;
+
+	bool IsEnabled() const;
+
 private:
 	/** Delegate to call to determine whether the asset should be set */
 	FOnShouldSetAsset OnShouldSetAsset;
 	/** Delegate to call when the object changes */
 	FOnSetObject OnObjectChanged;
+	/** Delegate to call to check if this widget should be enabled. */
+	FOnIsEnabled OnIsEnabled;
 	/** Path to the object */
 	TAttribute<FString> ObjectPath;
 	/** Handle to a property we modify (if any)*/
diff --git a/Engine/Source/Editor/UnrealEd/Classes/Editor/PropertyEditorTestObject.h b/Engine/Source/Editor/UnrealEd/Classes/Editor/PropertyEditorTestObject.h
--- a/Engine/Source/Editor/UnrealEd/Classes/Editor/PropertyEditorTestObject.h
+++ b/Engine/Source/Editor/UnrealEd/Classes/Editor/PropertyEditorTestObject.h
@@ -381,6 +381,9 @@ class UPropertyEditorTestObject : public UObject
 	UPROPERTY(EditAnywhere, Category=StructTests, meta=(AllowedClasses="Material,Texture"))
 	FSoftObjectPath MaterialOrTextureAssetReference;
 
+	UPROPERTY(EditAnywhere, Category=StructTests)
+	FSoftObjectPath DisabledByCanEditChange;
+
 	UPROPERTY(EditAnywhere, Category=StructTests, meta=(InlineEditConditionToggle))
 	bool bEditCondition;
 
@@ -635,4 +638,27 @@ class UPropertyEditorTestObject : public UObject
 
 	UPROPERTY(EditAnywhere, Category = "Inline Edit Conditions", meta = (EditCondition = "bSharedEditCondition"))
 	float UsesSharedEditCondition2;
+
+	bool CanEditChange(const FProperty* InProperty) const;
+};
+
+UCLASS(EditInlineNew, Blueprintable)
+class UBlueprintPropertyTestObject : public UObject
+{
+	GENERATED_BODY()
+public:
+	UPROPERTY(EditAnywhere, Category="Hidden")
+	int32 ShouldBeHidden;
+
+	UPROPERTY(EditAnywhere, Category="Visible")
+	int32 ShouldBeVisible;
+};
+
+UCLASS(Blueprintable)
+class UBlueprintPropertyContainerTestObject : public UObject
+{
+	GENERATED_BODY()
+public:
+	UPROPERTY(Instanced, EditAnywhere, Category="Default", meta=(ShowOnlyInnerProperties))
+	TArray<UBlueprintPropertyTestObject*> Array;
 };
\ No newline at end of file
diff --git a/Engine/Source/Editor/UnrealEd/Private/PropertyEditorTestObject.cpp b/Engine/Source/Editor/UnrealEd/Private/PropertyEditorTestObject.cpp
--- a/Engine/Source/Editor/UnrealEd/Private/PropertyEditorTestObject.cpp
+++ b/Engine/Source/Editor/UnrealEd/Private/PropertyEditorTestObject.cpp
@@ -1,3 +1,13 @@
 // Copyright Epic Games, Inc. All Rights Reserved.
 
-#include "Editor/PropertyEditorTestObject.h"
\ No newline at end of file
+#include "Editor/PropertyEditorTestObject.h"
+
+bool UPropertyEditorTestObject::CanEditChange(const FProperty* InProperty) const
+{
+	if (InProperty->GetFName() == GET_MEMBER_NAME_CHECKED(UPropertyEditorTestObject, DisabledByCanEditChange))
+	{
+		return false;
+	}
+
+	return true;
+}
\ No newline at end of file
