Adding additional validity checks before pinning a weak pointer

diff --git a/Engine/Source/Editor/PropertyEditor/Private/DetailItemNode.cpp b/Engine/Source/Editor/PropertyEditor/Private/DetailItemNode.cpp
--- a/Engine/Source/Editor/PropertyEditor/Private/DetailItemNode.cpp
+++ b/Engine/Source/Editor/PropertyEditor/Private/DetailItemNode.cpp
@@ -96,7 +96,7 @@ EDetailNodeType FDetailItemNode::GetNodeType() const
 
 TSharedPtr<IPropertyHandle> FDetailItemNode::CreatePropertyHandle() const
 {
-	if (Customization.HasPropertyNode() && ParentCategory.IsValid())
+	if (Customization.HasPropertyNode() && ParentCategory.IsValid() && ParentCategory.Pin()->IsParentLayoutValid())
 	{
 		return ParentCategory.Pin()->GetParentLayoutImpl().GetPropertyHandle(Customization.GetPropertyNode());
 	}
