Struct customizations are now allowed to omit the header row in arrays and sets. Maps still require the header so that a default key widget can be generated.

Change requested by UDN.

diff --git a/Engine/Source/Editor/PropertyEditor/Private/DetailPropertyRow.cpp b/Engine/Source/Editor/PropertyEditor/Private/DetailPropertyRow.cpp
--- a/Engine/Source/Editor/PropertyEditor/Private/DetailPropertyRow.cpp
+++ b/Engine/Source/Editor/PropertyEditor/Private/DetailPropertyRow.cpp
@@ -214,14 +214,12 @@ FDetailWidgetRow FDetailPropertyRow::GetWidgetRow()
 	}
 }
 
-static bool IsInContainer(const TSharedPtr<IPropertyHandle>& PropertyHandle)
+static bool IsHeaderRowRequired(const TSharedPtr<IPropertyHandle>& PropertyHandle)
 {
 	TSharedPtr<IPropertyHandle> ParentHandle = PropertyHandle->GetParentHandle();
 	while (ParentHandle.IsValid())
 	{
-		if (ParentHandle->AsArray().IsValid() || 
-			ParentHandle->AsSet().IsValid() ||
-			ParentHandle->AsMap().IsValid())
+		if (ParentHandle->AsMap().IsValid())
 		{
 			return true;
 		}
@@ -234,7 +232,7 @@ static bool IsInContainer(const TSharedPtr<IPropertyHandle>& PropertyHandle)
 
 static void FixEmptyHeaderRowInContainers(const TSharedPtr<IPropertyHandle>& PropertyHandle, const TSharedPtr<FDetailWidgetRow>& HeaderRow)
 {
-	if (IsInContainer(PropertyHandle))
+	if (IsHeaderRowRequired(PropertyHandle))
 	{
 		if (!HeaderRow->HasAnyContent())
 		{
