Bitmask flag dropdowns now have a tooltip that contains all the set flags.

diff --git a/Engine/Source/Editor/PropertyEditor/Private/UserInterface/PropertyEditor/SPropertyEditorNumeric.h b/Engine/Source/Editor/PropertyEditor/Private/UserInterface/PropertyEditor/SPropertyEditorNumeric.h
--- a/Engine/Source/Editor/PropertyEditor/Private/UserInterface/PropertyEditor/SPropertyEditorNumeric.h
+++ b/Engine/Source/Editor/PropertyEditor/Private/UserInterface/PropertyEditor/SPropertyEditorNumeric.h
@@ -164,6 +164,29 @@ public:
 			SAssignNew(PrimaryWidget, SComboButton)
 			.ComboButtonStyle(&ComboBoxStyle.ComboButtonStyle)
 			.ContentPadding(FMargin(4.0, 2.0))
+			.ToolTipText_Lambda([this, CreateBitmaskFlagsArray, Property]
+			{
+				TOptional<NumericType> Value = OnGetValue();
+				if (Value.IsSet())
+				{
+					TArray<FBitmaskFlagInfo> BitmaskFlags = CreateBitmaskFlagsArray(Property);
+
+					TArray<FText> SetFlags;
+					SetFlags.Reserve(BitmaskFlags.Num());
+
+					for (const FBitmaskFlagInfo& FlagInfo : BitmaskFlags)
+					{
+						if (TBitmaskValueHelpers<NumericType>::BitwiseAND(Value.GetValue(), FlagInfo.Value))
+						{
+							SetFlags.Add(FlagInfo.DisplayName);
+						}
+					}
+
+					return FText::Join(FText::FromString(" | "), SetFlags);
+				}
+
+				return FText::GetEmpty();
+			})
 			.ButtonContent()
 			[
 				SNew(STextBlock)
