Editor - Adding an option to "Remove All But This", to the SPropertyTableHeaderRow columns for the bulk property grid.

diff --git a/Engine/Source/Editor/PropertyEditor/Private/UserInterface/PropertyTable/SPropertyTableHeaderRow.h b/Engine/Source/Editor/PropertyEditor/Private/UserInterface/PropertyTable/SPropertyTableHeaderRow.h
--- a/Engine/Source/Editor/PropertyEditor/Private/UserInterface/PropertyTable/SPropertyTableHeaderRow.h
+++ b/Engine/Source/Editor/PropertyEditor/Private/UserInterface/PropertyTable/SPropertyTableHeaderRow.h
@@ -173,12 +173,23 @@ private:
 		const bool bShouldCloseWindowAfterMenuSelection = true;
 		FMenuBuilder MenuBuilder( bShouldCloseWindowAfterMenuSelection, NULL );
 		{
-			FUIAction Action( 
-				FExecuteAction::CreateSP( SharedThis( this ), &SPropertyTableHeaderRow::RemoveColumn, Column ),
-				FCanExecuteAction()
-			);
+			{
+				FUIAction Action(
+					FExecuteAction::CreateSP(SharedThis(this), &SPropertyTableHeaderRow::RemoveColumn, Column),
+					FCanExecuteAction()
+				);
+
+				MenuBuilder.AddMenuEntry( LOCTEXT("RemoveColumn_Name", "Remove"), LOCTEXT("RemoveColumn_ToolTip", "Removes the column from the table view"), FSlateIcon(), Action, NAME_None, EUserInterfaceActionType::Button );
+			}
 
-			MenuBuilder.AddMenuEntry( LOCTEXT("ColumnHeaderMenu", "Remove"), LOCTEXT("RemoveColumn_ToolTip", "Removes the column from the table view"), FSlateIcon(), Action, NAME_None, EUserInterfaceActionType::Button );
+			{
+				FUIAction Action(
+					FExecuteAction::CreateSP(SharedThis(this), &SPropertyTableHeaderRow::RemoveAllButThisColumn, Column),
+					FCanExecuteAction()
+				);
+
+				MenuBuilder.AddMenuEntry(LOCTEXT("RemoveAllButThisColumn_Name", "Remove All But This"), LOCTEXT("RemoveAllButThisColumn_ToolTip", "Removes all the columns from the table view except this"), FSlateIcon(), Action, NAME_None, EUserInterfaceActionType::Button);
+			}
 		}
 
 		return MenuBuilder.MakeWidget();
@@ -196,6 +207,19 @@ private:
 		PropertyTable->RemoveColumn( Column );
 	}
 
+	void RemoveAllButThisColumn(const TSharedRef<IPropertyTableColumn> ColumnToKeep)
+	{
+		TArray<TSharedRef<IPropertyTableColumn>> ExistingColumns = PropertyTable->GetColumns();
+		for (TSharedRef<IPropertyTableColumn> ExistingColumn : ExistingColumns)
+		{
+			if (ExistingColumn->IsFrozen() || ExistingColumn == ColumnToKeep)
+			{
+				continue;
+			}
+
+			PropertyTable->RemoveColumn(ExistingColumn);
+		}
+	}
 
 private:
 
