Fix details panel refresh issues (which go away after changing actor selection):
- Details view did not update correctly when removing the actor details object filter.
- Actor components box visibility

diff --git a/Engine/Source/Editor/LevelEditor/Private/SActorDetails.cpp b/Engine/Source/Editor/LevelEditor/Private/SActorDetails.cpp
--- a/Engine/Source/Editor/LevelEditor/Private/SActorDetails.cpp
+++ b/Engine/Source/Editor/LevelEditor/Private/SActorDetails.cpp
@@ -135,7 +135,7 @@ void SActorDetails::Construct(const FArguments& InArgs, const FName TabIdentifie
 	GEditor->RegisterForUndo(this);
 
 	ComponentsBox = SNew(SBox)
-		.Visibility(EVisibility::Collapsed);
+		.Visibility(this, &SActorDetails::GetComponentsBoxVisibility);
 
 	SCSEditor = SNew(SSCSEditor)
 		.EditorMode(EComponentEditorMode::ActorInstance)
@@ -240,7 +240,7 @@ void SActorDetails::OnDetailsViewObjectArrayChanged(const FString& InTitle, cons
 	// The DetailsView will already check validity every tick and hide itself when invalid, so this piggy-backs on that code instead of needing a second tick function.
 	if (InObjects.Num() == 0 && !LockedActorSelection.IsValid())
 	{
-		ComponentsBox->SetVisibility(EVisibility::Collapsed);
+		bShowingComponents = false;
 	}
 }
 
@@ -250,7 +250,7 @@ void SActorDetails::SetObjects(const TArray<UObject*>& InObjects, bool bForceRef
 	{
 		DetailsView->SetObjects(InObjects, bForceRefresh);
 
-		bool bShowingComponents = false;
+		bShowingComponents = false;
 
 		if(InObjects.Num() == 1 && FKismetEditorUtilities::CanCreateBlueprintOfClass(InObjects[0]->GetClass()))
 		{
@@ -270,9 +270,6 @@ void SActorDetails::SetObjects(const TArray<UObject*>& InObjects, bool bForceRef
 			}
 		}
 
-		bool ShowComponentsInDetailsPanel = ShowComponents->GetBool();
- 		ComponentsBox->SetVisibility((bShowingComponents & ShowComponentsInDetailsPanel) ? EVisibility::Visible : EVisibility::Collapsed);
-
 		if(DetailsView->GetHostTabManager().IsValid())
 		{
 			TSharedPtr<SDockTab> Tab = DetailsView->GetHostTabManager()->FindExistingLiveTab(DetailsView->GetIdentifier());
@@ -708,6 +705,11 @@ void SActorDetails::OnNativeComponentWarningHyperlinkClicked(const FSlateHyperli
 	}
 }
 
+EVisibility SActorDetails::GetComponentsBoxVisibility() const
+{
+	return (bShowingComponents & ShowComponents->GetBool()) ? EVisibility::Visible : EVisibility::Collapsed;
+}
+
 EVisibility SActorDetails::GetUCSComponentWarningVisibility() const
 {
 	bool bIsUneditableBlueprintComponent = false;
diff --git a/Engine/Source/Editor/LevelEditor/Private/SActorDetails.h b/Engine/Source/Editor/LevelEditor/Private/SActorDetails.h
--- a/Engine/Source/Editor/LevelEditor/Private/SActorDetails.h
+++ b/Engine/Source/Editor/LevelEditor/Private/SActorDetails.h
@@ -67,6 +67,7 @@ private:
 
 	bool IsPropertyReadOnly(const struct FPropertyAndParent& PropertyAndParent) const;
 	bool IsPropertyEditingEnabled() const;
+	EVisibility GetComponentsBoxVisibility() const;
 	EVisibility GetUCSComponentWarningVisibility() const;
 	EVisibility GetInheritedBlueprintComponentWarningVisibility() const;
 	EVisibility GetNativeComponentWarningVisibility() const;
@@ -93,6 +94,9 @@ private:
 	// Used to prevent reentrant changes
 	bool bSelectionGuard;
 
+	// True if the actor details should be showing the components.
+	bool bShowingComponents = false;
+
 	// True if the actor "root" node in the SCS editor is currently shown as selected
 	bool bShowingRootActorNodeSelected;
 };
diff --git a/Engine/Source/Editor/PropertyEditor/Private/PropertyEditorModule.cpp b/Engine/Source/Editor/PropertyEditor/Private/PropertyEditorModule.cpp
--- a/Engine/Source/Editor/PropertyEditor/Private/PropertyEditorModule.cpp
+++ b/Engine/Source/Editor/PropertyEditor/Private/PropertyEditorModule.cpp
@@ -555,7 +555,7 @@ void FPropertyEditorModule::UpdatePropertyViews( const TArray<UObject*>& NewObje
 				}
 				else
 				{
-					DetailViewPin->RemoveInvalidObjects();
+					DetailViewPin->ForceRefresh();
 				}
 			}
 		}
diff --git a/Engine/Source/Editor/PropertyEditor/Private/SDetailsView.cpp b/Engine/Source/Editor/PropertyEditor/Private/SDetailsView.cpp
--- a/Engine/Source/Editor/PropertyEditor/Private/SDetailsView.cpp
+++ b/Engine/Source/Editor/PropertyEditor/Private/SDetailsView.cpp
@@ -373,26 +373,22 @@ EVisibility SDetailsView::GetScrollBarVisibility() const
 void SDetailsView::ForceRefresh()
 {
 	TArray<UObject*> NewObjectList;
+	NewObjectList.Reserve(UnfilteredSelectedObjects.Num());
+	TArray<TWeakObjectPtr<UObject>> ValidSelectedObjects;
+	ValidSelectedObjects.Reserve(UnfilteredSelectedObjects.Num());
 
-	const FRootPropertyNodeList& RootNodes = GetRootNodes();
-	for(const TSharedPtr<FComplexPropertyNode>& ComplexRootNode : RootNodes)
+	for (const TWeakObjectPtr<UObject>& Object : UnfilteredSelectedObjects)
 	{
-		FObjectPropertyNode* RootNode = ComplexRootNode->AsObjectNode();
-		if(RootNode)
+		if (Object.IsValid())
 		{
-			// Simply re-add the same existing objects to cause a refresh
-			for(TPropObjectIterator Itor(RootNode->ObjectIterator()); Itor; ++Itor)
-			{
-				TWeakObjectPtr<UObject> Object = *Itor;
-				if(Object.IsValid())
-				{
-					NewObjectList.Add(Object.Get());
-				}
-			}
+			ValidSelectedObjects.Add(Object);
+			NewObjectList.Add(Object.Get());
 		}
 	}
 
-	SetObjectArrayPrivate( NewObjectList );
+	UnfilteredSelectedObjects = MoveTemp(ValidSelectedObjects);
+
+	SetObjectArrayPrivate(NewObjectList);
 }
 
 void SDetailsView::MoveScrollOffset(int32 DeltaOffset)
@@ -406,6 +402,16 @@ void SDetailsView::SetObjects(const TArray<UObject*>& InObjects, bool bForceRefr
 	{
 		if( bForceRefresh || ShouldSetNewObjects(InObjects) )
 		{
+			// Keep source object list around to reapply the object filter when it changes or force refresh.
+			UnfilteredSelectedObjects.Empty(InObjects.Num());
+			for (UObject* InObject : InObjects)
+			{
+				if (InObject)
+				{
+					UnfilteredSelectedObjects.Add(InObject);
+				}
+			}
+
 			SetObjectArrayPrivate(InObjects);
 		}
 	}
@@ -413,24 +419,17 @@ void SDetailsView::SetObjects(const TArray<UObject*>& InObjects, bool bForceRefr
 
 void SDetailsView::SetObjects(const TArray<TWeakObjectPtr<UObject>>& InObjects, bool bForceRefresh/* = false*/, bool bOverrideLock/* = false*/)
 {
-	if (!IsLocked() || bOverrideLock)
+	TArray<UObject*> SourceObjects;
+	SourceObjects.Reserve(InObjects.Num());
+	for (const TWeakObjectPtr<UObject>& Object : InObjects)
 	{
-		TArray<UObject*> SourceObjects;
-		SourceObjects.Reserve(InObjects.Num());
-
-		for (const TWeakObjectPtr<UObject>& Object : InObjects)
+		if (Object.IsValid())
 		{
-			if (Object.IsValid())
-			{
-				SourceObjects.Add(Object.Get());
-			}
-		}
-
-		if(bForceRefresh || ShouldSetNewObjects(SourceObjects))
-		{
-			SetObjectArrayPrivate(SourceObjects);
+			SourceObjects.Add(Object.Get());
 		}
 	}
+
+	SetObjects(SourceObjects, bForceRefresh, bOverrideLock);
 }
 
 void SDetailsView::SetObject(UObject* InObject, bool bForceRefresh)
@@ -445,35 +444,7 @@ void SDetailsView::SetObject(UObject* InObject, bool bForceRefresh)
 
 void SDetailsView::RemoveInvalidObjects()
 {
-	TArray<UObject*> ResetArray;
-
-	bool bAllFound = true;
-	
-	for(const TSharedPtr<FComplexPropertyNode>& ComplexRootNode : RootPropertyNodes)
-	{
-		FObjectPropertyNode* RootPropertyNode = ComplexRootNode->AsObjectNode();
-		if(RootPropertyNode)
-		{
-			for(TPropObjectIterator Itor(RootPropertyNode->ObjectIterator()); Itor; ++Itor)
-			{
-				TWeakObjectPtr<UObject> Object = *Itor;
-
-				if(Object.IsValid() && !Object->IsPendingKill())
-				{
-					ResetArray.Add(Object.Get());
-				}
-				else
-				{
-					bAllFound = false;
-				}
-			}
-		}
-	}
-
-	if (!bAllFound)
-	{
-		SetObjectArrayPrivate(ResetArray);
-	}
+	ForceRefresh();
 }
 
 void SDetailsView::SetObjectPackageOverrides(const TMap<TWeakObjectPtr<UObject>, TWeakObjectPtr<UPackage>>& InMapping)
@@ -720,46 +691,38 @@ void SDetailsView::SetObjectArrayPrivate(const TArray<UObject*>& InObjects)
 	double ElapsedTime = FPlatformTime::Seconds() - StartTime;
 }
 
-void SDetailsView::ReplaceObjects( const TMap<UObject*, UObject*>& OldToNewObjectMap )
+void SDetailsView::ReplaceObjects(const TMap<UObject*, UObject*>& OldToNewObjectMap)
 {
 	TArray<UObject*> NewObjectList;
-	bool bObjectsReplaced = false;
-
-	TArray<FObjectPropertyNode*> ObjectNodes;
-	for (TSharedPtr<FComplexPropertyNode>& RootNode : RootPropertyNodes)
-	{
-		PropertyEditorHelpers::CollectObjectNodes(RootNode, ObjectNodes);
-	}
+	NewObjectList.Reserve(UnfilteredSelectedObjects.Num());
+	TArray<TWeakObjectPtr<UObject>> NewUnfilteredSelectedObjects;
+	NewUnfilteredSelectedObjects.Reserve(UnfilteredSelectedObjects.Num());
 
-	for (FObjectPropertyNode* CurrentNode : ObjectNodes)
+	bool bNeedRefresh = false;
+	for (const TWeakObjectPtr<UObject>& Object : UnfilteredSelectedObjects)
 	{
-		// Scan all objects and look for objects which need to be replaced
-		for (int32 Idx = 0; Idx < CurrentNode->GetNumObjects(); ++Idx)
+		// We could be replacing an object that has already been garbage collected, so look up the object using the raw pointer.
+		UObject* Replacement = OldToNewObjectMap.FindRef(Object.GetEvenIfUnreachable());
+		if (Replacement)
 		{
-			// We could be replacing an object that has already been garbage collected, so look up the object using the raw pointer.
-			TWeakObjectPtr<UObject> Existing = CurrentNode->GetInstanceAsUObject(Idx);
-			UObject* Replacement = OldToNewObjectMap.FindRef(Existing.GetEvenIfUnreachable());
-
-			// Note: only root objects count for the new object list. 
-			// Sub-Objects (i.e components) count as needing to be replaced but they don't belong in the top level object list)
-			if (Replacement)
-			{
-				bObjectsReplaced = true;
-
-				if (CurrentNode->IsRootNode())
-				{
-					NewObjectList.Add(Replacement);
-				}
-			}
-			else if (CurrentNode->IsRootNode())
-			{
-				NewObjectList.Add(Existing.Get());
-			}
+			NewObjectList.Add(Replacement);
+			NewUnfilteredSelectedObjects.Add(Replacement);
+			bNeedRefresh = true;
+		}
+		else if (Object.IsValid())
+		{
+			NewObjectList.Add(Object.Get());
+			NewUnfilteredSelectedObjects.Add(Object);
+		}
+		else
+		{
+			bNeedRefresh = true;
 		}
 	}
 
-	if (bObjectsReplaced)
+	if (bNeedRefresh)
 	{
+		UnfilteredSelectedObjects = MoveTemp(NewUnfilteredSelectedObjects);
 		SetObjectArrayPrivate(NewObjectList);
 	}
 }
@@ -767,30 +730,22 @@ void SDetailsView::ReplaceObjects( const TMap<UObject*, UObject*>& OldToNewObjec
 void SDetailsView::RemoveDeletedObjects(const TArray<UObject*>& DeletedObjects)
 {
 	TArray<UObject*> NewObjectList;
-	bool bObjectsRemoved = false;
+	NewObjectList.Reserve(UnfilteredSelectedObjects.Num());
+	TArray<TWeakObjectPtr<UObject>> NewUnfilteredSelectedObjects;
+	NewUnfilteredSelectedObjects.Reserve(UnfilteredSelectedObjects.Num());
 
-	for(const TSharedPtr<FComplexPropertyNode>& ComplexRootNode : RootPropertyNodes)
+	for (const TWeakObjectPtr<UObject>& Object : UnfilteredSelectedObjects)
 	{
-		FObjectPropertyNode* RootPropertyNode = ComplexRootNode->AsObjectNode();
-		// Scan all objects and look for objects which need to be replaced
-		for (TPropObjectIterator Itor( RootPropertyNode->ObjectIterator() ); Itor; ++Itor)
+		if (Object.IsValid() && !DeletedObjects.Contains(Object.Get()))
 		{
-			if(DeletedObjects.Contains( Itor->Get()))
-			{
-				// An object we had needs to be removed
-				bObjectsRemoved = true;
-			}
-			else
-			{
-				// If the deleted object list does not contain the current object, its ok to keep it in the list
-				NewObjectList.Add( Itor->Get() );
-			}
+			NewUnfilteredSelectedObjects.Add(Object);
+			NewObjectList.Add(Object.Get());
 		}
 	}
 
-	// if any objects were replaced update the observed objects
-	if( bObjectsRemoved )
+	if (NewUnfilteredSelectedObjects.Num() != UnfilteredSelectedObjects.Num())
 	{
+		UnfilteredSelectedObjects = MoveTemp(NewUnfilteredSelectedObjects);
 		SetObjectArrayPrivate(NewObjectList);
 	}
 }
diff --git a/Engine/Source/Editor/PropertyEditor/Private/SDetailsView.h b/Engine/Source/Editor/PropertyEditor/Private/SDetailsView.h
--- a/Engine/Source/Editor/PropertyEditor/Private/SDetailsView.h
+++ b/Engine/Source/Editor/PropertyEditor/Private/SDetailsView.h
@@ -165,7 +165,10 @@ private:
 	/** Information about the current set of selected actors */
 	FSelectedActorInfo SelectedActorInfo;
 
-	/** Final set of selected objects for this detail view.  It may be different from the set passed in through SetObjects if there is an active filter */
+	/** Set of selected objects for this detail view that were passed in through SetObjects (before the object filter is applied). */
+	TArray<TWeakObjectPtr<UObject>> UnfilteredSelectedObjects;
+
+	/** Final set of selected objects for this detail view after applying the object filter. It may be different from the set passed in through SetObjects. */
 	TArray<TWeakObjectPtr<UObject>> SelectedObjects;
 
 	/** 
