PropertyEditor - Fix missing final commit change update for numerical values when releasing the mouse button on a value that already had the slider on.

diff --git a/Engine/Source/Editor/PropertyEditor/Private/UserInterface/PropertyEditor/SPropertyEditorNumeric.h b/Engine/Source/Editor/PropertyEditor/Private/UserInterface/PropertyEditor/SPropertyEditorNumeric.h
--- a/Engine/Source/Editor/PropertyEditor/Private/UserInterface/PropertyEditor/SPropertyEditorNumeric.h
+++ b/Engine/Source/Editor/PropertyEditor/Private/UserInterface/PropertyEditor/SPropertyEditorNumeric.h
@@ -533,7 +533,7 @@ private:
 	{
 		const TSharedRef< IPropertyHandle > PropertyHandle = PropertyEditor->GetPropertyHandle();
 		NumericType OrgValue(0);
-		if (PropertyHandle->GetValue(OrgValue) == FPropertyAccess::Fail || OrgValue != NewValue)
+		if (bIsUsingSlider || (PropertyHandle->GetValue(OrgValue) == FPropertyAccess::Fail || OrgValue != NewValue))
 		{
 			PropertyHandle->SetValue(NewValue);
 			LastSliderCommittedValue = NewValue;
