Removed bad namespace based enum forward declare which upsets clang.
To remove the enum property editors can now only create standalone mode property editor toolkits which is the only thing they supported anyway.

diff --git a/Engine/Source/Editor/ContentBrowser/Private/AssetContextMenu.cpp b/Engine/Source/Editor/ContentBrowser/Private/AssetContextMenu.cpp
--- a/Engine/Source/Editor/ContentBrowser/Private/AssetContextMenu.cpp
+++ b/Engine/Source/Editor/ContentBrowser/Private/AssetContextMenu.cpp
@@ -2244,7 +2244,7 @@ void FAssetContextMenu::ExecutePropertyMatrix()
 	if ( ObjectsForPropertiesMenu.Num() > 0 )
 	{
 		FPropertyEditorModule& PropertyEditorModule = FModuleManager::LoadModuleChecked<FPropertyEditorModule>( "PropertyEditor" );
-		PropertyEditorModule.CreatePropertyEditorToolkit( EToolkitMode::Standalone, TSharedPtr<IToolkitHost>(), ObjectsForPropertiesMenu );
+		PropertyEditorModule.CreatePropertyEditorToolkit(TSharedPtr<IToolkitHost>(), ObjectsForPropertiesMenu );
 	}
 }
 
diff --git a/Engine/Source/Editor/PropertyEditor/Private/PropertyEditorModule.cpp b/Engine/Source/Editor/PropertyEditor/Private/PropertyEditorModule.cpp
--- a/Engine/Source/Editor/PropertyEditor/Private/PropertyEditorModule.cpp
+++ b/Engine/Source/Editor/PropertyEditor/Private/PropertyEditorModule.cpp
@@ -398,18 +398,18 @@ FStructProperty* FPropertyEditorModule::RegisterStructOnScopeProperty(TSharedRef
 	return StructProperty;
 }
 
-TSharedRef< FAssetEditorToolkit > FPropertyEditorModule::CreatePropertyEditorToolkit( const EToolkitMode::Type Mode, const TSharedPtr< class IToolkitHost >& InitToolkitHost, UObject* ObjectToEdit )
+TSharedRef< FAssetEditorToolkit > FPropertyEditorModule::CreatePropertyEditorToolkit(const TSharedPtr< class IToolkitHost >& InitToolkitHost, UObject* ObjectToEdit )
 {
-	return FPropertyEditorToolkit::CreateEditor( Mode, InitToolkitHost, ObjectToEdit );
+	return FPropertyEditorToolkit::CreateEditor(EToolkitMode::Standalone, InitToolkitHost, ObjectToEdit);
 }
 
 
-TSharedRef< FAssetEditorToolkit > FPropertyEditorModule::CreatePropertyEditorToolkit( const EToolkitMode::Type Mode, const TSharedPtr< IToolkitHost >& InitToolkitHost, const TArray< UObject* >& ObjectsToEdit )
+TSharedRef< FAssetEditorToolkit > FPropertyEditorModule::CreatePropertyEditorToolkit( const TSharedPtr< IToolkitHost >& InitToolkitHost, const TArray< UObject* >& ObjectsToEdit )
 {
-	return FPropertyEditorToolkit::CreateEditor( Mode, InitToolkitHost, ObjectsToEdit );
+	return FPropertyEditorToolkit::CreateEditor(EToolkitMode::Standalone, InitToolkitHost, ObjectsToEdit);
 }
 
-TSharedRef< FAssetEditorToolkit > FPropertyEditorModule::CreatePropertyEditorToolkit( const EToolkitMode::Type Mode, const TSharedPtr< IToolkitHost >& InitToolkitHost, const TArray< TWeakObjectPtr< UObject > >& ObjectsToEdit )
+TSharedRef< FAssetEditorToolkit > FPropertyEditorModule::CreatePropertyEditorToolkit(const TSharedPtr< IToolkitHost >& InitToolkitHost, const TArray< TWeakObjectPtr< UObject > >& ObjectsToEdit )
 {
 	TArray< UObject* > RawObjectsToEdit;
 	for( auto ObjectIter = ObjectsToEdit.CreateConstIterator(); ObjectIter; ++ObjectIter )
@@ -417,7 +417,7 @@ TSharedRef< FAssetEditorToolkit > FPropertyEditorModule::CreatePropertyEditorToo
 		RawObjectsToEdit.Add( ObjectIter->Get() );
 	}
 
-	return FPropertyEditorToolkit::CreateEditor( Mode, InitToolkitHost, RawObjectsToEdit );
+	return FPropertyEditorToolkit::CreateEditor(EToolkitMode::Standalone, InitToolkitHost, RawObjectsToEdit );
 }
 
 TSharedRef<IPropertyChangeListener> FPropertyEditorModule::CreatePropertyChangeListener()
diff --git a/Engine/Source/Editor/PropertyEditor/Private/SDetailsView.cpp b/Engine/Source/Editor/PropertyEditor/Private/SDetailsView.cpp
--- a/Engine/Source/Editor/PropertyEditor/Private/SDetailsView.cpp
+++ b/Engine/Source/Editor/PropertyEditor/Private/SDetailsView.cpp
@@ -350,7 +350,7 @@ bool SDetailsView::CanOpenRawPropertyEditor() const
 FReply SDetailsView::OnOpenRawPropertyEditorClicked()
 {
 	FPropertyEditorModule& PropertyEditorModule = FModuleManager::LoadModuleChecked<FPropertyEditorModule>( "PropertyEditor" );
-	PropertyEditorModule.CreatePropertyEditorToolkit( EToolkitMode::Standalone, TSharedPtr<IToolkitHost>(), SelectedObjects );
+	PropertyEditorModule.CreatePropertyEditorToolkit(TSharedPtr<IToolkitHost>(), SelectedObjects );
 
 	return FReply::Handled();
 }
diff --git a/Engine/Source/Editor/PropertyEditor/Public/PropertyEditorModule.h b/Engine/Source/Editor/PropertyEditor/Public/PropertyEditorModule.h
--- a/Engine/Source/Editor/PropertyEditor/Public/PropertyEditorModule.h
+++ b/Engine/Source/Editor/PropertyEditor/Public/PropertyEditorModule.h
@@ -66,11 +66,6 @@ class IDetailsView;
 class IToolkitHost;
 class FAssetEditorToolkit;
 
-namespace EToolkitMode
-{
-	enum Type;
-};
-
 /**
  * Base class for adding an extra data to identify a custom property type
  */
@@ -272,9 +267,9 @@ public:
 	/**
 	 *
 	 */
-	virtual TSharedRef< FAssetEditorToolkit > CreatePropertyEditorToolkit( const EToolkitMode::Type Mode, const TSharedPtr< class IToolkitHost >& InitToolkitHost, UObject* ObjectToEdit );
-	virtual TSharedRef< FAssetEditorToolkit > CreatePropertyEditorToolkit( const EToolkitMode::Type Mode, const TSharedPtr< IToolkitHost >& InitToolkitHost, const TArray< UObject* >& ObjectsToEdit );
-	virtual TSharedRef< FAssetEditorToolkit > CreatePropertyEditorToolkit( const EToolkitMode::Type Mode, const TSharedPtr< IToolkitHost >& InitToolkitHost, const TArray< TWeakObjectPtr< UObject > >& ObjectsToEdit );
+	virtual TSharedRef< FAssetEditorToolkit > CreatePropertyEditorToolkit(const TSharedPtr< class IToolkitHost >& InitToolkitHost, UObject* ObjectToEdit );
+	virtual TSharedRef< FAssetEditorToolkit > CreatePropertyEditorToolkit(const TSharedPtr< IToolkitHost >& InitToolkitHost, const TArray< UObject* >& ObjectsToEdit );
+	virtual TSharedRef< FAssetEditorToolkit > CreatePropertyEditorToolkit(const TSharedPtr< IToolkitHost >& InitToolkitHost, const TArray< TWeakObjectPtr< UObject > >& ObjectsToEdit );
 
 	FPropertyTypeLayoutCallback GetPropertyTypeCustomization(const FProperty* InProperty,const IPropertyHandle& PropertyHandle, const FCustomPropertyTypeLayoutMap& InstancedPropertyTypeLayoutMap);
 	FPropertyTypeLayoutCallback FindPropertyTypeLayoutCallback(FName PropertyTypeName, const IPropertyHandle& PropertyHandle, const FCustomPropertyTypeLayoutMap& InstancedPropertyTypeLayoutMapp);
