Removed the "clear" action from material properties in actor details.

diff --git a/Engine/Source/Editor/PropertyEditor/Private/MaterialList.cpp b/Engine/Source/Editor/PropertyEditor/Private/MaterialList.cpp
--- a/Engine/Source/Editor/PropertyEditor/Private/MaterialList.cpp
+++ b/Engine/Source/Editor/PropertyEditor/Private/MaterialList.cpp
@@ -166,6 +166,7 @@ public:
 					[
 						SNew( SObjectPropertyEntryBox )
 						.ObjectPath(this, &FMaterialItemView::OnGetObjectPath)
+						.AllowClear(false)
 						.AllowedClass(UMaterialInterface::StaticClass())
 						.OnObjectChanged(this, &FMaterialItemView::OnSetObject)
 						.ThumbnailPool(ThumbnailPool)
