SDetailSingleItemRow's now have text for when they are being dragged into a place they cannot be dropped onto (like dropping an array element into a completely different array)

diff --git a/Engine/Source/Editor/PropertyEditor/Private/SDetailSingleItemRow.cpp b/Engine/Source/Editor/PropertyEditor/Private/SDetailSingleItemRow.cpp
--- a/Engine/Source/Editor/PropertyEditor/Private/SDetailSingleItemRow.cpp
+++ b/Engine/Source/Editor/PropertyEditor/Private/SDetailSingleItemRow.cpp
@@ -166,6 +166,27 @@ void SDetailSingleItemRow::OnArrayDragLeave(const FDragDropEvent& DragDropEvent)
 	bIsHoveredDragTarget = false;
 }
 
+bool SDetailSingleItemRow::CheckValidDrop(const TSharedPtr<SDetailSingleItemRow> RowPtr) const
+{
+	TSharedPtr<FPropertyNode> SwappingPropertyNode = RowPtr->SwappablePropertyNode;
+	if (SwappingPropertyNode.IsValid() && SwappablePropertyNode.IsValid())
+	{
+		if (SwappingPropertyNode != SwappablePropertyNode)
+		{
+			int32 OriginalIndex = SwappingPropertyNode->GetArrayIndex();
+			int32 NewIndex = SwappablePropertyNode->GetArrayIndex();
+			TSharedPtr<IPropertyHandle> SwappingHandle = PropertyEditorHelpers::GetPropertyHandle(SwappingPropertyNode.ToSharedRef(), OwnerTreeNode.Pin()->GetDetailsView()->GetNotifyHook(), OwnerTreeNode.Pin()->GetDetailsView()->GetPropertyUtilities());
+			TSharedPtr<IPropertyHandleArray> ParentHandle = SwappingHandle->GetParentHandle()->AsArray();
+
+			if (ParentHandle.IsValid() && SwappablePropertyNode->GetParentNode() == SwappingPropertyNode->GetParentNode())
+			{
+				return true;
+			}
+		}
+	}
+	return false;
+}
+
 FReply SDetailSingleItemRow::OnArrayDrop(const FDragDropEvent& DragDropEvent)
 {
 	bIsHoveredDragTarget = false;
@@ -179,44 +200,64 @@ FReply SDetailSingleItemRow::OnArrayDrop(const FDragDropEvent& DragDropEvent)
 	{
 		return FReply::Unhandled();
 	}
-	TSharedPtr<FPropertyNode> SwappingPropertyNode = RowPtr->SwappablePropertyNode;
-	if (SwappingPropertyNode.IsValid() && SwappablePropertyNode.IsValid())
+
+	if (CheckValidDrop(RowPtr))
 	{
-		if (SwappingPropertyNode != SwappablePropertyNode)
-		{
-			int32 OriginalIndex = SwappingPropertyNode->GetArrayIndex();
-			int32 NewIndex = SwappablePropertyNode->GetArrayIndex();
-			TSharedPtr<IPropertyHandle> SwappingHandle = PropertyEditorHelpers::GetPropertyHandle(SwappingPropertyNode.ToSharedRef(), OwnerTreeNode.Pin()->GetDetailsView()->GetNotifyHook(), OwnerTreeNode.Pin()->GetDetailsView()->GetPropertyUtilities());
-			TSharedPtr<IPropertyHandleArray> ParentHandle = SwappingHandle->GetParentHandle()->AsArray();
-			if (ParentHandle.IsValid() && SwappablePropertyNode->GetParentNode() == SwappingPropertyNode->GetParentNode())
-			{
-				// Need to swap the moving and target expansion states before saving
-				bool bOriginalSwappableExpansion = SwappablePropertyNode->HasNodeFlags(EPropertyNodeFlags::Expanded) != 0;
-				bool bOriginalSwappingExpansion = SwappingPropertyNode->HasNodeFlags(EPropertyNodeFlags::Expanded) != 0;
-				SwappablePropertyNode->SetNodeFlags(EPropertyNodeFlags::Expanded, bOriginalSwappingExpansion);
-				SwappingPropertyNode->SetNodeFlags(EPropertyNodeFlags::Expanded, bOriginalSwappableExpansion);
+		TSharedPtr<FPropertyNode> SwappingPropertyNode = RowPtr->SwappablePropertyNode;
+		TSharedPtr<IPropertyHandle> SwappingHandle = PropertyEditorHelpers::GetPropertyHandle(SwappingPropertyNode.ToSharedRef(), OwnerTreeNode.Pin()->GetDetailsView()->GetNotifyHook(), OwnerTreeNode.Pin()->GetDetailsView()->GetPropertyUtilities());
+		TSharedPtr<IPropertyHandleArray> ParentHandle = SwappingHandle->GetParentHandle()->AsArray();
+		int32 OriginalIndex = SwappingPropertyNode->GetArrayIndex();
+		int32 NewIndex = SwappablePropertyNode->GetArrayIndex();
 
-				IDetailsViewPrivate* DetailsView = OwnerTreeNode.Pin()->GetDetailsView();
-				DetailsView->SaveExpandedItems(SwappablePropertyNode->GetParentNodeSharedPtr().ToSharedRef());
-				FScopedTransaction Transaction(NSLOCTEXT("UnrealEd", "MoveRow", "Move Row"));
+		// Need to swap the moving and target expansion states before saving
+		bool bOriginalSwappableExpansion = SwappablePropertyNode->HasNodeFlags(EPropertyNodeFlags::Expanded) != 0;
+		bool bOriginalSwappingExpansion = SwappingPropertyNode->HasNodeFlags(EPropertyNodeFlags::Expanded) != 0;
+		SwappablePropertyNode->SetNodeFlags(EPropertyNodeFlags::Expanded, bOriginalSwappingExpansion);
+		SwappingPropertyNode->SetNodeFlags(EPropertyNodeFlags::Expanded, bOriginalSwappableExpansion);
 
-				SwappingHandle->GetParentHandle()->NotifyPreChange();
+		IDetailsViewPrivate* DetailsView = OwnerTreeNode.Pin()->GetDetailsView();
+		DetailsView->SaveExpandedItems(SwappablePropertyNode->GetParentNodeSharedPtr().ToSharedRef());
+		FScopedTransaction Transaction(NSLOCTEXT("UnrealEd", "MoveRow", "Move Row"));
 
-				ParentHandle->MoveElementTo(OriginalIndex, NewIndex);
+		SwappingHandle->GetParentHandle()->NotifyPreChange();
 
+		ParentHandle->MoveElementTo(OriginalIndex, NewIndex);
 
-				FPropertyChangedEvent MoveEvent(SwappingHandle->GetParentHandle()->GetProperty(), EPropertyChangeType::Unspecified);
-				SwappingHandle->GetParentHandle()->NotifyPostChange();
-				if (DetailsView->GetPropertyUtilities().IsValid())
-				{
-					DetailsView->GetPropertyUtilities()->NotifyFinishedChangingProperties(MoveEvent);
-				}
-			}
+		FPropertyChangedEvent MoveEvent(SwappingHandle->GetParentHandle()->GetProperty(), EPropertyChangeType::Unspecified);
+		SwappingHandle->GetParentHandle()->NotifyPostChange();
+		if (DetailsView->GetPropertyUtilities().IsValid())
+		{
+			DetailsView->GetPropertyUtilities()->NotifyFinishedChangingProperties(MoveEvent);
 		}
 	}
 	return FReply::Handled();
 }
 
+TOptional<EItemDropZone> SDetailSingleItemRow::OnArrayCanAcceptDrop(const FDragDropEvent& DragDropEvent, EItemDropZone DropZone, TSharedPtr< FDetailTreeNode > Type)
+{
+	bIsHoveredDragTarget = false;
+	TSharedPtr<FArrayRowDragDropOp> ArrayDropOp = DragDropEvent.GetOperationAs< FArrayRowDragDropOp >();
+	TSharedPtr<SDetailSingleItemRow> RowPtr = nullptr;
+	if (ArrayDropOp.IsValid() && ArrayDropOp->Row.IsValid())
+	{
+		RowPtr = ArrayDropOp->Row.Pin();
+	}
+	if (!RowPtr.IsValid())
+	{
+		return TOptional<EItemDropZone>();
+	}
+
+	if (CheckValidDrop(RowPtr))
+	{
+		ArrayDropOp->DecoratorText->SetText(NSLOCTEXT("ArrayDragDrop", "PlaceRowHere", "Place Row Here"));
+	}
+	else
+	{
+		ArrayDropOp->DecoratorText->SetText(NSLOCTEXT("ArrayDragDrop", "CannotPlaceRowHere", "Cannot Place Row Here"));
+	}
+	return TOptional<EItemDropZone>();
+}
+
 FReply SDetailSingleItemRow::OnArrayHeaderDrop(const FDragDropEvent& DragDropEvent)
 {
 	OnArrayDragLeave(DragDropEvent);
@@ -273,6 +314,7 @@ void SDetailSingleItemRow::Construct( const FArguments& InArgs, FDetailLayoutCus
 	FOnTableRowDragEnter ArrayDragDelegate;
 	FOnTableRowDragLeave ArrayDragLeaveDelegate;
 	FOnTableRowDrop ArrayDropDelegate;
+	FOnCanAcceptDrop ArrayAcceptDropDelegate;
 
 	const bool bIsValidTreeNode = InOwnerTreeNode->GetParentCategory().IsValid() && InOwnerTreeNode->GetParentCategory()->IsParentLayoutValid();
 	if(bIsValidTreeNode)
@@ -362,6 +404,7 @@ void SDetailSingleItemRow::Construct( const FArguments& InArgs, FDetailLayoutCus
 				ArrayDragDelegate = FOnTableRowDragEnter::CreateSP(this, &SDetailSingleItemRow::OnArrayDragEnter);
 				ArrayDragLeaveDelegate = FOnTableRowDragLeave::CreateSP(this, &SDetailSingleItemRow::OnArrayDragLeave);
 				ArrayDropDelegate = FOnTableRowDrop::CreateSP(this, &SDetailSingleItemRow::OnArrayDrop);
+				ArrayAcceptDropDelegate = FOnCanAcceptDrop::CreateSP(this, &SDetailSingleItemRow::OnArrayCanAcceptDrop);
 				SwappablePropertyNode = PropertyNode;
 			}
 			else if (PropertyNode.IsValid() 
@@ -371,6 +414,7 @@ void SDetailSingleItemRow::Construct( const FArguments& InArgs, FDetailLayoutCus
 				ArrayDragDelegate = FOnTableRowDragEnter::CreateSP(this, &SDetailSingleItemRow::OnArrayDragEnter);
 				ArrayDragLeaveDelegate = FOnTableRowDragLeave::CreateSP(this, &SDetailSingleItemRow::OnArrayDragLeave);
 				ArrayDropDelegate = FOnTableRowDrop::CreateSP(this, &SDetailSingleItemRow::OnArrayHeaderDrop);
+				ArrayAcceptDropDelegate = FOnCanAcceptDrop::CreateSP(this, &SDetailSingleItemRow::OnArrayCanAcceptDrop);
 			}
 			
 			InternalLeftColumnRowBox->AddSlot()
@@ -498,7 +542,8 @@ void SDetailSingleItemRow::Construct( const FArguments& InArgs, FDetailLayoutCus
 			.ShowSelection(false)
 			.OnDragEnter(ArrayDragDelegate)
 			.OnDragLeave(ArrayDragLeaveDelegate)
-			.OnDrop(ArrayDropDelegate),
+			.OnDrop(ArrayDropDelegate)
+			.OnCanAcceptDrop(ArrayAcceptDropDelegate),
 		InOwnerTableView
 	);
 }
@@ -863,7 +908,7 @@ FArrayRowDragDropOp::FArrayRowDragDropOp(TSharedPtr<SDetailSingleItemRow> InRow)
 			.AutoWidth()
 			.VAlign(VAlign_Center)
 			[
-				SNew(STextBlock)
+				SAssignNew(DecoratorText, STextBlock)
 				.Text(NSLOCTEXT("ArrayDragDrop", "PlaceRowHere", "Place Row Here"))
 			]
 		];
diff --git a/Engine/Source/Editor/PropertyEditor/Private/SDetailSingleItemRow.h b/Engine/Source/Editor/PropertyEditor/Private/SDetailSingleItemRow.h
--- a/Engine/Source/Editor/PropertyEditor/Private/SDetailSingleItemRow.h
+++ b/Engine/Source/Editor/PropertyEditor/Private/SDetailSingleItemRow.h
@@ -104,7 +104,13 @@ private:
 	FReply OnArrayDrop(const FDragDropEvent& DragDropEvent);
 	FReply OnArrayHeaderDrop(const FDragDropEvent& DragDropEvent);
 
+	TOptional<EItemDropZone> OnArrayCanAcceptDrop(const FDragDropEvent& DragDropEvent, EItemDropZone DropZone, TSharedPtr< FDetailTreeNode > Type);
+
 	TSharedPtr<FPropertyNode> GetCopyPastePropertyNode() const;
+
+	/** Checks if the current drop event is being dropped into a valid location
+	 */
+	bool CheckValidDrop(const TSharedPtr<SDetailSingleItemRow> RowPtr) const;
 private:
 	TWeakPtr<IDetailKeyframeHandler> KeyframeHandler;
 	/** Customization for this widget */
@@ -133,4 +139,5 @@ public:
 	}
 
 	TWeakPtr<class SDetailSingleItemRow> Row;
+	TSharedPtr<STextBlock> DecoratorText;
 };
