Added a details panel to the bulk property editor

diff --git a/Engine/Source/Editor/PropertyEditor/Private/PropertyEditorToolkit.cpp b/Engine/Source/Editor/PropertyEditor/Private/PropertyEditorToolkit.cpp
--- a/Engine/Source/Editor/PropertyEditor/Private/PropertyEditorToolkit.cpp
+++ b/Engine/Source/Editor/PropertyEditor/Private/PropertyEditorToolkit.cpp
@@ -22,6 +22,7 @@ const FName FPropertyEditorToolkit::ApplicationId( TEXT( "PropertyEditorToolkitA
 const FName FPropertyEditorToolkit::TreeTabId( TEXT( "PropertyEditorToolkit_PropertyTree" ) );
 const FName FPropertyEditorToolkit::GridTabId( TEXT( "PropertyEditorToolkit_PropertyTable" ) );
 const FName FPropertyEditorToolkit::TreePinAsColumnHeaderId( TEXT( "PropertyEditorToolkit_PinAsColumnHeader" ) );
+const FName FPropertyEditorToolkit::DetailsTabId(TEXT("PropertyEditorToolkit_DetailsPanel"));
 
 void FPropertyEditorToolkit::RegisterTabSpawners(const TSharedRef<class FTabManager>& InTabManager)
 {
@@ -33,21 +34,28 @@ void FPropertyEditorToolkit::RegisterTabSpawners(const TSharedRef<class FTabMana
 		.SetIcon(FSlateIcon(FEditorStyle::GetStyleSetName(), "LevelEditor.Tabs.Viewports"));
 
 	InTabManager->RegisterTabSpawner( TreeTabId, FOnSpawnTab::CreateSP(this, &FPropertyEditorToolkit::SpawnTab_PropertyTree) )
-		.SetDisplayName( LOCTEXT("PropertiesTab", "Details") )
+		.SetDisplayName( LOCTEXT("PropertiesTab", "Display") )
 		.SetGroup( WorkspaceMenuCategory.ToSharedRef() )
 		.SetIcon(FSlateIcon(FEditorStyle::GetStyleSetName(), "PropertyEditor.Grid.TabIcon"));
+
+	InTabManager->RegisterTabSpawner(DetailsTabId, FOnSpawnTab::CreateSP(this, &FPropertyEditorToolkit::SpawnTab_DetailsPanel))
+		.SetDisplayName(LOCTEXT("DetailsTab", "Property Editor"))
+		.SetGroup(WorkspaceMenuCategory.ToSharedRef())
+		.SetIcon(FSlateIcon(FEditorStyle::GetStyleSetName(), "LevelEditor.Tabs.Details"));
 }
 
 void FPropertyEditorToolkit::UnregisterTabSpawners(const TSharedRef<class FTabManager>& InTabManager)
 {
 	InTabManager->UnregisterTabSpawner( GridTabId );
 	InTabManager->UnregisterTabSpawner( TreeTabId );
+	InTabManager->UnregisterTabSpawner( DetailsTabId );
 }
 
 FPropertyEditorToolkit::FPropertyEditorToolkit()
 	: PropertyTree()
 	, PropertyTable()
 	, PathToRoot()
+	, DetailsView()
 {
 	PinSequence.AddCurve( 0, 1.0f, ECurveEaseFunction::QuadIn );
 }
@@ -137,11 +145,12 @@ void FPropertyEditorToolkit::Initialize( const EToolkitMode::Type Mode, const TS
 
 		CreatePropertyTree();
 		CreatePropertyTable();
+		CreateDetailsPanel();
 
 		PropertyTable->SetObjects(AdjustedObjectsToEdit);
 		TableColumnsChanged();
 
-		TSharedRef<FTabManager::FLayout> StandaloneDefaultLayout = FTabManager::NewLayout("Standalone_PropertyEditorToolkit_Layout")
+		TSharedRef<FTabManager::FLayout> StandaloneDefaultLayout = FTabManager::NewLayout("Standalone_PropertyEditorToolkit_Layout_v1")
 			->AddArea
 			(
 				FTabManager::NewPrimaryArea()->SetOrientation(Orient_Horizontal)
@@ -155,8 +164,9 @@ void FPropertyEditorToolkit::Initialize( const EToolkitMode::Type Mode, const TS
 				(
 					FTabManager::NewStack()
 					->SetSizeCoefficient(0.2f)
-					->SetHideTabWell(true)
 					->AddTab(TreeTabId, ETabState::OpenedTab)
+					->AddTab(DetailsTabId, ETabState::OpenedTab)
+					->SetForegroundTab(TreeTabId)
 				)
 			);
 
@@ -171,6 +181,8 @@ void FPropertyEditorToolkit::Initialize( const EToolkitMode::Type Mode, const TS
 		}
 		PropertyTree->SetObjectArray(AdjustedObjectsToEditWeak);
 
+		DetailsView->SetObjects(AdjustedObjectsToEditWeak, true);
+
 		PinColor = FSlateColor(FLinearColor(1, 1, 1, 0));
 		GEditor->GetTimerManager()->SetTimer(TimerHandle_TickPinColor, FTimerDelegate::CreateSP(this, &FPropertyEditorToolkit::TickPinColorAndOpacity), 0.1f, true);
 	}
@@ -183,7 +195,7 @@ TSharedRef<SDockTab> FPropertyEditorToolkit::SpawnTab_PropertyTree( const FSpawn
 
 	TSharedRef<SDockTab> TreeToolkitTab = SNew(SDockTab)
 		.Icon( FEditorStyle::GetBrush("PropertyEditor.Properties.TabIcon") )
-		.Label( LOCTEXT("GenericDetailsTitle", "Details") )
+		.Label( LOCTEXT("GenericDetailsTitle", "Display") )
 		.TabColorScale( GetTabColorScale() )
 		.Content()
 		[
@@ -260,6 +272,20 @@ TSharedRef<SDockTab> FPropertyEditorToolkit::SpawnTab_PropertyTable( const FSpaw
 	return GridToolkitTab;
 }
 
+TSharedRef<SDockTab> FPropertyEditorToolkit::SpawnTab_DetailsPanel(const FSpawnTabArgs& Args)
+{
+	check(Args.GetTabId() == DetailsTabId);
+
+	TSharedPtr<SDockTab> DetailsTab = SNew(SDockTab)
+		.Icon(FEditorStyle::GetBrush("LevelEditor.Tabs.Details"))
+		.Label(LOCTEXT("GenericDetailsPanel", "Property Editor"))
+		[
+			DetailsView.ToSharedRef()
+		];
+
+	return DetailsTab.ToSharedRef();
+}
+
 
 void FPropertyEditorToolkit::CreatePropertyTree()
 {
@@ -283,6 +309,12 @@ void FPropertyEditorToolkit::CreatePropertyTable()
 	PropertyTable->OnRootPathChanged()->AddSP( this, &FPropertyEditorToolkit::GridRootPathChanged );
 }
 
+void FPropertyEditorToolkit::CreateDetailsPanel()
+{
+	FPropertyEditorModule& PropertyEditorModule = FModuleManager::GetModuleChecked<FPropertyEditorModule>("PropertyEditor");
+	const FDetailsViewArgs DetailsViewArgs(false, false, true, FDetailsViewArgs::HideNameArea);
+	DetailsView = PropertyEditorModule.CreateDetailView(DetailsViewArgs);
+}
 
 void FPropertyEditorToolkit::ConstructTreeColumns( const TSharedRef< class SHeaderRow >& HeaderRow )
 {
@@ -478,6 +510,7 @@ void FPropertyEditorToolkit::GridSelectionChanged()
 	}
 
 	PropertyTree->SetObjectArray(SelectedRawObjects);
+	DetailsView->SetObjects(SelectedRawObjects, true);
 
 	const TSet< TSharedRef< IPropertyTableRow > > SelectedRows = PropertyTable->GetSelectedRows();
 
diff --git a/Engine/Source/Editor/PropertyEditor/Private/PropertyEditorToolkit.h b/Engine/Source/Editor/PropertyEditor/Private/PropertyEditorToolkit.h
--- a/Engine/Source/Editor/PropertyEditor/Private/PropertyEditorToolkit.h
+++ b/Engine/Source/Editor/PropertyEditor/Private/PropertyEditorToolkit.h
@@ -53,8 +53,6 @@ public:
 	static TSharedRef<FPropertyEditorToolkit> CreateEditor( const EToolkitMode::Type Mode, const TSharedPtr< IToolkitHost >& InitToolkitHost, UObject* ObjectToEdit );
 
 	static TSharedRef<FPropertyEditorToolkit> CreateEditor( const EToolkitMode::Type Mode, const TSharedPtr< IToolkitHost >& InitToolkitHost, const TArray<UObject*>& ObjectsToEdit );
-
-
 private:
 	static TSharedPtr<FPropertyEditorToolkit> FindExistingEditor( UObject* Object );
 
@@ -66,10 +64,14 @@ private:
 
 	void CreateGridView();
 
+	void CreateDetailsPanel();
+
 	TSharedRef<SDockTab> SpawnTab_PropertyTree( const FSpawnTabArgs& Args );
 
 	TSharedRef<SDockTab> SpawnTab_PropertyTable( const FSpawnTabArgs& Args ) ;
 
+	TSharedRef<SDockTab> SpawnTab_DetailsPanel(const FSpawnTabArgs& Args);
+
 	void GridSelectionChanged();
 
 	void GridRootPathChanged();
@@ -92,7 +94,6 @@ private:
 
 	EVisibility GetAddColumnInstructionsOverlayVisibility() const;
 
-
 private:
 
 	TSharedPtr< class SPropertyTreeViewImpl > PropertyTree;
@@ -100,6 +101,9 @@ private:
 
 	TSharedPtr< FPropertyPath > PathToRoot;
 
+	/** Details panel */
+	TSharedPtr<class IDetailsView> DetailsView;
+
 	TArray< TSharedRef< FPropertyPath > > PropertyPathsAddedAsColumns;
 
 	/** Animation sequence to pulse the pin image */
@@ -113,6 +117,7 @@ private:
 	static const FName ApplicationId;
 	static const FName TreeTabId;
 	static const FName GridTabId;
-
+	static const FName DetailsTabId;
 	static const FName TreePinAsColumnHeaderId;
+
 };
