Details Panel Lock Icon Alignment fix

diff --git a/Engine/Source/Editor/PropertyEditor/Private/SDetailNameArea.cpp b/Engine/Source/Editor/PropertyEditor/Private/SDetailNameArea.cpp
--- a/Engine/Source/Editor/PropertyEditor/Private/SDetailNameArea.cpp
+++ b/Engine/Source/Editor/PropertyEditor/Private/SDetailNameArea.cpp
@@ -164,6 +164,7 @@ TSharedRef< SWidget > SDetailNameArea::BuildObjectNameArea( const TArray< TWeakO
 		{
 			ObjectNameArea->AddSlot()
 				.HAlign(HAlign_Right)
+				.VAlign(VAlign_Center)
 				.AutoWidth()
 				[
 					SNew( SButton )
