Fixed fixed size array properties defined by enums not showing their display names

diff --git a/Engine/Source/Editor/PropertyEditor/Private/ItemPropertyNode.cpp b/Engine/Source/Editor/PropertyEditor/Private/ItemPropertyNode.cpp
--- a/Engine/Source/Editor/PropertyEditor/Private/ItemPropertyNode.cpp
+++ b/Engine/Source/Editor/PropertyEditor/Private/ItemPropertyNode.cpp
@@ -568,20 +568,19 @@ FText FItemPropertyNode::GetDisplayName() const
 							}
 						}
 					}
+				}
 
-					if (FinalDisplayName.IsEmpty())
+				if(FinalDisplayName.IsEmpty())
+				{
+					// This item is a member of an array, its display name is its index 
+					if (PropertyPtr == NULL || ArraySizeEnum == NULL)
 					{
 						FinalDisplayName = FText::AsNumber(GetArrayIndex());
 					}
-				}
-				// This item is a member of an array, its display name is its index 
-				else if (PropertyPtr == NULL || ArraySizeEnum == NULL)
-				{
-					FinalDisplayName = FText::AsNumber(GetArrayIndex());
-				}
-				else
-				{
-					FinalDisplayName = ArraySizeEnum->GetDisplayNameTextByIndex(GetArrayIndex());
+					else
+					{
+						FinalDisplayName = ArraySizeEnum->GetDisplayNameTextByIndex(GetArrayIndex());
+					}
 				}
 			}
 			// Maps should have display names that reflect the key and value types
