- Fixed a bug from WindowsNoEditor->Windows rename. Also hunted down other potential issues, and general cleanup of "NoEditor"

diff --git a/Engine/Source/Editor/MainFrame/Private/Frame/MainFrameActions.cpp b/Engine/Source/Editor/MainFrame/Private/Frame/MainFrameActions.cpp
--- a/Engine/Source/Editor/MainFrame/Private/Frame/MainFrameActions.cpp
+++ b/Engine/Source/Editor/MainFrame/Private/Frame/MainFrameActions.cpp
@@ -418,8 +418,8 @@ void FMainFrameActionCallbacks::CookContent(const FName InPlatformInfoName)
 	}
 	else
 	{
-		OptionalParams += TEXT(" -targetplatform=");
-		OptionalParams += *PlatformInfo->Name.ToString();
+		OptionalParams += TEXT(" -platform=");
+		OptionalParams += *PlatformInfo->DataDrivenPlatformInfo->UBTPlatformString;
 	}
 
 	OptionalParams += GetCookingOptionalParams();
diff --git a/Engine/Source/Editor/MainFrame/Private/Menus/MainMenu.cpp b/Engine/Source/Editor/MainFrame/Private/Menus/MainMenu.cpp
--- a/Engine/Source/Editor/MainFrame/Private/Menus/MainMenu.cpp
+++ b/Engine/Source/Editor/MainFrame/Private/Menus/MainMenu.cpp
@@ -506,8 +506,8 @@ void FMainMenu::RegisterFileProjectMenu()
 
 		if (RunningTargetPlatform != nullptr)
 		{
-			const FName CookedPlatformName = *(RunningTargetPlatform->PlatformName() + TEXT("NoEditor"));
-			const FText CookedPlatformText = FText::FromString(RunningTargetPlatform->PlatformName());
+			const FName CookedPlatformName = RunningTargetPlatform->GetTargetPlatformInfo().VanillaInfo->Name;
+			const FText CookedPlatformText = FText::FromString(CookedPlatformName.ToString());
 
 			FUIAction Action(
 				FExecuteAction::CreateStatic(&FMainFrameActionCallbacks::CookContent, CookedPlatformName),
diff --git a/Engine/Source/Editor/PropertyEditor/Private/SPerPlatformPropertiesWidget.cpp b/Engine/Source/Editor/PropertyEditor/Private/SPerPlatformPropertiesWidget.cpp
--- a/Engine/Source/Editor/PropertyEditor/Private/SPerPlatformPropertiesWidget.cpp
+++ b/Engine/Source/Editor/PropertyEditor/Private/SPerPlatformPropertiesWidget.cpp
@@ -59,27 +59,13 @@ void SPerPlatformPropertiesWidget::ConstructChildren()
 		TMultiMap<FName, FName> GroupToPlatform;
 						
 		// Create mapping from platform to platform groups and remove postfixes and invalid platform names
-		const TArray<FString> Filters = { TEXT("NoEditor"), TEXT("Client"), TEXT("Server") };
 		for (const FName& PlatformName : VanillaPlatformNameArray)
 		{
-			const PlatformInfo::FTargetPlatformInfo* PlatformInfo = PlatformInfo::FindPlatformInfo(PlatformName);
-			FString PlatformNameString = PlatformName.ToString();
-			for ( const FString& Filter : Filters)
-			{
-				const int32 Position = PlatformNameString.Find(Filter);
-				if (Position != INDEX_NONE)
-				{
-					PlatformNameString.RemoveAt(Position, Filter.Len());
-					break;
-				}
-			}
-
-			// Add filtered name if it isn't already set, and also add to group mapping
-			const FName FilteredName = FName(*PlatformNameString);				
-			if (PlatformNameString.Len() && !PlatformOverrides.Contains(FilteredName))
+			// Add platform name if it isn't already set, and also add to group mapping
+			if (!PlatformOverrides.Contains(PlatformName))
 			{	
-				BasePlatformNameArray.AddUnique(FilteredName);
-				GroupToPlatform.AddUnique(PlatformInfo->DataDrivenPlatformInfo->PlatformGroupName, FilteredName);
+				BasePlatformNameArray.AddUnique(PlatformName);
+				GroupToPlatform.AddUnique(PlatformInfo::FindPlatformInfo(PlatformName)->DataDrivenPlatformInfo->PlatformGroupName, PlatformName);
 			}
 		}
 
diff --git a/Engine/Source/Editor/UnrealEd/Private/Commandlets/NativeCodeGenCommandlet.cpp b/Engine/Source/Editor/UnrealEd/Private/Commandlets/NativeCodeGenCommandlet.cpp
--- a/Engine/Source/Editor/UnrealEd/Private/Commandlets/NativeCodeGenCommandlet.cpp
+++ b/Engine/Source/Editor/UnrealEd/Private/Commandlets/NativeCodeGenCommandlet.cpp
@@ -19,7 +19,7 @@ int32 UNativeCodeGenCommandlet::Main(const FString& Params)
 
 	if (Platforms.Num() == 0)
 	{
-		UE_LOG(LogNativeCodeGenCommandletCommandlet, Warning, TEXT("Missing platforms argument, should be -run=NativeCodeGen platform1 platform2, eg -run=NativeCodeGen windowsnoeditor"));
+		UE_LOG(LogNativeCodeGenCommandletCommandlet, Warning, TEXT("Missing platforms argument, should be -run=NativeCodeGen platform1 platform2, eg -run=NativeCodeGen windows"));
 		return 0;
 	}
 
diff --git a/Engine/Source/Editor/UnrealEd/Private/PlayLevelViaLauncher.cpp b/Engine/Source/Editor/UnrealEd/Private/PlayLevelViaLauncher.cpp
--- a/Engine/Source/Editor/UnrealEd/Private/PlayLevelViaLauncher.cpp
+++ b/Engine/Source/Editor/UnrealEd/Private/PlayLevelViaLauncher.cpp
@@ -456,20 +456,14 @@ void UEditorEngine::HandleStageStarted(const FString& InStage, TWeakPtr<SNotific
 	if (InStage.Contains(TEXT("Cooking")) || InStage.Contains(TEXT("Cook Task")))
 	{
 		FString PlatformName = LastPlayUsingLauncherDeviceId.Left(LastPlayUsingLauncherDeviceId.Find(TEXT("@")));
-		if (PlatformName.Contains(TEXT("NoEditor")))
-		{
-			PlatformName = PlatformName.Left(PlatformName.Find(TEXT("NoEditor")));
-		}
+		PlatformName = PlatformInfo::FindPlatformInfo(*PlatformName)->VanillaInfo->Name.ToString();
 		Arguments.Add(TEXT("PlatformName"), FText::FromString(PlatformName));
 		NotificationText = FText::Format(LOCTEXT("LauncherTaskProcessingNotification", "Processing Assets for {PlatformName}..."), Arguments);
 	}
 	else if (InStage.Contains(TEXT("Build Task")))
 	{
 		FString PlatformName = LastPlayUsingLauncherDeviceId.Left(LastPlayUsingLauncherDeviceId.Find(TEXT("@")));
-		if (PlatformName.Contains(TEXT("NoEditor")))
-		{
-			PlatformName = PlatformName.Left(PlatformName.Find(TEXT("NoEditor")));
-		}
+		PlatformName = PlatformInfo::FindPlatformInfo(*PlatformName)->VanillaInfo->Name.ToString();
 		Arguments.Add(TEXT("PlatformName"), FText::FromString(PlatformName));
 		if (!LauncherSessionInfo->bPlayUsingLauncherBuild)
 		{
diff --git a/Engine/Source/Programs/AutomationTool/Gauntlet/Unreal/Utils/Gauntlet.UnrealHelpers.cs b/Engine/Source/Programs/AutomationTool/Gauntlet/Unreal/Utils/Gauntlet.UnrealHelpers.cs
--- a/Engine/Source/Programs/AutomationTool/Gauntlet/Unreal/Utils/Gauntlet.UnrealHelpers.cs
+++ b/Engine/Source/Programs/AutomationTool/Gauntlet/Unreal/Utils/Gauntlet.UnrealHelpers.cs
@@ -121,22 +121,14 @@ namespace Gauntlet
 
 				if (ProcessType == UnrealTargetRole.Client)
 				{
-					if (UsesSharedBuildType)
-					{
-						Platform += "NoEditor";
-					}
-					else
+					if (!UsesSharedBuildType)
 					{
 						Platform += "Client";
 					}
 				}
 				else if (ProcessType == UnrealTargetRole.Server)
 				{
-					if (UsesSharedBuildType)
-					{
-						Platform += "NoEditor";
-					}
-					else
+					if (!UsesSharedBuildType)
 					{
 						Platform += "Server";
 					}
diff --git a/Engine/Source/Runtime/Core/Private/Misc/DataDrivenPlatformInfoRegistry.cpp b/Engine/Source/Runtime/Core/Private/Misc/DataDrivenPlatformInfoRegistry.cpp
--- a/Engine/Source/Runtime/Core/Private/Misc/DataDrivenPlatformInfoRegistry.cpp
+++ b/Engine/Source/Runtime/Core/Private/Misc/DataDrivenPlatformInfoRegistry.cpp
@@ -363,7 +363,7 @@ static void PrepForTurnkeyReport(FString& Command, FString& BaseCommandline, FSt
 
 static FString ConvertToDDPIPlatform(const FString& Platform)
 {
-	FString  New = Platform.Replace(TEXT("NoEditor"), TEXT("")).Replace(TEXT("Client"), TEXT("")).Replace(TEXT("Server"), TEXT(""));
+	FString  New = Platform.Replace(TEXT("Client"), TEXT("")).Replace(TEXT("Server"), TEXT(""));
 	if (New == TEXT("Win64"))
 	{
 		New = TEXT("Windows");
