Fix many customizations not displaying their name in the new right aligned format

diff --git a/Engine/Source/Editor/PropertyEditor/Public/DetailWidgetRow.h b/Engine/Source/Editor/PropertyEditor/Public/DetailWidgetRow.h
--- a/Engine/Source/Editor/PropertyEditor/Public/DetailWidgetRow.h
+++ b/Engine/Source/Editor/PropertyEditor/Public/DetailWidgetRow.h
@@ -95,7 +95,7 @@ public:
 	PROPERTYEDITOR_API const static float DefaultValueMaxWidth;
 
 	FDetailWidgetRow()
-		: NameWidget( *this, 0.0f, 0.0f, HAlign_Fill, VAlign_Center )
+		: NameWidget( *this, 0.0f, 0.0f, HAlign_Right, VAlign_Center )
 		, ValueWidget( *this, DefaultValueMinWidth, DefaultValueMaxWidth, HAlign_Left, VAlign_Fill )
 		, WholeRowWidget( *this, 0.0f, 0.0f, HAlign_Fill, VAlign_Fill )
 		, VisibilityAttr( EVisibility::Visible )
