Details view reskin phase 1.

The details view now displays gridlines between rows, and now has 4 columns - edit condition, name, value, and reset/extensions.

Edit condition widgets are now created in SDetailSingleItemRow and are placed in their own fixed column on the left of the details view. FEditConditionParser is now in FPropertyNode rather FPropertyEditor, which allows access to it more broadly - notably in SDetailSingleItemRow.

The Reset to Default SWidget can no longer be replaced - the logic for it being visible and a handler for when it is pressed is still customizable. SPropertyEditorAsset no longer has a special slot for the reset to default widget.

All extension widgets are now placed in a resizable column on the right of the panel, including the Keyframe button.

FDetailColumnSizeData is now publicly exposed by IDetailsView, and is no longer passed around as much.

diff --git a/Engine/Plugins/Enterprise/DataprepEditor/Source/DataprepEditor/Private/Widgets/DataprepAssetView.cpp b/Engine/Plugins/Enterprise/DataprepEditor/Source/DataprepEditor/Private/Widgets/DataprepAssetView.cpp
--- a/Engine/Plugins/Enterprise/DataprepEditor/Source/DataprepEditor/Private/Widgets/DataprepAssetView.cpp
+++ b/Engine/Plugins/Enterprise/DataprepEditor/Source/DataprepEditor/Private/Widgets/DataprepAssetView.cpp
@@ -403,11 +403,7 @@ void SDataprepAssetView::Construct( const FArguments& InArgs, UDataprepAssetInte
 
 	bIsChecked = true;
 
-	ColumnWidth = 0.7f;
-	ColumnSizeData = MakeShared< FDetailColumnSizeData >();
-	ColumnSizeData->LeftColumnWidth = TAttribute<float>(this, &SDataprepAssetView::OnGetLeftColumnWidth);
-	ColumnSizeData->RightColumnWidth = TAttribute<float>(this, &SDataprepAssetView::OnGetRightColumnWidth);
-	ColumnSizeData->OnWidthChanged = SSplitter::FOnSlotResized::CreateSP(this, &SDataprepAssetView::OnSetColumnWidth);
+	ColumnSizeData = MakeShared<FDetailColumnSizeData>();
 
 	UDataprepAssetProducers* AssetProducers = DataprepAssetInterfacePtr->GetProducers();
 	check( AssetProducers );
diff --git a/Engine/Plugins/Enterprise/DataprepEditor/Source/DataprepEditor/Private/Widgets/DataprepAssetView.h b/Engine/Plugins/Enterprise/DataprepEditor/Source/DataprepEditor/Private/Widgets/DataprepAssetView.h
--- a/Engine/Plugins/Enterprise/DataprepEditor/Source/DataprepEditor/Private/Widgets/DataprepAssetView.h
+++ b/Engine/Plugins/Enterprise/DataprepEditor/Source/DataprepEditor/Private/Widgets/DataprepAssetView.h
@@ -39,10 +39,6 @@ private:
 	/** Handles changes in the Dataprep asset */
 	void OnDataprepAssetChanged(FDataprepAssetChangeType ChangeType);
 
-	float OnGetLeftColumnWidth() const { return 1.0f - ColumnWidth; }
-	float OnGetRightColumnWidth() const { return ColumnWidth; }
-	void OnSetColumnWidth(float InWidth) { ColumnWidth = InWidth; }
-
 	TSharedRef<ITableRow> OnGenerateRowForCategoryTree( TSharedRef<EDataprepCategory> InTreeNode, const TSharedRef<STableViewBase>& InOwnerTable );
 	void OnGetChildrenForCategoryTree( TSharedRef<EDataprepCategory> InTreeNode, TArray< TSharedRef<EDataprepCategory> >& OutChildren ) {}
 
@@ -62,9 +58,6 @@ private:
 	/** Container used by all splitters in the details view, so that they move in sync */
 	TSharedPtr< FDetailColumnSizeData > ColumnSizeData;
 
-	/** Relative width to control splitters */
-	float ColumnWidth;
-
 	FDelegateHandle OnParameterizationWasEdited;
 };
 
diff --git a/Engine/Plugins/Enterprise/DataprepEditor/Source/DataprepEditor/Private/Widgets/DataprepWidgets.cpp b/Engine/Plugins/Enterprise/DataprepEditor/Source/DataprepEditor/Private/Widgets/DataprepWidgets.cpp
--- a/Engine/Plugins/Enterprise/DataprepEditor/Source/DataprepEditor/Private/Widgets/DataprepWidgets.cpp
+++ b/Engine/Plugins/Enterprise/DataprepEditor/Source/DataprepEditor/Private/Widgets/DataprepWidgets.cpp
@@ -595,19 +595,8 @@ void SDataprepDetailsView::Construct(const FArguments& InArgs)
 	bColumnPadding = InArgs._ColumnPadding;
 	bResizableColumn = InArgs._ResizableColumn;
 
-	if (InArgs._ColumnSizeData.IsValid())
-	{
-		ColumnSizeData = InArgs._ColumnSizeData;
-	}
-	else
-	{
-		ColumnWidth = 0.7f;
-		ColumnSizeData = MakeShared<FDetailColumnSizeData>();
-		ColumnSizeData->LeftColumnWidth = TAttribute<float>(this, &SDataprepDetailsView::OnGetLeftColumnWidth);
-		ColumnSizeData->RightColumnWidth = TAttribute<float>(this, &SDataprepDetailsView::OnGetRightColumnWidth);
-		ColumnSizeData->OnWidthChanged = SSplitter::FOnSlotResized::CreateSP(this, &SDataprepDetailsView::OnSetColumnWidth);
-	}
-	
+	ColumnSizeData = MakeShared<FDetailColumnSizeData>();
+
 	FPropertyEditorModule& PropertyEditorModule = FModuleManager::GetModuleChecked<FPropertyEditorModule>("PropertyEditor");
 
 	FPropertyRowGeneratorArgs Args;
@@ -794,20 +783,6 @@ void SDataprepInstanceParentWidget::Construct(const FArguments& InArgs)
 		return;
 	}
 
-	if (InArgs._ColumnSizeData.IsValid())
-	{
-		ColumnSizeData = InArgs._ColumnSizeData;
-	}
-	else
-	{
-		ColumnWidth = 0.5f;
-		ColumnSizeData = MakeShared<FDetailColumnSizeData>();
-		ColumnSizeData->LeftColumnWidth = TAttribute<float>(this, &SDataprepInstanceParentWidget::OnGetLeftColumnWidth);
-		ColumnSizeData->RightColumnWidth = TAttribute<float>(this, &SDataprepInstanceParentWidget::OnGetRightColumnWidth);
-		ColumnSizeData->OnWidthChanged = SSplitter::FOnSlotResized::CreateSP(this, &SDataprepInstanceParentWidget::OnSetColumnWidth);
-	}
-
-
 	TSharedRef<SWidget> NameWidget = SNew(SHorizontalBox)
 	.Clipping(EWidgetClipping::OnDemand)
 	+ SHorizontalBox::Slot()
diff --git a/Engine/Plugins/Enterprise/DataprepEditor/Source/DataprepEditor/Private/Widgets/DataprepWidgets.h b/Engine/Plugins/Enterprise/DataprepEditor/Source/DataprepEditor/Private/Widgets/DataprepWidgets.h
--- a/Engine/Plugins/Enterprise/DataprepEditor/Source/DataprepEditor/Private/Widgets/DataprepWidgets.h
+++ b/Engine/Plugins/Enterprise/DataprepEditor/Source/DataprepEditor/Private/Widgets/DataprepWidgets.h
@@ -67,7 +67,6 @@ class SDataprepCategoryWidget : public STableRow< TSharedPtr< EDataprepCategory
 public:
 
 	SLATE_BEGIN_ARGS(SDataprepCategoryWidget) {}
-	SLATE_ARGUMENT(TSharedPtr< FDetailColumnSizeData >, ColumnSizeData)
 	SLATE_ARGUMENT(FText, Title)
 	SLATE_ARGUMENT(TSharedPtr< SWidget >, TitleDetail)
 	SLATE_END_ARGS()
@@ -95,7 +94,6 @@ public:
 		, _ColumnPadding(false)
 		, _ResizableColumn(true)
 		{}
-		SLATE_ARGUMENT(TSharedPtr<FDetailColumnSizeData>, ColumnSizeData)
 		SLATE_ARGUMENT(UObject*, Object)
 		SLATE_ARGUMENT(float, Spacing)
 		SLATE_ARGUMENT(bool, ColumnPadding)
@@ -149,17 +147,6 @@ private:
 	 */
 	void CreateDefaultWidget(int32 Index, TSharedPtr< SWidget >& NameWidget, TSharedPtr< SWidget >& ValueWidget, float LeftPadding, const FDataprepParameterizationContext& ParameterizationContext);
 
-	/** Callback used by all splitters in the details view, so that they move in sync */
-	void OnLeftColumnResized(float InNewWidth)
-	{
-		// This has to be bound or the splitter will take it upon itself to determine the size
-		// We do nothing here because it is handled by the column size data
-	}
-
-	float OnGetLeftColumnWidth() const { return 1.0f - ColumnWidth; }
-	float OnGetRightColumnWidth() const { return ColumnWidth; }
-	void OnSetColumnWidth(float InWidth) { ColumnWidth = InWidth < 0.5f ? 0.5f : InWidth; }
-
 	/** Callback to track property changes on array properties */
 	void OnPropertyChanged( const struct FPropertyChangedEvent& InEvent );
 
@@ -192,9 +179,6 @@ private:
 	/** Delegate handle to track when a object was transacted */
 	FDelegateHandle OnObjectTransactedHandle;
 
-	/** Relative width to control splitters. */
-	float ColumnWidth;
-
 	/** Points to the currently used column size data. Can be provided via argument as well. */
 	TSharedPtr< FDetailColumnSizeData > ColumnSizeData;
 
diff --git a/Engine/Plugins/Enterprise/VariantManager/Source/VariantManager/Private/SVariantManager.cpp b/Engine/Plugins/Enterprise/VariantManager/Source/VariantManager/Private/SVariantManager.cpp
--- a/Engine/Plugins/Enterprise/VariantManager/Source/VariantManager/Private/SVariantManager.cpp
+++ b/Engine/Plugins/Enterprise/VariantManager/Source/VariantManager/Private/SVariantManager.cpp
@@ -389,11 +389,7 @@ void SVariantManager::Construct(const FArguments& InArgs, TSharedRef<FVariantMan
 				.Orientation(Orient_Horizontal)
 				+ SSplitter::Slot()
 				.Value(ColumnSizeData.LeftColumnWidth)
-				.OnSlotResized(SSplitter::FOnSlotResized::CreateLambda([](float InNewWidth)
-				{
-					//This has to be bound or the splitter will take it upon itself to determine the size
-					//We do nothing here because it is handled by the column size data
-				}))
+				.OnSlotResized(ColumnSizeData.OnNameColumnResized)
 				[
 					SNew(SBox)
 					.HeightOverride(CommonHeaderMaxHeight)
@@ -447,7 +443,7 @@ void SVariantManager::Construct(const FArguments& InArgs, TSharedRef<FVariantMan
 				]
 				+ SSplitter::Slot()
 				.Value(ColumnSizeData.RightColumnWidth)
-				.OnSlotResized(ColumnSizeData.OnWidthChanged)
+				.OnSlotResized(ColumnSizeData.OnValueColumnResized)
 				[
 					SNew(SBox)
 					.HeightOverride(CommonHeaderMaxHeight)
diff --git a/Engine/Source/Editor/AnimGraph/Private/AnimBlueprintPinInfoDetails.cpp b/Engine/Source/Editor/AnimGraph/Private/AnimBlueprintPinInfoDetails.cpp
--- a/Engine/Source/Editor/AnimGraph/Private/AnimBlueprintPinInfoDetails.cpp
+++ b/Engine/Source/Editor/AnimGraph/Private/AnimBlueprintPinInfoDetails.cpp
@@ -45,7 +45,7 @@ void FAnimBlueprintFunctionPinInfoDetails::CustomizeHeader(TSharedRef<IPropertyH
 	StructPropertyHandle->MarkResetToDefaultCustomized(true);
 
 	InHeaderRow
-	.DiffersFromDefault(false)
+	.OverrideResetToDefault(FResetToDefaultOverride::Hide())
 	.NameContent()
 	[
 		SNew(SHorizontalBox)
diff --git a/Engine/Source/Editor/DetailCustomizations/Private/ComponentTransformDetails.cpp b/Engine/Source/Editor/DetailCustomizations/Private/ComponentTransformDetails.cpp
--- a/Engine/Source/Editor/DetailCustomizations/Private/ComponentTransformDetails.cpp
+++ b/Engine/Source/Editor/DetailCustomizations/Private/ComponentTransformDetails.cpp
@@ -320,6 +320,7 @@ void FComponentTransformDetails::GenerateChildContent( IDetailChildrenBuilder& C
 		ChildrenBuilder.AddCustomRow( LOCTEXT("LocationFilter", "Location") )
 		.CopyAction( CreateCopyAction( ETransformField::Location ) )
 		.PasteAction( CreatePasteAction( ETransformField::Location ) )
+		.OverrideResetToDefault(FResetToDefaultOverride::Create(TAttribute<bool>(this, &FComponentTransformDetails::GetLocationResetVisibility), FSimpleDelegate::CreateSP(this, &FComponentTransformDetails::OnLocationResetClicked)))
 		.NameContent()
 		.VAlign(VAlign_Center)
 		[
@@ -354,29 +355,6 @@ void FComponentTransformDetails::GenerateChildContent( IDetailChildrenBuilder& C
 				.OnBeginSliderMovement( this, &FComponentTransformDetails::OnBeginLocationSlider )
 				.OnEndSliderMovement( this, &FComponentTransformDetails::OnEndLocationSlider )
 			]
-			+SHorizontalBox::Slot()
-			.AutoWidth()
-			[
-				// Just take up space for alignment
-				SNew( SBox )
-				.WidthOverride( 18.0f )
-			]
-			+SHorizontalBox::Slot()
-			.VAlign(VAlign_Center)
-			.AutoWidth()
-			[
-				SNew(SButton)
-				.OnClicked(this, &FComponentTransformDetails::OnLocationResetClicked)
-				.Visibility(this, &FComponentTransformDetails::GetLocationResetVisibility)
-				.ContentPadding(FMargin(5.f, 0.f))
-				.ToolTipText(LOCTEXT("ResetToDefaultToolTip", "Reset to Default"))
-				.ButtonStyle( FEditorStyle::Get(), "NoBorder" )
-				.Content()
-				[
-					SNew(SImage)
-					.Image( FEditorStyle::GetBrush("PropertyWindow.DiffersFromDefault") )
-				]
-			]
 		];
 	}
 	
@@ -392,6 +370,7 @@ void FComponentTransformDetails::GenerateChildContent( IDetailChildrenBuilder& C
 		ChildrenBuilder.AddCustomRow( LOCTEXT("RotationFilter", "Rotation") )
 		.CopyAction( CreateCopyAction(ETransformField::Rotation) )
 		.PasteAction( CreatePasteAction(ETransformField::Rotation) )
+		.OverrideResetToDefault(FResetToDefaultOverride::Create(TAttribute<bool>(this, &FComponentTransformDetails::GetRotationResetVisibility), FSimpleDelegate::CreateSP(this, &FComponentTransformDetails::OnRotationResetClicked)))
 		.NameContent()
 		.VAlign(VAlign_Center)
 		[
@@ -425,29 +404,6 @@ void FComponentTransformDetails::GenerateChildContent( IDetailChildrenBuilder& C
 				.TypeInterface( TypeInterface )
 				.Font( FontInfo )
 			]
-			+SHorizontalBox::Slot()
-			.AutoWidth()
-			[
-				// Just take up space for alignment
-				SNew( SBox )
-				.WidthOverride( 18.0f )
-			]
-			+SHorizontalBox::Slot()
-			.VAlign(VAlign_Center)
-			.AutoWidth()
-			[
-				SNew(SButton)
-				.OnClicked(this, &FComponentTransformDetails::OnRotationResetClicked)
-				.Visibility(this, &FComponentTransformDetails::GetRotationResetVisibility)
-				.ContentPadding(FMargin(5.f, 0.f))
-				.ToolTipText(LOCTEXT("ResetToDefaultToolTip", "Reset to Default"))
-				.ButtonStyle( FEditorStyle::Get(), "NoBorder" )
-				.Content()
-				[
-					SNew(SImage)
-					.Image( FEditorStyle::GetBrush("PropertyWindow.DiffersFromDefault") )
-				]
-			]
 		];
 	}
 	
@@ -457,6 +413,7 @@ void FComponentTransformDetails::GenerateChildContent( IDetailChildrenBuilder& C
 		ChildrenBuilder.AddCustomRow( LOCTEXT("ScaleFilter", "Scale") )
 		.CopyAction( CreateCopyAction(ETransformField::Scale) )
 		.PasteAction( CreatePasteAction(ETransformField::Scale) )
+		.OverrideResetToDefault(FResetToDefaultOverride::Create(TAttribute<bool>(this, &FComponentTransformDetails::GetScaleResetVisibility), FSimpleDelegate::CreateSP(this, &FComponentTransformDetails::OnScaleResetClicked)))
 		.NameContent()
 		.VAlign(VAlign_Center)
 		[
@@ -510,22 +467,6 @@ void FComponentTransformDetails::GenerateChildContent( IDetailChildrenBuilder& C
 					.ColorAndOpacity( FSlateColor::UseForeground() )
 				]
 			]
-			+SHorizontalBox::Slot()
-			.VAlign(VAlign_Center)
-			.AutoWidth()
-			[
-				SNew(SButton)
-				.OnClicked(this, &FComponentTransformDetails::OnScaleResetClicked)
-				.Visibility(this, &FComponentTransformDetails::GetScaleResetVisibility)
-				.ContentPadding(FMargin(5.f, 0.f))
-				.ToolTipText(LOCTEXT("ResetToDefaultToolTip", "Reset to Default"))
-				.ButtonStyle( FEditorStyle::Get(), "NoBorder" )
-				.Content()
-				[
-					SNew(SImage)
-					.Image( FEditorStyle::GetBrush("PropertyWindow.DiffersFromDefault") )
-				]
-			]
 		];
 	}
 }
@@ -753,16 +694,16 @@ struct FGetRootComponentArchetype
 	}
 };
 
-EVisibility FComponentTransformDetails::GetLocationResetVisibility() const
+bool FComponentTransformDetails::GetLocationResetVisibility() const
 {
 	const USceneComponent* Archetype = FGetRootComponentArchetype::Get(SelectedObjects[0].Get());
 	const FVector Data = Archetype ? Archetype->GetRelativeLocation() : FVector::ZeroVector;
 
 	// unset means multiple differing values, so show "Reset to Default" in that case
-	return CachedLocation.IsSet() && CachedLocation.X.GetValue() == Data.X && CachedLocation.Y.GetValue() == Data.Y && CachedLocation.Z.GetValue() == Data.Z ? EVisibility::Hidden : EVisibility::Visible;
+	return CachedLocation.IsSet() && CachedLocation.X.GetValue() == Data.X && CachedLocation.Y.GetValue() == Data.Y && CachedLocation.Z.GetValue() == Data.Z ? false : true;
 }
 
-FReply FComponentTransformDetails::OnLocationResetClicked()
+void FComponentTransformDetails::OnLocationResetClicked()
 {
 	const FText TransactionName = LOCTEXT("ResetLocation", "Reset Location");
 	FScopedTransaction Transaction(TransactionName);
@@ -771,20 +712,18 @@ FReply FComponentTransformDetails::OnLocationResetClicked()
 	const FVector Data = Archetype ? Archetype->GetRelativeLocation() : FVector::ZeroVector;
 
 	OnSetTransform(ETransformField::Location, EAxisList::All, Data, false, true);
-
-	return FReply::Handled();
 }
 
-EVisibility FComponentTransformDetails::GetRotationResetVisibility() const
+bool FComponentTransformDetails::GetRotationResetVisibility() const
 {
 	const USceneComponent* Archetype = FGetRootComponentArchetype::Get(SelectedObjects[0].Get());
 	const FVector Data = Archetype ? Archetype->GetRelativeRotation().Euler() : FVector::ZeroVector;
 
 	// unset means multiple differing values, so show "Reset to Default" in that case
-	return CachedRotation.IsSet() && CachedRotation.X.GetValue() == Data.X && CachedRotation.Y.GetValue() == Data.Y && CachedRotation.Z.GetValue() == Data.Z ? EVisibility::Hidden : EVisibility::Visible;
+	return CachedRotation.IsSet() && CachedRotation.X.GetValue() == Data.X && CachedRotation.Y.GetValue() == Data.Y && CachedRotation.Z.GetValue() == Data.Z ? false : true;
 }
 
-FReply FComponentTransformDetails::OnRotationResetClicked()
+void FComponentTransformDetails::OnRotationResetClicked()
 {
 	const FText TransactionName = LOCTEXT("ResetRotation", "Reset Rotation");
 	FScopedTransaction Transaction(TransactionName);
@@ -793,20 +732,18 @@ FReply FComponentTransformDetails::OnRotationResetClicked()
 	const FVector Data = Archetype ? Archetype->GetRelativeRotation().Euler() : FVector::ZeroVector;
 
 	OnSetTransform(ETransformField::Rotation, EAxisList::All, Data, false, true);
-
-	return FReply::Handled();
 }
 
-EVisibility FComponentTransformDetails::GetScaleResetVisibility() const
+bool FComponentTransformDetails::GetScaleResetVisibility() const
 {
 	const USceneComponent* Archetype = FGetRootComponentArchetype::Get(SelectedObjects[0].Get());
 	const FVector Data = Archetype ? Archetype->GetRelativeScale3D() : FVector(1.0f);
 
 	// unset means multiple differing values, so show "Reset to Default" in that case
-	return CachedScale.IsSet() && CachedScale.X.GetValue() == Data.X && CachedScale.Y.GetValue() == Data.Y && CachedScale.Z.GetValue() == Data.Z ? EVisibility::Hidden : EVisibility::Visible;
+	return CachedScale.IsSet() && CachedScale.X.GetValue() == Data.X && CachedScale.Y.GetValue() == Data.Y && CachedScale.Z.GetValue() == Data.Z ? false : true;
 }
 
-FReply FComponentTransformDetails::OnScaleResetClicked()
+void FComponentTransformDetails::OnScaleResetClicked()
 {
 	const FText TransactionName = LOCTEXT("ResetScale", "Reset Scale");
 	FScopedTransaction Transaction(TransactionName);
@@ -815,8 +752,6 @@ FReply FComponentTransformDetails::OnScaleResetClicked()
 	const FVector Data = Archetype ? Archetype->GetRelativeScale3D() : FVector(1.0f);
 
 	OnSetTransform(ETransformField::Scale, EAxisList::All, Data, false, true);
-
-	return FReply::Handled();
 }
 
 void FComponentTransformDetails::ExtendXScaleContextMenu( FMenuBuilder& MenuBuilder )
diff --git a/Engine/Source/Editor/DetailCustomizations/Private/ComponentTransformDetails.h b/Engine/Source/Editor/DetailCustomizations/Private/ComponentTransformDetails.h
--- a/Engine/Source/Editor/DetailCustomizations/Private/ComponentTransformDetails.h
+++ b/Engine/Source/Editor/DetailCustomizations/Private/ComponentTransformDetails.h
@@ -183,13 +183,13 @@ private:
 	FUIAction CreatePasteAction( ETransformField::Type TransformField ) const;
 
 	/** Called when the "Reset to Default" button for the location has been clicked */
-	FReply OnLocationResetClicked();
+	void OnLocationResetClicked();
 
 	/** Called when the "Reset to Default" button for the rotation has been clicked */
-	FReply OnRotationResetClicked();
+	void OnRotationResetClicked();
 
 	/** Called when the "Reset to Default" button for the scale has been clicked */
-	FReply OnScaleResetClicked();
+	void OnScaleResetClicked();
 
 	/** Extend the context menu for the X component */
 	void ExtendXScaleContextMenu( FMenuBuilder& MenuBuilder );
@@ -212,7 +212,7 @@ private:
 	/** @return The Z component of location */
 	TOptional<float> GetLocationZ() const { return CachedLocation.Z; }
 	/** @return The visibility of the "Reset to Default" button for the location component */
-	EVisibility GetLocationResetVisibility() const;
+	bool GetLocationResetVisibility() const;
 
 	/** @return The X component of rotation */
 	TOptional<float> GetRotationX() const { return CachedRotation.X; }
@@ -221,7 +221,7 @@ private:
 	/** @return The Z component of rotation */
 	TOptional<float> GetRotationZ() const { return CachedRotation.Z; }
 	/** @return The visibility of the "Reset to Default" button for the rotation component */
-	EVisibility GetRotationResetVisibility() const;
+	bool GetRotationResetVisibility() const;
 
 	/** @return The X component of scale */
 	TOptional<float> GetScaleX() const { return CachedScale.X; }
@@ -230,7 +230,7 @@ private:
 	/** @return The Z component of scale */
 	TOptional<float> GetScaleZ() const { return CachedScale.Z; }
 	/** @return The visibility of the "Reset to Default" button for the scale component */
-	EVisibility GetScaleResetVisibility() const;
+	bool GetScaleResetVisibility() const;
 
 	/** Cache a single unit to display all location components in */
 	void CacheCommonLocationUnits();
diff --git a/Engine/Source/Editor/EditorStyle/Private/StarshipStyle.cpp b/Engine/Source/Editor/EditorStyle/Private/StarshipStyle.cpp
--- a/Engine/Source/Editor/EditorStyle/Private/StarshipStyle.cpp
+++ b/Engine/Source/Editor/EditorStyle/Private/StarshipStyle.cpp
@@ -3169,9 +3169,9 @@ void FStarshipEditorStyle::FStyle::SetupPropertyEditorStyles()
 		Set( "DetailsView.ViewOptions",         	  new CORE_IMAGE_BRUSH_SVG("Starship/Common/settings",  Icon16x16, FLinearColor::White) );
 		Set( "DetailsView.EditConfigProperties",      new IMAGE_BRUSH("Icons/icon_PropertyMatrix_16px",  Icon16x16, FLinearColor::White ) );
 
-		Set( "DetailsView.CollapsedCategory",         new FSlateColorBrush(FStyleColors::Background));
-		Set( "DetailsView.CategoryTop",               new FSlateColorBrush(FStyleColors::Background));
+		Set( "DetailsView.CollapsedCategory",         new FSlateColorBrush(FStyleColors::Header));
 		Set( "DetailsView.CollapsedCategory_Hovered", new FSlateColorBrush(FStyleColors::Header));
+		Set( "DetailsView.CategoryTop",               new FSlateColorBrush(FStyleColors::Header));
 		Set( "DetailsView.CategoryTop_Hovered",       new FSlateColorBrush(FStyleColors::Header));
 		Set( "DetailsView.CategoryBottom",            new FSlateColorBrush(FStyleColors::Recessed));
 
@@ -3179,18 +3179,25 @@ void FStarshipEditorStyle::FStyle::SetupPropertyEditorStyles()
 		Set( "DetailsView.CategoryMiddle_Hovered",    new FSlateColorBrush(FStyleColors::Header));
 		Set( "DetailsView.CategoryMiddle_Highlighted",new FSlateColorBrush(FStyleColors::Background));
 		Set( "DetailsView.CategoryMiddle_Active",     new FSlateColorBrush(FStyleColors::Background));
-		
+
 		Set( "DetailsView.PropertyIsFavorite", new IMAGE_BRUSH("PropertyView/Favorites_Enabled", Icon12x12));
 		Set( "DetailsView.PropertyIsNotFavorite", new IMAGE_BRUSH("PropertyView/Favorites_Disabled", Icon12x12));
 		Set( "DetailsView.NoFavoritesSystem", new IMAGE_BRUSH("PropertyView/NoFavoritesSystem", Icon12x12));
 
 		Set( "DetailsView.Splitter", FSplitterStyle()
-			.SetHandleNormalBrush(   FSlateNoResource())                   
+			.SetHandleNormalBrush(FSlateColorBrush(FStyleColors::Transparent))                   
+			.SetHandleHighlightBrush(FSlateColorBrush(FStyleColors::Recessed))
+		);
+
+		Set( "DetailsView.Splitter.Outer", FSplitterStyle()
+			.SetHandleNormalBrush(FSlateColorBrush(FStyleColors::Recessed))                   
 			.SetHandleHighlightBrush(FSlateColorBrush(FStyleColors::Recessed))
 		);
 
-		Set( "DetailsView.AdvancedDropdownBorder",      new FSlateColorBrush(FStyleColors::Background));
-		Set( "DetailsView.AdvancedDropdownBorder.Open", new FSlateColorBrush(FStyleColors::Background));
+		Set("DetailsView.GridLine", new FSlateColorBrush(FStyleColors::Recessed));
+
+		Set( "DetailsView.AdvancedDropdownBorder",      new FSlateColorBrush(FStyleColors::Recessed));
+		Set( "DetailsView.AdvancedDropdownBorder.Open", new FSlateColorBrush(FStyleColors::Recessed));
 
 		Set( "DetailsView.CategoryFontStyle", FStyleFonts::Get().NormalBold);
 		Set( "DetailsView.CategoryTextStyle", FTextBlockStyle(NormalText)
diff --git a/Engine/Source/Editor/MaterialEditor/Private/SMaterialLayersFunctionsTree.cpp b/Engine/Source/Editor/MaterialEditor/Private/SMaterialLayersFunctionsTree.cpp
--- a/Engine/Source/Editor/MaterialEditor/Private/SMaterialLayersFunctionsTree.cpp
+++ b/Engine/Source/Editor/MaterialEditor/Private/SMaterialLayersFunctionsTree.cpp
@@ -226,9 +226,6 @@ void SMaterialLayersFunctionsInstanceTreeItem::Construct(const FArguments& InArg
 	StackParameterData = InArgs._StackParameterData;
 	MaterialEditorInstance = InArgs._MaterialEditorInstance;
 	Tree = InArgs._InTree;
-	ColumnSizeData.LeftColumnWidth = TAttribute<float>(Tree, &SMaterialLayersFunctionsInstanceTree::OnGetLeftColumnWidth);
-	ColumnSizeData.RightColumnWidth = TAttribute<float>(Tree, &SMaterialLayersFunctionsInstanceTree::OnGetRightColumnWidth);
-	ColumnSizeData.OnWidthChanged = SSplitter::FOnSlotResized::CreateSP(Tree, &SMaterialLayersFunctionsInstanceTree::OnSetColumnWidth);
 
 	TSharedRef<SWidget> LeftSideWidget = SNullWidget::NullWidget;
 	TSharedRef<SWidget> RightSideWidget = SNullWidget::NullWidget;
@@ -251,10 +248,10 @@ void SMaterialLayersFunctionsInstanceTreeItem::Construct(const FArguments& InArg
 
 		if (StackParameterData->ParameterInfo.Index != 0)
 		{
-			TAttribute<bool>::FGetter IsEnabledGetter = TAttribute<bool>::FGetter::CreateSP(InArgs._InTree, &SMaterialLayersFunctionsInstanceTree::IsLayerVisible, StackParameterData->ParameterInfo.Index);
+			TAttribute<bool>::FGetter IsEnabledGetter = TAttribute<bool>::FGetter::CreateSP(Tree, &SMaterialLayersFunctionsInstanceTree::IsLayerVisible, StackParameterData->ParameterInfo.Index);
 			TAttribute<bool> IsEnabledAttribute = TAttribute<bool>::Create(IsEnabledGetter);
 
-			FOnClicked VisibilityClickedDelegate = FOnClicked::CreateSP(InArgs._InTree, &SMaterialLayersFunctionsInstanceTree::ToggleLayerVisibility, StackParameterData->ParameterInfo.Index);
+			FOnClicked VisibilityClickedDelegate = FOnClicked::CreateSP(Tree, &SMaterialLayersFunctionsInstanceTree::ToggleLayerVisibility, StackParameterData->ParameterInfo.Index);
 
 			HeaderRowWidget->AddSlot()
 				.AutoWidth()
@@ -999,6 +996,7 @@ void SMaterialLayersFunctionsInstanceTreeItem::Construct(const FArguments& InArg
 		default:
 			break;
 		}
+
 		WrapperWidget->AddSlot()
 			.AutoHeight()
 			[
@@ -1011,8 +1009,8 @@ void SMaterialLayersFunctionsInstanceTreeItem::Construct(const FArguments& InArg
 					.PhysicalSplitterHandleSize(1.0f)
 					.HitDetectionSplitterHandleSize(5.0f)
 					+ SSplitter::Slot()
-					.Value(ColumnSizeData.LeftColumnWidth)
-					.OnSlotResized(ColumnSizeData.OnWidthChanged)
+					.Value(Tree->ColumnSizeData.NameColumnWidth)
+					.OnSlotResized(SSplitter::FOnSlotResized::CreateLambda([](float){}))
 					.Value(0.25f)
 					[
 						SNew(SHorizontalBox)
@@ -1031,8 +1029,8 @@ void SMaterialLayersFunctionsInstanceTreeItem::Construct(const FArguments& InArg
 						]
 					]
 					+ SSplitter::Slot()
-					.Value(ColumnSizeData.RightColumnWidth)
-					.OnSlotResized(ColumnSizeData.OnWidthChanged)
+					.Value(Tree->ColumnSizeData.ValueColumnWidth)
+					.OnSlotResized(Tree->ColumnSizeData.OnValueColumnResized)
 					[
 						SNew(SHorizontalBox)
 						+ SHorizontalBox::Slot()
@@ -1083,7 +1081,8 @@ FString SMaterialLayersFunctionsInstanceTreeItem::GetInstancePath(SMaterialLayer
 
 void SMaterialLayersFunctionsInstanceTree::Construct(const FArguments& InArgs)
 {
-	ColumnWidth = 0.5f;
+	ColumnSizeData.OnValueColumnResized.Execute(0.5f);
+
 	MaterialEditorInstance = InArgs._InMaterialEditorInstance;
 	Wrapper = InArgs._InWrapper;
 	ShowHiddenDelegate = InArgs._InShowHiddenDelegate;
@@ -1901,9 +1900,6 @@ void SMaterialLayersFunctionsMaterialTreeItem::Construct(const FArguments& InArg
 	StackParameterData = InArgs._StackParameterData;
 	MaterialEditorInstance = InArgs._MaterialEditorInstance;
 	Tree = InArgs._InTree;
-	ColumnSizeData.LeftColumnWidth = TAttribute<float>(Tree, &SMaterialLayersFunctionsMaterialTree::OnGetLeftColumnWidth);
-	ColumnSizeData.RightColumnWidth = TAttribute<float>(Tree, &SMaterialLayersFunctionsMaterialTree::OnGetRightColumnWidth);
-	ColumnSizeData.OnWidthChanged = SSplitter::FOnSlotResized::CreateSP(Tree, &SMaterialLayersFunctionsMaterialTree::OnSetColumnWidth);
 
 	TSharedRef<SWidget> LeftSideWidget = SNullWidget::NullWidget;
 	TSharedRef<SWidget> RightSideWidget = SNullWidget::NullWidget;
@@ -2490,8 +2486,8 @@ void SMaterialLayersFunctionsMaterialTreeItem::Construct(const FArguments& InArg
 					.PhysicalSplitterHandleSize(1.0f)
 					.HitDetectionSplitterHandleSize(5.0f)
 					+ SSplitter::Slot()
-					.Value(ColumnSizeData.LeftColumnWidth)
-					.OnSlotResized(ColumnSizeData.OnWidthChanged)
+					.Value(Tree->ColumnSizeData.NameColumnWidth)
+					.OnSlotResized(Tree->ColumnSizeData.OnNameColumnResized)
 					.Value(0.25f)
 					[
 						SNew(SHorizontalBox)
@@ -2510,8 +2506,8 @@ void SMaterialLayersFunctionsMaterialTreeItem::Construct(const FArguments& InArg
 						]
 					]
 					+ SSplitter::Slot()
-					.Value(ColumnSizeData.RightColumnWidth)
-					.OnSlotResized(ColumnSizeData.OnWidthChanged)
+					.Value(Tree->ColumnSizeData.ValueColumnWidth)
+					.OnSlotResized(Tree->ColumnSizeData.OnValueColumnResized)
 					[
 						SNew(SHorizontalBox)
 						+ SHorizontalBox::Slot()
@@ -2556,7 +2552,8 @@ FString SMaterialLayersFunctionsMaterialTreeItem::GetInstancePath(SMaterialLayer
 
 void SMaterialLayersFunctionsMaterialTree::Construct(const FArguments& InArgs)
 {
-	ColumnWidth = 0.5f;
+	ColumnSizeData.OnValueColumnResized.Execute(0.5f);
+
 	MaterialEditorInstance = InArgs._InMaterialEditorInstance;
 	Wrapper = InArgs._InWrapper;
 	CreateGroupsWidget();
diff --git a/Engine/Source/Editor/MaterialEditor/Private/SMaterialLayersFunctionsTree.h b/Engine/Source/Editor/MaterialEditor/Private/SMaterialLayersFunctionsTree.h
--- a/Engine/Source/Editor/MaterialEditor/Private/SMaterialLayersFunctionsTree.h
+++ b/Engine/Source/Editor/MaterialEditor/Private/SMaterialLayersFunctionsTree.h
@@ -15,6 +15,7 @@
 #include "IDetailPropertyRow.h"
 #include "MaterialPropertyHelpers.h"
 #include "PropertyCustomizationHelpers.h"
+
 class IPropertyHandle;
 class UMaterialEditorInstanceConstant;
 class SMaterialLayersFunctionsInstanceTree;
@@ -34,7 +35,6 @@ public:
 	SLATE_ARGUMENT(SMaterialLayersFunctionsInstanceTree*, InTree)
 	SLATE_END_ARGS()
 
-	FDetailColumnSizeData ColumnSizeData;
 	bool bIsBeingDragged;
 
 private:
@@ -134,9 +134,6 @@ public:
 	void OnExpansionChanged(TSharedPtr<FSortedParamData> Item, bool bIsExpanded);
 	void SetParentsExpansionState();
 
-	float OnGetLeftColumnWidth() const { return 1.0f - ColumnWidth; }
-	float OnGetRightColumnWidth() const { return ColumnWidth; }
-	void OnSetColumnWidth(float InWidth) { ColumnWidth = InWidth; }
 	void ShowHiddenValues(bool& bShowHiddenParameters) { bShowHiddenParameters = true; }
 	FName LayersFunctionsParameterName;
 	class UDEditorParameterValue* FunctionParameter;
@@ -180,8 +177,7 @@ private:
 
 	TArray<FUnsortedParamData> NonLayerProperties;
 
-	/** The actual width of the right column.  The left column is 1-ColumnWidth */
-	float ColumnWidth;
+	FDetailColumnSizeData ColumnSizeData;
 
 	SMaterialLayersFunctionsInstanceWrapper* Wrapper;
 
@@ -284,9 +280,6 @@ public:
 		void OnGetChildrenMaterialLayersFunctionsTreeView(TSharedPtr<FSortedParamData> InParent, TArray< TSharedPtr<FSortedParamData> >& OutChildren);
 		void OnExpansionChanged(TSharedPtr<FSortedParamData> Item, bool bIsExpanded);
 		void SetParentsExpansionState();
-		float OnGetLeftColumnWidth() const { return 1.0f - ColumnWidth; }
-		float OnGetRightColumnWidth() const { return ColumnWidth; }
-		void OnSetColumnWidth(float InWidth) { ColumnWidth = InWidth; }
 		void ShowHiddenValues(bool& bShowHiddenParameters) { bShowHiddenParameters = true; }
 		FName LayersFunctionsParameterName;
 		class UDEditorParameterValue* FunctionParameter;
@@ -313,8 +306,7 @@ private:
 
 	TArray<FUnsortedParamData> NonLayerProperties;
 
-	/** The actual width of the right column.  The left column is 1-ColumnWidth */
-	float ColumnWidth;
+	FDetailColumnSizeData ColumnSizeData;
 
 	SMaterialLayersFunctionsMaterialWrapper* Wrapper;
 };
diff --git a/Engine/Source/Editor/MaterialEditor/Private/SMaterialParametersOverviewWidget.cpp b/Engine/Source/Editor/MaterialEditor/Private/SMaterialParametersOverviewWidget.cpp
--- a/Engine/Source/Editor/MaterialEditor/Private/SMaterialParametersOverviewWidget.cpp
+++ b/Engine/Source/Editor/MaterialEditor/Private/SMaterialParametersOverviewWidget.cpp
@@ -76,9 +76,6 @@ void SMaterialParametersOverviewTreeItem::Construct(const FArguments& InArgs, co
 	StackParameterData = InArgs._StackParameterData;
 	MaterialEditorInstance = InArgs._MaterialEditorInstance;
 	Tree = InArgs._InTree;
-	ColumnSizeData.LeftColumnWidth = TAttribute<float>(Tree.Pin().Get(), &SMaterialParametersOverviewTree::OnGetLeftColumnWidth);
-	ColumnSizeData.RightColumnWidth = TAttribute<float>(Tree.Pin().Get(), &SMaterialParametersOverviewTree::OnGetRightColumnWidth);
-	ColumnSizeData.OnWidthChanged = SSplitter::FOnSlotResized::CreateSP(Tree.Pin().Get(), &SMaterialParametersOverviewTree::OnSetColumnWidth);
 
 	TSharedRef<SWidget> LeftSideWidget = SNullWidget::NullWidget;
 	TSharedRef<SWidget> RightSideWidget = SNullWidget::NullWidget;
@@ -457,7 +454,10 @@ void SMaterialParametersOverviewTreeItem::Construct(const FArguments& InArgs, co
 // END PROPERTY CHILD
 
 // FINAL WRAPPER
+
 	{
+		FDetailColumnSizeData& ColumnSizeData = InArgs._InTree->GetColumnSizeData();
+
 		float ValuePadding = bisPaddedProperty ? 20.0f : 0.0f;
 		WrapperWidget->AddSlot()
 			.AutoHeight()
@@ -471,8 +471,8 @@ void SMaterialParametersOverviewTreeItem::Construct(const FArguments& InArgs, co
 					.PhysicalSplitterHandleSize(1.0f)
 					.HitDetectionSplitterHandleSize(5.0f)
 					+ SSplitter::Slot()
-					.Value(ColumnSizeData.LeftColumnWidth)
-					.OnSlotResized(ColumnSizeData.OnWidthChanged)
+					.Value(ColumnSizeData.NameColumnWidth)
+					.OnSlotResized(ColumnSizeData.OnNameColumnResized)
 					.Value(0.25f)
 					[
 						SNew(SHorizontalBox)
@@ -491,8 +491,8 @@ void SMaterialParametersOverviewTreeItem::Construct(const FArguments& InArgs, co
 						]
 					]
 					+ SSplitter::Slot()
-					.Value(ColumnSizeData.RightColumnWidth)
-					.OnSlotResized(ColumnSizeData.OnWidthChanged)
+					.Value(ColumnSizeData.ValueColumnWidth)
+					.OnSlotResized(ColumnSizeData.OnValueColumnResized) 
 					[
 						SNew(SHorizontalBox)
 						+ SHorizontalBox::Slot()
@@ -524,12 +524,10 @@ void SMaterialParametersOverviewTreeItem::Construct(const FArguments& InArgs, co
 void SMaterialParametersOverviewTree::Construct(const FArguments& InArgs)
 {
 	bHasAnyParameters = false;
-	ColumnWidth = 0.5f;
 	MaterialEditorInstance = InArgs._InMaterialEditorInstance;
 	Owner = InArgs._InOwner;
 	CreateGroupsWidget();
 
-
 	STreeView<TSharedPtr<FSortedParamData>>::Construct(
 		STreeView::FArguments()
 		.TreeItemsSource(&SortedParameters)
@@ -828,68 +826,68 @@ void SMaterialParametersOverviewPanel::Refresh()
 	if (NestedTree->HasAnyParameters())
 	{
 		this->ChildSlot
-			[
-				SNew(SVerticalBox)
-				+ SVerticalBox::Slot()
+		[
+			SNew(SVerticalBox)
+			+ SVerticalBox::Slot()
 			[
 				SNew(SBorder)
 				.BorderImage(this, &SMaterialParametersOverviewPanel::GetBackgroundImage)
-			.Padding(FMargin(4.0f))
-			[
-				SNew(SVerticalBox)
-				+ SVerticalBox::Slot()
-			.AutoHeight()
-			[
-				SAssignNew(HeaderBox, SHorizontalBox)
-				+ SHorizontalBox::Slot()
-			.Padding(FMargin(3.0f, 1.0f))
-			.HAlign(HAlign_Left)
-			.AutoWidth()
-			.VAlign(VAlign_Center)
-			[
-				SNew(STextBlock)
-				.Text(LOCTEXT("ParameterDefaults", "Parameter Defaults"))
-			.Font(FEditorStyle::GetFontStyle("DetailsView.CategoryFontStyle"))
-			.ShadowOffset(FVector2D(1.0f, 1.0f))
-			]
-			]
-		+ SVerticalBox::Slot()
-			.Padding(FMargin(3.0f, 2.0f, 3.0f, 3.0f))
-			[
-				SNew(SBorder)
-				.BorderImage(FEditorStyle::GetBrush("DetailsView.CategoryTop"))
-			[
-				SNew(SHorizontalBox)
-				+ SHorizontalBox::Slot()
-			.HAlign(HAlign_Fill)
-			[
-				SNew(SWidgetSwitcher)
-				.WidgetIndex(this, &SMaterialParametersOverviewPanel::GetPanelIndex)
-			+ SWidgetSwitcher::Slot()
-			[
-				SNew(STextBlock)
-				.Text(LOCTEXT("AddParams", "Add parameters to see them here."))
-			]
-		+ SWidgetSwitcher::Slot()
-			[
-				NestedTree.ToSharedRef()
-			]
-			]
-		+ SHorizontalBox::Slot()
-			.HAlign(HAlign_Right)
-			.AutoWidth()
-			[
-				SNew(SBox)
-				.WidthOverride(16.0f)
-			[
-				ExternalScrollbar.ToSharedRef()
-			]
-			]
-			]
-			]
-			]
+				.Padding(FMargin(4.0f))
+				[
+					SNew(SVerticalBox)
+					+ SVerticalBox::Slot()
+					.AutoHeight()
+					[
+						SAssignNew(HeaderBox, SHorizontalBox)
+						+ SHorizontalBox::Slot()
+						.Padding(FMargin(3.0f, 1.0f))
+						.HAlign(HAlign_Left)
+						.AutoWidth()
+						.VAlign(VAlign_Center)
+						[
+							SNew(STextBlock)
+							.Text(LOCTEXT("ParameterDefaults", "Parameter Defaults"))
+							.Font(FEditorStyle::GetFontStyle("DetailsView.CategoryFontStyle"))
+							.ShadowOffset(FVector2D(1.0f, 1.0f))
+						]
+					]
+					+ SVerticalBox::Slot()
+					.Padding(FMargin(3.0f, 2.0f, 3.0f, 3.0f))
+					[
+						SNew(SBorder)
+						.BorderImage(FEditorStyle::GetBrush("DetailsView.CategoryTop"))
+						[
+							SNew(SHorizontalBox)
+							+ SHorizontalBox::Slot()
+							.HAlign(HAlign_Fill)
+							[
+								SNew(SWidgetSwitcher)
+								.WidgetIndex(this, &SMaterialParametersOverviewPanel::GetPanelIndex)
+								+ SWidgetSwitcher::Slot()
+								[
+									SNew(STextBlock)
+									.Text(LOCTEXT("AddParams", "Add parameters to see them here."))
+								]
+								+ SWidgetSwitcher::Slot()
+								[
+									NestedTree.ToSharedRef()
+								]
+							]
+							+ SHorizontalBox::Slot()
+							.HAlign(HAlign_Right)
+							.AutoWidth()
+							[
+								SNew(SBox)
+								.WidthOverride(16.0f)
+								[
+									ExternalScrollbar.ToSharedRef()
+								]
+							]
+						]
+					]
+				]
 			]
-			];
+		];
 
 		HeaderBox->AddSlot()
 			.FillWidth(1.0f)
@@ -905,28 +903,28 @@ void SMaterialParametersOverviewPanel::Refresh()
 				[
 					SNew(SButton)
 					.ButtonStyle(FEditorStyle::Get(), "FlatButton.DarkGrey")
-				.HAlign(HAlign_Center)
-				.OnClicked(OnChildButtonClicked)
-				.ToolTipText(LOCTEXT("SaveToChildInstance", "Save To Child Instance"))
-				.Content()
-				[
-					SNew(SHorizontalBox)
-					+ SHorizontalBox::Slot()
-				.AutoWidth()
-				[
-					SNew(STextBlock)
-					.Font(FEditorStyle::Get().GetFontStyle("FontAwesome.10"))
-				.TextStyle(FEditorStyle::Get(), "NormalText.Important")
-				.Text(FText::FromString(FString(TEXT("\xf0c7 \xf149"))) /*fa-filter*/)
-				]
-			+ SHorizontalBox::Slot()
-				.AutoWidth()
-				[
-					SNew(STextBlock)
-					.TextStyle(FEditorStyle::Get(), "NormalText.Important")
-				.Text(FText::FromString(FString(TEXT(" Save Child"))) /*fa-filter*/)
-				]
-				]
+					.HAlign(HAlign_Center)
+					.OnClicked(OnChildButtonClicked)
+					.ToolTipText(LOCTEXT("SaveToChildInstance", "Save To Child Instance"))
+					.Content()
+					[
+						SNew(SHorizontalBox)
+						+ SHorizontalBox::Slot()
+						.AutoWidth()
+						[
+							SNew(STextBlock)
+							.Font(FEditorStyle::Get().GetFontStyle("FontAwesome.10"))
+							.TextStyle(FEditorStyle::Get(), "NormalText.Important")
+							.Text(FText::FromString(FString(TEXT("\xf0c7 \xf149"))) /*fa-filter*/)
+						]
+						+ SHorizontalBox::Slot()
+						.AutoWidth()
+						[
+							SNew(STextBlock)
+							.TextStyle(FEditorStyle::Get(), "NormalText.Important")
+							.Text(FText::FromString(FString(TEXT(" Save Child"))) /*fa-filter*/)
+						]
+					]
 				];
 		}
 	}
diff --git a/Engine/Source/Editor/MaterialEditor/Private/SMaterialParametersOverviewWidget.h b/Engine/Source/Editor/MaterialEditor/Private/SMaterialParametersOverviewWidget.h
--- a/Engine/Source/Editor/MaterialEditor/Private/SMaterialParametersOverviewWidget.h
+++ b/Engine/Source/Editor/MaterialEditor/Private/SMaterialParametersOverviewWidget.h
@@ -58,8 +58,6 @@ private:
 
 	/** The set of material parameters this is associated with */
 	UMaterialEditorPreviewParameters* MaterialEditorInstance;
-
-	FDetailColumnSizeData ColumnSizeData;
 };
 
 // ********* SMaterialParametersOverviewPanel *******
@@ -93,6 +91,7 @@ private:
 
 	TSharedPtr<class SScrollBar> ExternalScrollbar;
 	TWeakPtr<class IPropertyRowGenerator> Generator;
+	
 };
 
 // ********* SMaterialParametersOverviewTree *******
@@ -120,9 +119,6 @@ public:
 	void OnExpansionChanged(TSharedPtr<FSortedParamData> Item, bool bIsExpanded);
 	void SetParentsExpansionState();
 
-	float OnGetLeftColumnWidth() const { return 1.0f - ColumnWidth; }
-	float OnGetRightColumnWidth() const { return ColumnWidth; }
-	void OnSetColumnWidth(float InWidth) { ColumnWidth = InWidth; }
 	TSharedPtr<class FAssetThumbnailPool> GetTreeThumbnailPool();
 
 	/** Object that stores all of the possible parameters we can edit */
@@ -134,6 +130,8 @@ public:
 	TWeakPtr<SMaterialParametersOverviewPanel> GetOwner() { return Owner; }
 	bool HasAnyParameters() const { return bHasAnyParameters; }
 
+	FDetailColumnSizeData& GetColumnSizeData() { return ColumnSizeData; }
+
 protected:
 
 	void ShowSubParameters();
@@ -144,11 +142,9 @@ private:
 
 	TArray<FUnsortedParamData> UnsortedParameters;
 
-	/** The actual width of the right column.  The left column is 1-ColumnWidth */
-	float ColumnWidth;
-
 	TWeakPtr<SMaterialParametersOverviewPanel> Owner;
 
 	bool bHasAnyParameters;
-
+	
+	FDetailColumnSizeData ColumnSizeData;
 };
diff --git a/Engine/Source/Editor/PropertyEditor/Private/DetailAdvancedDropdownNode.cpp b/Engine/Source/Editor/PropertyEditor/Private/DetailAdvancedDropdownNode.cpp
--- a/Engine/Source/Editor/PropertyEditor/Private/DetailAdvancedDropdownNode.cpp
+++ b/Engine/Source/Editor/PropertyEditor/Private/DetailAdvancedDropdownNode.cpp
@@ -17,7 +17,6 @@ public:
 		SLATE_ATTRIBUTE( bool, IsExpanded )
 		SLATE_ATTRIBUTE( bool, IsButtonEnabled )
 		SLATE_ARGUMENT( bool, ShouldShowAdvancedButton )
-		SLATE_ARGUMENT( FDetailColumnSizeData, ColumnSizeData )
 		SLATE_EVENT( FOnClicked, OnClicked )
 	SLATE_END_ARGS()
 
@@ -26,22 +25,21 @@ public:
 	 *
 	 * @param InArgs   A declaration from which to construct the widget
 	 */
-	void Construct( const FArguments& InArgs, const TSharedRef<STableViewBase>& InOwnerTableView, bool bIsTopNode, bool bInDisplayShowAdvancedMessage, bool bShowSplitter )
+	void Construct( const FArguments& InArgs, IDetailsViewPrivate* InDetailsView, const TSharedRef<STableViewBase>& InOwnerTableView, bool bIsTopNode, bool bInDisplayShowAdvancedMessage)
 	{
 		IsExpanded = InArgs._IsExpanded;
-
+		DetailsView = InDetailsView;
 		bDisplayShowAdvancedMessage = bInDisplayShowAdvancedMessage;
 
 		TSharedPtr<SWidget> ContentWidget;
 		if( bIsTopNode )
 		{
-			ContentWidget =
-				SNew( SBorder )
-				.BorderImage( FEditorStyle::GetBrush("DetailsView.CategoryMiddle") )
-				.Padding( FMargin( 0.0f, 3.0f, SDetailTableRowBase::ScrollbarPaddingSize, 0.0f ) )
+			ContentWidget = SNew( SBorder )
+				.BorderImage(FEditorStyle::GetBrush("DetailsView.CategoryMiddle") )
+				.Padding(FMargin( 0.0f, 0.0f, SDetailTableRowBase::ScrollbarPaddingSize, 0.0f ) )
 				[
 					SNew( SImage )
-					.Image( FEditorStyle::GetBrush("DetailsView.AdvancedDropdownBorder.Open") )
+					.Image(FEditorStyle::GetBrush("DetailsView.AdvancedDropdownBorder.Open"))
 				];
 		}
 		else if( InArgs._ShouldShowAdvancedButton )
@@ -78,57 +76,15 @@ public:
 						]
 					]
 				];
-
 		}
 		else
 		{
-			TSharedPtr<SWidget> SplitterArea;
-			if( bShowSplitter )
-			{
-				SplitterArea = 
-					SNew( SSplitter )
-					.PhysicalSplitterHandleSize( 1.0f )
-					.HitDetectionSplitterHandleSize( 5.0f )
-					.Style( FEditorStyle::Get(), "DetailsView.Splitter" )
-					+ SSplitter::Slot()
-					.Value( InArgs._ColumnSizeData.LeftColumnWidth )
-					.OnSlotResized( SSplitter::FOnSlotResized::CreateSP( this, &SAdvancedDropdownRow::OnLeftColumnResized ) )
-					[
-						SNew( SHorizontalBox )
-						+ SHorizontalBox::Slot()
-						.Padding( 3.0f, 0.0f )
-						.HAlign( HAlign_Left )
-						.VAlign( VAlign_Center )
-						.AutoWidth()
-						[
-							SNew( SExpanderArrow, SharedThis(this) )
-						]
-						+ SHorizontalBox::Slot()
-						.HAlign( HAlign_Left )
-						.Padding( FMargin( 0.0f, 2.5f, 2.0f, 2.5f ) )
-						[
-							SNew( SSpacer )
-						]
-
-					]
-					+ SSplitter::Slot()
-					.Value( InArgs._ColumnSizeData.RightColumnWidth )
-					.OnSlotResized( InArgs._ColumnSizeData.OnWidthChanged )
-					[
-						SNew( SSpacer )
-					];
-			}
-			else
-			{
-				SplitterArea = SNew(SSpacer);
-			}
-
 			ContentWidget =
 				SNew( SBorder )
-				.BorderImage( FEditorStyle::GetBrush("DetailsView.CategoryMiddle") )
+				.BorderImage( FEditorStyle::GetBrush("DetailsView.AdvancedDropdownBorder") )
 				.Padding( FMargin( 0.0f, 0.0f, SDetailTableRowBase::ScrollbarPaddingSize, 2.0f ) )
 				[
-					SplitterArea.ToSharedRef()	
+					SNew(SSpacer)
 				];
 		}
 		
@@ -146,7 +102,7 @@ public:
         		.HeightOverride(2.f)
         		[ 
 					SNew( SBorder )
-					.BorderImage( FEditorStyle::GetBrush("DetailsView.CategoryBottom") )
+					.BorderImage( FEditorStyle::GetBrush("DetailsView.AdvancedDropdownBorder") )
 				]
 			]
 		];
@@ -158,12 +114,8 @@ public:
 			InOwnerTableView
 		);	
 	}
+
 private:
-	void OnLeftColumnResized( float InNewWidth )
-	{
-		// This has to be bound or the splitter will take it upon itself to determine the size
-		// We do nothing here because it is handled by the column size data
-	}
 
 	EVisibility OnGetHelpTextVisibility() const
 	{
@@ -191,18 +143,16 @@ private:
 	TAttribute<bool> IsExpanded;
 	TSharedPtr<SButton> ExpanderButton;
 	bool bDisplayShowAdvancedMessage;
+	IDetailsViewPrivate* DetailsView;
 };
 
-
-TSharedRef< ITableRow > FAdvancedDropdownNode::GenerateWidgetForTableView( const TSharedRef<STableViewBase>& OwnerTable, const FDetailColumnSizeData& ColumnSizeData, bool bAllowFavoriteSystem)
+TSharedRef< ITableRow > FAdvancedDropdownNode::GenerateWidgetForTableView( const TSharedRef<STableViewBase>& OwnerTable, bool bAllowFavoriteSystem)
 {
-	return 
-		SNew( SAdvancedDropdownRow, OwnerTable, bIsTopNode, bDisplayShowAdvancedMessage, bShowSplitter )
+	return SNew( SAdvancedDropdownRow, ParentCategory.GetDetailsView(), OwnerTable, bIsTopNode, bDisplayShowAdvancedMessage )
 		.OnClicked( this, &FAdvancedDropdownNode::OnAdvancedDropDownClicked )
 		.IsButtonEnabled( IsEnabled )
 		.IsExpanded( IsExpanded )
-		.ShouldShowAdvancedButton( bShouldShowAdvancedButton )
-		.ColumnSizeData( ColumnSizeData );
+		.ShouldShowAdvancedButton( bShouldShowAdvancedButton );
 }
 
 bool FAdvancedDropdownNode::GenerateStandaloneWidget(FDetailWidgetRow& OutRow) const
diff --git a/Engine/Source/Editor/PropertyEditor/Private/DetailAdvancedDropdownNode.h b/Engine/Source/Editor/PropertyEditor/Private/DetailAdvancedDropdownNode.h
--- a/Engine/Source/Editor/PropertyEditor/Private/DetailAdvancedDropdownNode.h
+++ b/Engine/Source/Editor/PropertyEditor/Private/DetailAdvancedDropdownNode.h
@@ -15,14 +15,13 @@
 class FAdvancedDropdownNode : public FDetailTreeNode, public TSharedFromThis<FAdvancedDropdownNode>
 {
 public:
-	FAdvancedDropdownNode( FDetailCategoryImpl& InParentCategory, const TAttribute<bool>& InExpanded, const TAttribute<bool>& InEnabled, bool bInShouldShowAdvancedButton, bool bInDisplayShowAdvancedMessage, bool bInShowSplitter )
+	FAdvancedDropdownNode( FDetailCategoryImpl& InParentCategory, const TAttribute<bool>& InExpanded, const TAttribute<bool>& InEnabled, bool bInShouldShowAdvancedButton, bool bInDisplayShowAdvancedMessage )
 		: ParentCategory( InParentCategory )
 		, IsEnabled( InEnabled )
 		, IsExpanded( InExpanded )
 		, bShouldShowAdvancedButton( bInShouldShowAdvancedButton )
 		, bIsTopNode( false )
 		, bDisplayShowAdvancedMessage( bInDisplayShowAdvancedMessage )
-		, bShowSplitter( bInShowSplitter )
 	{}
 
 	FAdvancedDropdownNode( FDetailCategoryImpl& InParentCategory, bool bInIsTopNode )
@@ -30,12 +29,11 @@ public:
 		, bShouldShowAdvancedButton( false )
 		, bIsTopNode( bInIsTopNode )
 		, bDisplayShowAdvancedMessage( false  )
-		, bShowSplitter( false )
 	{}
 private:
 	/** IDetailTreeNode Interface */
 	virtual IDetailsViewPrivate* GetDetailsView() const override{ return ParentCategory.GetDetailsView(); }
-	virtual TSharedRef< ITableRow > GenerateWidgetForTableView( const TSharedRef<STableViewBase>& OwnerTable, const FDetailColumnSizeData& ColumnSizeData, bool bAllowFavoriteSystem) override;
+	virtual TSharedRef< ITableRow > GenerateWidgetForTableView( const TSharedRef<STableViewBase>& OwnerTable, bool bAllowFavoriteSystem) override;
 	virtual bool GenerateStandaloneWidget(FDetailWidgetRow& OutRow) const override;
 	virtual void GetChildren(FDetailNodeList& OutChildren) override {}
 	virtual void OnItemExpansionChanged( bool bIsExpanded, bool bShouldSaveState) override {}
@@ -51,6 +49,7 @@ private:
 
 	/** Called when the advanced drop down arrow is clicked */
 	FReply OnAdvancedDropDownClicked();
+
 private:
 	FDetailCategoryImpl& ParentCategory;
 	TAttribute<bool> IsEnabled;
@@ -58,5 +57,4 @@ private:
 	bool bShouldShowAdvancedButton;
 	bool bIsTopNode;
 	bool bDisplayShowAdvancedMessage;
-	bool bShowSplitter;
 };
diff --git a/Engine/Source/Editor/PropertyEditor/Private/DetailCategoryBuilderImpl.cpp b/Engine/Source/Editor/PropertyEditor/Private/DetailCategoryBuilderImpl.cpp
--- a/Engine/Source/Editor/PropertyEditor/Private/DetailCategoryBuilderImpl.cpp
+++ b/Engine/Source/Editor/PropertyEditor/Private/DetailCategoryBuilderImpl.cpp
@@ -663,7 +663,7 @@ void FDetailCategoryImpl::SetDisplayName(FName InCategoryName, const FText& Loca
 }
 
 
-TSharedRef<ITableRow> FDetailCategoryImpl::GenerateWidgetForTableView(const TSharedRef<STableViewBase>& OwnerTable, const FDetailColumnSizeData& ColumnSizeData, bool bAllowFavoriteSystem)
+TSharedRef<ITableRow> FDetailCategoryImpl::GenerateWidgetForTableView(const TSharedRef<STableViewBase>& OwnerTable, bool bAllowFavoriteSystem)
 {
 	TSharedPtr<SWidget> HeaderContent = HeaderContentWidget;
 	if (InlinePropertyNode.IsValid())
@@ -673,12 +673,10 @@ TSharedRef<ITableRow> FDetailCategoryImpl::GenerateWidgetForTableView(const TSha
 		HeaderContent = Row.ValueWidget.Widget;
 	}
 
-	return
-		SNew(SDetailCategoryTableRow, AsShared(), OwnerTable)
+	return SNew(SDetailCategoryTableRow, AsShared(), OwnerTable)
 		.InnerCategory(DetailLayoutBuilder.IsValid() ? DetailLayoutBuilder.Pin()->IsLayoutForExternalRoot() : false)
 		.DisplayName(GetDisplayName())
-		.HeaderContent(HeaderContent)
-		.ColumnSizeData(&ColumnSizeData);
+		.HeaderContent(HeaderContent);
 }
 
 
@@ -920,8 +918,7 @@ void FDetailCategoryImpl::GenerateChildrenForLayouts()
 			AdvancedDropdownNodeTop = MakeShareable(new FAdvancedDropdownNode(*this, true));
 		}
 
-		const bool bShowSplitter = bLastItemHasMultipleColumns;
-		AdvancedDropdownNodeBottom = MakeShareable(new FAdvancedDropdownNode(*this, ShowAdvanced, IsEnabled, AdvancedChildNodes.Num() > 0, SimpleChildNodes.Num() == 0, bShowSplitter));
+		AdvancedDropdownNodeBottom = MakeShareable(new FAdvancedDropdownNode(*this, ShowAdvanced, IsEnabled, AdvancedChildNodes.Num() > 0, SimpleChildNodes.Num() == 0));
 	}
 }
 
diff --git a/Engine/Source/Editor/PropertyEditor/Private/DetailCategoryBuilderImpl.h b/Engine/Source/Editor/PropertyEditor/Private/DetailCategoryBuilderImpl.h
--- a/Engine/Source/Editor/PropertyEditor/Private/DetailCategoryBuilderImpl.h
+++ b/Engine/Source/Editor/PropertyEditor/Private/DetailCategoryBuilderImpl.h
@@ -181,7 +181,7 @@ public:
 
 	/** FDetailTreeNode interface */
 	virtual IDetailsViewPrivate* GetDetailsView() const override { return DetailLayoutBuilder.Pin()->GetDetailsView(); }
-	virtual TSharedRef< ITableRow > GenerateWidgetForTableView(const TSharedRef<STableViewBase>& OwnerTable, const FDetailColumnSizeData& ColumnSizeData, bool bAllowFavoriteSystem) override;
+	virtual TSharedRef< ITableRow > GenerateWidgetForTableView(const TSharedRef<STableViewBase>& OwnerTable, bool bAllowFavoriteSystem) override;
 	virtual bool GenerateStandaloneWidget(FDetailWidgetRow& OutRow) const override;
 
 	/** IDetailTreeNode interface */
diff --git a/Engine/Source/Editor/PropertyEditor/Private/DetailCategoryGroupNode.cpp b/Engine/Source/Editor/PropertyEditor/Private/DetailCategoryGroupNode.cpp
--- a/Engine/Source/Editor/PropertyEditor/Private/DetailCategoryGroupNode.cpp
+++ b/Engine/Source/Editor/PropertyEditor/Private/DetailCategoryGroupNode.cpp
@@ -2,6 +2,8 @@
 
 #include "DetailCategoryGroupNode.h"
 #include "PropertyCustomizationHelpers.h"
+#include "PropertyEditorConstants.h"
+
 void SDetailCategoryTableRow::Construct( const FArguments& InArgs, TSharedRef<FDetailTreeNode> InOwnerTreeNode, const TSharedRef<STableViewBase>& InOwnerTableView )
 {
 	OwnerTreeNode = InOwnerTreeNode;
@@ -9,8 +11,6 @@ void SDetailCategoryTableRow::Construct( const FArguments& InArgs, TSharedRef<FD
 	bIsInnerCategory = InArgs._InnerCategory;
 	bShowBorder = InArgs._ShowBorder;
 
-	const FDetailColumnSizeData* OptionalSizeData = InArgs._ColumnSizeData;
-
 	TSharedPtr<SWidget> Widget = SNullWidget::NullWidget;
 
 	float MyContentTopPadding = 2.0f;
@@ -19,68 +19,39 @@ void SDetailCategoryTableRow::Construct( const FArguments& InArgs, TSharedRef<FD
 	float ChildSlotPadding = 2.0f;
 	float BorderVerticalPadding = 3.0f;
 
-	if (OptionalSizeData != nullptr)
-	{
-		// If we're going to draw a splitter, move any padding from the child slot and the border
-		// to the content widget instead
+	FDetailColumnSizeData& ColumnSizeData = InOwnerTreeNode->GetDetailsView()->GetColumnSizeData();
 
-		MyContentTopPadding += ChildSlotPadding + 2 * BorderVerticalPadding;
-		MyContentBottomPadding += 2 * BorderVerticalPadding;
+	MyContentTopPadding += ChildSlotPadding + 2 * BorderVerticalPadding;
+	MyContentBottomPadding += 2 * BorderVerticalPadding;
 
-		ChildSlotPadding = 0.0f;
-		BorderVerticalPadding = 0.0f;
-	}
+	ChildSlotPadding = 0.0f;
+	BorderVerticalPadding = 0.0f;
 
-	TSharedPtr<SHorizontalBox> MyContent = SNew(SHorizontalBox)
-		+SHorizontalBox::Slot()
-			.VAlign(VAlign_Center)
-			.Padding(2.0f, MyContentTopPadding, 2.0f, MyContentBottomPadding)
-			.AutoWidth()
-			[
-				SNew(SExpanderArrow, SharedThis(this))
-			]
+	Widget = SNew(SHorizontalBox)
 		+ SHorizontalBox::Slot()
-			.VAlign(VAlign_Center)
-			.AutoWidth()
-			[
-				SNew(STextBlock)
-				.TransformPolicy(ETextTransformPolicy::ToUpper)
-				.Text(InArgs._DisplayName)
-				.Font(FAppStyle::Get().GetFontStyle(bIsInnerCategory ? "PropertyWindow.NormalFont" : "DetailsView.CategoryFontStyle"))
-				.TextStyle(FAppStyle::Get(), "DetailsView.CategoryTextStyle")
-			];
-
-	// If column size data was specified, add a separator
-	if (OptionalSizeData != nullptr)
-	{
-		Widget = SNew(SSplitter)
-			.Style(FEditorStyle::Get(), "DetailsView.Splitter")
-			.PhysicalSplitterHandleSize(1.0f)
-			.HitDetectionSplitterHandleSize(5.0f)
-
-			+ SSplitter::Slot()
-			.Value(OptionalSizeData->LeftColumnWidth)
-			.OnSlotResized(SSplitter::FOnSlotResized::CreateSP(this, &SDetailCategoryTableRow::OnColumnResized))
-			[
-				MyContent.ToSharedRef()
-			]
-			
-			+SSplitter::Slot()
-			.Value(OptionalSizeData->RightColumnWidth)
-			.OnSlotResized(OptionalSizeData->OnWidthChanged)
-			[
-				SNew(SHorizontalBox)
-				+ SHorizontalBox::Slot()
-				[
-					InArgs._HeaderContent.IsValid() ? InArgs._HeaderContent.ToSharedRef() : SNullWidget::NullWidget
-				]
-			];
-	}
-	else
-	{
-		Widget = MyContent;
-	}
-
+		.VAlign(VAlign_Center)
+		.Padding(12.0f, MyContentTopPadding, 12.0f, MyContentBottomPadding)
+		.AutoWidth()
+		[
+			SNew(SExpanderArrow, SharedThis(this))
+		]
+		+ SHorizontalBox::Slot()
+		.VAlign(VAlign_Center)
+		.FillWidth(0.5f)
+		[
+			SNew(STextBlock)
+			.TransformPolicy(ETextTransformPolicy::ToUpper)
+			.Text(InArgs._DisplayName)
+			.Font(FAppStyle::Get().GetFontStyle(bIsInnerCategory ? PropertyEditorConstants::PropertyFontStyle : PropertyEditorConstants::CategoryFontStyle))
+			.TextStyle(FAppStyle::Get(), "DetailsView.CategoryTextStyle")
+		]
+		+ SHorizontalBox::Slot()
+		.VAlign(VAlign_Center)
+		.FillWidth(0.5f)
+		[
+			InArgs._HeaderContent.IsValid() ? InArgs._HeaderContent.ToSharedRef() : SNullWidget::NullWidget
+		];
+	
 	ChildSlot
 	.Padding( 0.0f, bIsInnerCategory ? 0.0f : ChildSlotPadding, 0.0f, 0.0f )
 	[	
@@ -111,10 +82,20 @@ const FSlateBrush* SDetailCategoryTableRow::GetBackgroundImage() const
 	{
 		if (IsHovered())
 		{
+			if (bIsInnerCategory)
+			{
+				return FEditorStyle::GetBrush("DetailsView.CategoryMiddle_Hovered");
+			}
+
 			return IsItemExpanded() ? FEditorStyle::GetBrush("DetailsView.CategoryTop_Hovered") : FEditorStyle::GetBrush("DetailsView.CollapsedCategory_Hovered");
 		}
 		else
 		{
+			if (bIsInnerCategory)
+			{
+				return FEditorStyle::GetBrush("DetailsView.CategoryMiddle");
+			}
+			
 			return IsItemExpanded() ? FEditorStyle::GetBrush("DetailsView.CategoryTop") : FEditorStyle::GetBrush("DetailsView.CollapsedCategory");
 		}
 	}
@@ -135,17 +116,11 @@ FReply SDetailCategoryTableRow::OnMouseButtonDown( const FGeometry& MyGeometry,
 	}
 }
 
-FReply SDetailCategoryTableRow::OnMouseButtonDoubleClick( const FGeometry& InMyGeometry, const FPointerEvent& InMouseEvent )
+FReply SDetailCategoryTableRow::OnMouseButtonDoubleClick(const FGeometry& InMyGeometry, const FPointerEvent& InMouseEvent)
 {
 	return OnMouseButtonDown(InMyGeometry, InMouseEvent);
 }
 
-void SDetailCategoryTableRow::OnColumnResized(float InNewWidth)
-{
-	// Intentionally do nothing
-}
-
-
 FDetailCategoryGroupNode::FDetailCategoryGroupNode( const FDetailNodeList& InChildNodes, FName InGroupName, FDetailCategoryImpl& InParentCategory )
 	: ChildNodes( InChildNodes )
 	, ParentCategory( InParentCategory )
@@ -156,16 +131,12 @@ FDetailCategoryGroupNode::FDetailCategoryGroupNode( const FDetailNodeList& InChi
 {
 }
 
-TSharedRef< ITableRow > FDetailCategoryGroupNode::GenerateWidgetForTableView( const TSharedRef<STableViewBase>& OwnerTable, const FDetailColumnSizeData& ColumnSizeData, bool bAllowFavoriteSystem)
+TSharedRef< ITableRow > FDetailCategoryGroupNode::GenerateWidgetForTableView( const TSharedRef<STableViewBase>& OwnerTable, bool bAllowFavoriteSystem)
 {
-	const FDetailColumnSizeData* SizeData = bHasSplitter ? &ColumnSizeData : nullptr;
-
-	return
-		SNew( SDetailCategoryTableRow, AsShared(), OwnerTable )
+	return SNew( SDetailCategoryTableRow, AsShared(), OwnerTable )
 		.DisplayName( FText::FromName(GroupName) )
 		.InnerCategory( true )
-		.ShowBorder( bShowBorder )
-		.ColumnSizeData(SizeData);
+		.ShowBorder( bShowBorder );
 }
 
 
@@ -181,7 +152,7 @@ bool FDetailCategoryGroupNode::GenerateStandaloneWidget(FDetailWidgetRow& OutRow
 	return true;
 }
 
-void FDetailCategoryGroupNode::GetChildren(FDetailNodeList& OutChildren)
+void FDetailCategoryGroupNode::GetChildren( FDetailNodeList& OutChildren )
 {
 	for( int32 ChildIndex = 0; ChildIndex < ChildNodes.Num(); ++ChildIndex )
 	{
diff --git a/Engine/Source/Editor/PropertyEditor/Private/DetailCategoryGroupNode.h b/Engine/Source/Editor/PropertyEditor/Private/DetailCategoryGroupNode.h
--- a/Engine/Source/Editor/PropertyEditor/Private/DetailCategoryGroupNode.h
+++ b/Engine/Source/Editor/PropertyEditor/Private/DetailCategoryGroupNode.h
@@ -20,29 +20,27 @@ public:
 	SLATE_BEGIN_ARGS( SDetailCategoryTableRow )
 		: _InnerCategory( false )
 		, _ShowBorder( true )
-		, _ColumnSizeData(nullptr)
 	{}
 		SLATE_ARGUMENT( FText, DisplayName )
 		SLATE_ARGUMENT( bool, InnerCategory )
 		SLATE_ARGUMENT( TSharedPtr<SWidget>, HeaderContent )
 		SLATE_ARGUMENT( bool, ShowBorder )
-		SLATE_ARGUMENT( const FDetailColumnSizeData*, ColumnSizeData )
 	SLATE_END_ARGS()
 
 	void Construct( const FArguments& InArgs, TSharedRef<FDetailTreeNode> InOwnerTreeNode, const TSharedRef<STableViewBase>& InOwnerTableView );
+
 private:
 	EVisibility IsSeparatorVisible() const;
 	const FSlateBrush* GetBackgroundImage() const;
-private:
+
 	virtual FReply OnMouseButtonDown( const FGeometry& MyGeometry, const FPointerEvent& MouseEvent ) override;
 	virtual FReply OnMouseButtonDoubleClick( const FGeometry& InMyGeometry, const FPointerEvent& InMouseEvent ) override;
-	void OnColumnResized(float InNewWidth);
+
 private:
 	bool bIsInnerCategory;
 	bool bShowBorder;
 };
 
-
 class FDetailCategoryGroupNode : public FDetailTreeNode, public TSharedFromThis<FDetailCategoryGroupNode>
 {
 public:
@@ -61,7 +59,7 @@ private:
 	virtual void OnItemExpansionChanged( bool bIsExpanded, bool bShouldSaveState) override {}
 	virtual bool ShouldBeExpanded() const override { return true; }
 	virtual ENodeVisibility GetVisibility() const override { return bShouldBeVisible ? ENodeVisibility::Visible : ENodeVisibility::HiddenDueToFiltering; }
-	virtual TSharedRef< ITableRow > GenerateWidgetForTableView( const TSharedRef<STableViewBase>& OwnerTable, const FDetailColumnSizeData& ColumnSizeData, bool bAllowFavoriteSystem) override;
+	virtual TSharedRef< ITableRow > GenerateWidgetForTableView( const TSharedRef<STableViewBase>& OwnerTable, bool bAllowFavoriteSystem) override;
 	virtual bool GenerateStandaloneWidget(FDetailWidgetRow& OutRow) const override;
 
 	virtual EDetailNodeType GetNodeType() const override { return EDetailNodeType::Category; }
@@ -73,6 +71,7 @@ private:
 	virtual void Tick( float DeltaTime ) override {}
 	virtual bool ShouldShowOnlyChildren() const override { return false; }
 	virtual FName GetNodeName() const override { return GroupName; }
+
 private:
 	FDetailNodeList ChildNodes;
 	FDetailCategoryImpl& ParentCategory;
diff --git a/Engine/Source/Editor/PropertyEditor/Private/DetailGroup.cpp b/Engine/Source/Editor/PropertyEditor/Private/DetailGroup.cpp
--- a/Engine/Source/Editor/PropertyEditor/Private/DetailGroup.cpp
+++ b/Engine/Source/Editor/PropertyEditor/Private/DetailGroup.cpp
@@ -174,26 +174,8 @@ FDetailWidgetRow FDetailGroup::GetWidgetRow()
 		[
 			MakeNameWidget()
 		];
-		Row.ValueContent()
-		[
-			SNew(SHorizontalBox)
-			+ SHorizontalBox::Slot()
-			.VAlign(VAlign_Center)
-			.AutoWidth()
-			[
-				SNew(SButton)
-				.OnClicked(this, &FDetailGroup::OnResetClicked)
-				.Visibility(this, &FDetailGroup::GetResetVisibility)
-				.ContentPadding(FMargin(5.f, 0.f))
-				.ToolTipText(LOCTEXT("ResetToDefaultToolTip", "Reset to Default"))
-				.ButtonStyle(FEditorStyle::Get(), "NoBorder")
-				.Content()
-				[
-					SNew(SImage)
-					.Image(FEditorStyle::GetBrush("PropertyWindow.DiffersFromDefault"))
-				]
-			]
-		];
+
+		Row.OverrideResetToDefault(FResetToDefaultOverride::Create(TAttribute<bool>(this, &FDetailGroup::IsResetVisible), FSimpleDelegate::CreateSP(this, &FDetailGroup::OnResetClicked)));
 
 		return Row;
 	}
@@ -244,7 +226,7 @@ TSharedRef<SWidget> FDetailGroup::MakeNameWidget()
 		];
 }
 
-FReply FDetailGroup::OnResetClicked()
+void FDetailGroup::OnResetClicked()
 {
 	if (ResetEnabled)
 	{
@@ -260,11 +242,9 @@ FReply FDetailGroup::OnResetClicked()
 			OnDetailGroupReset.Broadcast();
 		}
 	}
-
-	return FReply::Handled();
 }
 
-EVisibility FDetailGroup::GetResetVisibility() const
+bool FDetailGroup::IsResetVisible() const
 {
 	if (ResetEnabled)
 	{
@@ -276,13 +256,13 @@ EVisibility FDetailGroup::GetResetVisibility() const
 			{
 				if (PropertyHandle->DiffersFromDefault())
 				{
-					return EVisibility::Visible;
+					return true;
 				}
 			}
 		}
 	}
 
-	return EVisibility::Hidden;
+	return false;
 }
 
 void FDetailGroup::EnableReset(bool InValue)
diff --git a/Engine/Source/Editor/PropertyEditor/Private/DetailGroup.h b/Engine/Source/Editor/PropertyEditor/Private/DetailGroup.h
--- a/Engine/Source/Editor/PropertyEditor/Private/DetailGroup.h
+++ b/Engine/Source/Editor/PropertyEditor/Private/DetailGroup.h
@@ -97,8 +97,8 @@ private:
 	TSharedRef<SWidget> MakeNameWidget();
 
 	/** Called when the "Reset to Default" button for the location has been clicked */
-	FReply OnResetClicked();
-	EVisibility GetResetVisibility() const;
+	void OnResetClicked();
+	bool IsResetVisible() const;
 	bool GetAllChildrenPropertyHandles(TArray<TSharedPtr<IPropertyHandle>>& PropertyHandles) const;
 	bool GetAllChildrenPropertyHandlesRecursive(const FDetailGroup* CurrentDetailGroup, TArray<TSharedPtr<IPropertyHandle>>& PropertyHandles) const;
 
diff --git a/Engine/Source/Editor/PropertyEditor/Private/DetailItemNode.cpp b/Engine/Source/Editor/PropertyEditor/Private/DetailItemNode.cpp
--- a/Engine/Source/Editor/PropertyEditor/Private/DetailItemNode.cpp
+++ b/Engine/Source/Editor/PropertyEditor/Private/DetailItemNode.cpp
@@ -219,7 +219,7 @@ void FDetailItemNode::SetExpansionState(bool bWantsExpanded, bool bSaveState)
 	OnItemExpansionChanged(bIsExpanded, bSaveState);
 }
 
-TSharedRef< ITableRow > FDetailItemNode::GenerateWidgetForTableView( const TSharedRef<STableViewBase>& OwnerTable, const FDetailColumnSizeData& ColumnSizeData, bool bAllowFavoriteSystem)
+TSharedRef< ITableRow > FDetailItemNode::GenerateWidgetForTableView( const TSharedRef<STableViewBase>& OwnerTable, bool bAllowFavoriteSystem)
 {
 	FTagMetaData TagMeta(TEXT("DetailRowItem"));
 	if (ParentCategory.IsValid())
@@ -233,21 +233,20 @@ TSharedRef< ITableRow > FDetailItemNode::GenerateWidgetForTableView( const TShar
 			TagMeta.Tag = Customization.GetWidgetRow().RowTagName;
 		}
 	}
+
 	if( Customization.HasPropertyNode() && Customization.GetPropertyNode()->AsCategoryNode() )
 	{
 		return
 			SNew(SDetailCategoryTableRow, AsShared(), OwnerTable)
 			.DisplayName(Customization.GetPropertyNode()->GetDisplayName())
 			.AddMetaData<FTagMetaData>(TagMeta)
-			.ColumnSizeData(&ColumnSizeData)
-			.InnerCategory( true );
+			.InnerCategory(true);
 	}
 	else
 	{
 		return
 			SNew(SDetailSingleItemRow, &Customization, HasMultiColumnWidget(), AsShared(), OwnerTable )
 			.AddMetaData<FTagMetaData>(TagMeta)
-			.ColumnSizeData(ColumnSizeData)
 			.AllowFavoriteSystem(bAllowFavoriteSystem);
 	}
 }
@@ -583,7 +582,7 @@ static bool PassesAllFilters( FDetailItemNode* ItemNode, const FDetailLayoutCust
 		else if (InCustomization.HasCustomWidget())
 		{
 			const bool bPassesTextFilter = bPassesCategoryFilter || bPassesValueFilter || Local::StringPassesFilter(InFilter, InCustomization.WidgetDecl->FilterTextString.ToString());
-			const bool bPassesModifiedFilter = (InFilter.bShowOnlyModifiedProperties == false || InCustomization.WidgetDecl->DiffersFromDefaultAttr.Get() == true);
+			const bool bPassesModifiedFilter = InFilter.bShowOnlyModifiedProperties == false || InCustomization.WidgetDecl->EditConditionValue.Get(false);
 			//@todo we need to support custom widgets for keyable,animated, in particularly for transforms(ComponentTransformDetails).
 			const bool bPassesKeyableFilter = (InFilter.bShowKeyable == false);
 			const bool bPassesAnimatedFilter = (InFilter.bShowAnimated == false);
diff --git a/Engine/Source/Editor/PropertyEditor/Private/DetailItemNode.h b/Engine/Source/Editor/PropertyEditor/Private/DetailItemNode.h
--- a/Engine/Source/Editor/PropertyEditor/Private/DetailItemNode.h
+++ b/Engine/Source/Editor/PropertyEditor/Private/DetailItemNode.h
@@ -55,7 +55,7 @@ public:
 
 	/** FDetailTreeNode interface */
 	virtual IDetailsViewPrivate* GetDetailsView() const override{ return ParentCategory.Pin()->GetDetailsView(); }
-	virtual TSharedRef< ITableRow > GenerateWidgetForTableView( const TSharedRef<STableViewBase>& OwnerTable, const FDetailColumnSizeData& ColumnSizeData, bool bAllowFavoriteSystem) override;
+	virtual TSharedRef< ITableRow > GenerateWidgetForTableView( const TSharedRef<STableViewBase>& OwnerTable, bool bAllowFavoriteSystem) override;
 	virtual bool GenerateStandaloneWidget(FDetailWidgetRow& OutRow) const override;
 	virtual void GetChildren( FDetailNodeList& OutChildren )  override;
 	virtual void OnItemExpansionChanged( bool bInIsExpanded, bool bShouldSaveState) override;
diff --git a/Engine/Source/Editor/PropertyEditor/Private/DetailLayoutBuilderImpl.cpp b/Engine/Source/Editor/PropertyEditor/Private/DetailLayoutBuilderImpl.cpp
--- a/Engine/Source/Editor/PropertyEditor/Private/DetailLayoutBuilderImpl.cpp
+++ b/Engine/Source/Editor/PropertyEditor/Private/DetailLayoutBuilderImpl.cpp
@@ -196,19 +196,27 @@ void FDetailLayoutBuilderImpl::ForceRefreshDetails()
 	PropertyDetailsUtilities.Pin()->ForceRefresh();
 }
 
-FDetailCategoryImpl& FDetailLayoutBuilderImpl::DefaultCategory( FName CategoryName )
+FDetailCategoryImpl& FDetailLayoutBuilderImpl::DefaultCategory(FName CategoryName)
 {
-	TSharedPtr<FDetailCategoryImpl>& CategoryImpl = DefaultCategoryMap.FindOrAdd( CategoryName );
+	for (const TSharedRef<FDetailTreeNode>& RootTreeNode : AllRootTreeNodes)
+	{
+		if (RootTreeNode->GetNodeType() == EDetailNodeType::Category && 
+			CategoryName == RootTreeNode->GetNodeName())
+		{
+			return (FDetailCategoryImpl&) RootTreeNode.Get();
+		}
+	}
 
-	if( !CategoryImpl.IsValid() )
+	TSharedPtr<FDetailCategoryImpl>& CategoryImpl = DefaultCategoryMap.FindOrAdd(CategoryName);
+	if (!CategoryImpl.IsValid())
 	{
-		CategoryImpl = MakeShareable( new FDetailCategoryImpl( CategoryName, SharedThis(this) ) );
+		CategoryImpl = MakeShareable(new FDetailCategoryImpl(CategoryName, SharedThis(this)));
 
 		// We want categories within a type to display in the order they were added but sorting is unstable so we make unique numbers 
-		uint32 SortOrder = (uint32)ECategoryPriority::Default * 1000 + (DefaultCategoryMap.Num() - 1);
-		CategoryImpl->SetSortOrder( SortOrder );
+		uint32 SortOrder = (uint32) ECategoryPriority::Default * 1000 + (DefaultCategoryMap.Num() - 1);
+		CategoryImpl->SetSortOrder(SortOrder);
 
-		CategoryImpl->SetDisplayName( CategoryName, FText::GetEmpty() );
+		CategoryImpl->SetDisplayName(CategoryName, FText::GetEmpty());
 	}
 
 	return *CategoryImpl;
diff --git a/Engine/Source/Editor/PropertyEditor/Private/DetailMultiTopLevelObjectRootNode.cpp b/Engine/Source/Editor/PropertyEditor/Private/DetailMultiTopLevelObjectRootNode.cpp
--- a/Engine/Source/Editor/PropertyEditor/Private/DetailMultiTopLevelObjectRootNode.cpp
+++ b/Engine/Source/Editor/PropertyEditor/Private/DetailMultiTopLevelObjectRootNode.cpp
@@ -145,7 +145,7 @@ ENodeVisibility FDetailMultiTopLevelObjectRootNode::GetVisibility() const
 	return FinalVisibility;
 }
 
-TSharedRef< ITableRow > FDetailMultiTopLevelObjectRootNode::GenerateWidgetForTableView(const TSharedRef<STableViewBase>& OwnerTable, const FDetailColumnSizeData& ColumnSizeData, bool bAllowFavoriteSystem)
+TSharedRef< ITableRow > FDetailMultiTopLevelObjectRootNode::GenerateWidgetForTableView(const TSharedRef<STableViewBase>& OwnerTable, bool bAllowFavoriteSystem)
 {
 	EExpansionArrowUsage ExpansionArrowUsage = EExpansionArrowUsage::None;
 
diff --git a/Engine/Source/Editor/PropertyEditor/Private/DetailMultiTopLevelObjectRootNode.h b/Engine/Source/Editor/PropertyEditor/Private/DetailMultiTopLevelObjectRootNode.h
--- a/Engine/Source/Editor/PropertyEditor/Private/DetailMultiTopLevelObjectRootNode.h
+++ b/Engine/Source/Editor/PropertyEditor/Private/DetailMultiTopLevelObjectRootNode.h
@@ -51,7 +51,7 @@ private:
 	virtual void OnItemExpansionChanged(bool bIsExpanded, bool bShouldSaveState) override;
 	virtual bool ShouldBeExpanded() const override;
 	virtual ENodeVisibility GetVisibility() const override;
-	virtual TSharedRef<ITableRow> GenerateWidgetForTableView(const TSharedRef<STableViewBase>& OwnerTable, const FDetailColumnSizeData& ColumnSizeData, bool bAllowFavoriteSystem) override;
+	virtual TSharedRef<ITableRow> GenerateWidgetForTableView(const TSharedRef<STableViewBase>& OwnerTable, bool bAllowFavoriteSystem) override;
 	virtual bool GenerateStandaloneWidget(FDetailWidgetRow& OutRow) const override;
 	virtual void GetChildren(FDetailNodeList& OutChildren)  override;
 	virtual void FilterNode(const FDetailFilter& InFilter) override;
diff --git a/Engine/Source/Editor/PropertyEditor/Private/DetailPropertyRow.cpp b/Engine/Source/Editor/PropertyEditor/Private/DetailPropertyRow.cpp
--- a/Engine/Source/Editor/PropertyEditor/Private/DetailPropertyRow.cpp
+++ b/Engine/Source/Editor/PropertyEditor/Private/DetailPropertyRow.cpp
@@ -102,10 +102,8 @@ IDetailPropertyRow& FDetailPropertyRow::ShowPropertyButtons( bool bInShowPropert
 
 IDetailPropertyRow& FDetailPropertyRow::EditCondition( TAttribute<bool> EditConditionValue, FOnBooleanValueChanged OnEditConditionValueChanged )
 {
-	CustomEditCondition = MakeShareable( new FCustomEditCondition );
-
-	CustomEditCondition->EditConditionValue = EditConditionValue;
-	CustomEditCondition->OnEditConditionValueChanged = OnEditConditionValueChanged;
+	CustomEditConditionValue = EditConditionValue;
+	CustomEditConditionValueChanged = OnEditConditionValueChanged;
 	return *this;
 }
 
@@ -151,8 +149,8 @@ void FDetailPropertyRow::GetDefaultWidgets( TSharedPtr<SWidget>& OutNameWidget,
 		CustomTypeInterface->CustomizeHeader(PropertyHandle.ToSharedRef(), *CustomTypeRow, *this);
 	}
 
-	MakeNameOrKeyWidget(Row,CustomTypeRow);
-	MakeValueWidget(Row,CustomTypeRow,bAddWidgetDecoration);
+	MakeNameOrKeyWidget(Row, CustomTypeRow);
+	MakeValueWidget(Row, CustomTypeRow, bAddWidgetDecoration);
 
 	OutNameWidget = Row.NameWidget.Widget;
 	OutValueWidget = Row.ValueWidget.Widget;
@@ -616,18 +614,18 @@ EVisibility FDetailPropertyRow::GetPropertyVisibility() const
 
 bool FDetailPropertyRow::HasEditCondition() const
 {
-	return ( PropertyEditor.IsValid() && PropertyEditor->HasEditCondition() ) || CustomEditCondition.IsValid();
+	return (PropertyEditor.IsValid() && PropertyEditor->HasEditCondition()) || CustomEditConditionValue.IsSet();
 }
 
 bool FDetailPropertyRow::GetEnabledState() const
 {
 	bool Result = IsParentEnabled.Get();
 
-	if( HasEditCondition() ) 
+	if (HasEditCondition())
 	{
-		if (CustomEditCondition.IsValid())
+		if (CustomEditConditionValue.IsSet())
 		{
-			Result = Result && CustomEditCondition->EditConditionValue.Get();
+			Result = Result && CustomEditConditionValue.Get();
 		}
 		else
 		{
@@ -662,40 +660,30 @@ bool FDetailPropertyRow::GetForceAutoExpansion() const
 void FDetailPropertyRow::MakeNameOrKeyWidget( FDetailWidgetRow& Row, const TSharedPtr<FDetailWidgetRow> InCustomRow ) const
 {
 	EVerticalAlignment VerticalAlignment = VAlign_Center;
-	EHorizontalAlignment HorizontalAlignment = HAlign_Fill;
+	EHorizontalAlignment HorizontalAlignment = HAlign_Right;
 
 	// We will only use key widgets for non-struct keys
 	const bool bHasKeyNode = PropertyKeyEditor.IsValid() && !PropertyHandle->HasMetaData(TEXT("ReadOnlyKeys"));
 
-	if( !bHasKeyNode && InCustomRow.IsValid() )
+	if (!bHasKeyNode && InCustomRow.IsValid())
 	{
 		VerticalAlignment = InCustomRow->NameWidget.VerticalAlignment;
 		HorizontalAlignment = InCustomRow->NameWidget.HorizontalAlignment;
 	}
 
 	TAttribute<bool> IsEnabledAttrib = CustomIsEnabledAttrib;
-
-	TSharedRef<SHorizontalBox> NameHorizontalBox = SNew(SHorizontalBox);
 	
-	if( HasEditCondition() )
+	if (HasEditCondition())
 	{
 		IsEnabledAttrib.Bind( this, &FDetailPropertyRow::GetEnabledState );
-
-		NameHorizontalBox->AddSlot()
-		.AutoWidth()
-		.Padding( 0.0f, 0.0f )
-		.VAlign(VAlign_Center)
-		.HAlign(HAlign_Right)
-		[
-			SNew( SEditConditionWidget, PropertyEditor )
-			.CustomEditCondition( CustomEditCondition.IsValid() ? *CustomEditCondition : FCustomEditCondition() )
-		];
 	}
 
+	TSharedRef<SHorizontalBox> NameHorizontalBox = SNew(SHorizontalBox);
+	
 	TSharedPtr<SWidget> NameWidget;
 
 	// Key nodes take precedence over custom rows
-	if ( bHasKeyNode )
+	if (bHasKeyNode)
 	{
 		// Does this key have a custom type, use it
 		if (CachedKeyCustomTypeInterface)
@@ -717,7 +705,7 @@ void FDetailPropertyRow::MakeNameOrKeyWidget( FDetailWidgetRow& Row, const TShar
 		}
 
 	}
-	else if( InCustomRow.IsValid() )
+	else if (InCustomRow.IsValid())
 	{
 		NameWidget = 
 			SNew( SBox )
@@ -730,8 +718,7 @@ void FDetailPropertyRow::MakeNameOrKeyWidget( FDetailWidgetRow& Row, const TShar
 	{
 		NameWidget = 
 			SNew( SPropertyNameWidget, PropertyEditor )
-			.IsEnabled( IsEnabledAttrib )
-			.DisplayResetToDefault( false );
+			.IsEnabled( IsEnabledAttrib );
 	}
 
 	SHorizontalBox::FSlot& Slot = NameHorizontalBox->AddSlot()
@@ -755,7 +742,7 @@ void FDetailPropertyRow::MakeNameOrKeyWidget( FDetailWidgetRow& Row, const TShar
 	}
 
 	Row.NameContent()
-	.HAlign( HAlign_Right )
+	.HAlign( HorizontalAlignment )
 	.VAlign( VerticalAlignment )
 	[
 		NameHorizontalBox
@@ -786,32 +773,6 @@ void FDetailPropertyRow::MakeValueWidget( FDetailWidgetRow& Row, const TSharedPt
 		SNew( SHorizontalBox )
 		.IsEnabled( IsEnabledAttrib );
 
-	TSharedPtr<SResetToDefaultPropertyEditor> ResetButton = nullptr;
-	TSharedPtr<SWidget> ResetWidget = nullptr;
-	if (!PropertyHandle->HasMetaData(TEXT("NoResetToDefault")) && !PropertyHandle->GetInstanceMetaData(TEXT("NoResetToDefault")))
-	{
-		if (PropertyHandle->IsResetToDefaultCustomized())
-		{
-			// FIXME: Workaround for JIRA UE-73210.
-			// We had an oscillating SPropertyValueWidget width while dragging a UMG widget in the designer.
-			// The way drag&drop is implemented (SDesignerView::ProcessDropAndAddWidget), a new UCanvasPanelSlot gets 
-			// recreated every frame, so the details panel gets refreshed every frame. Since new property rows are created 
-			// before old ones are destroyed in the details panel, the HasCustomResetToDefault flag on the property node 
-			// toggles from frame to frame, so we alternate between having a ResetToDefaultPropertyEditor and not having one.
-			// By having a spacer fill the blank, the property row layout doesn't change while dragging, but we still see 
-			// a flashing yellow reset arrow (when visible).
-			const FSlateBrush* DiffersFromDefaultBrush = FEditorStyle::GetBrush("PropertyWindow.DiffersFromDefault");
-			ResetWidget = SNew(SSpacer).Size(DiffersFromDefaultBrush != nullptr ? DiffersFromDefaultBrush->ImageSize : FVector2D(8.0f, 8.0f));
-		}
-		else
-		{
-			SAssignNew(ResetButton, SResetToDefaultPropertyEditor, PropertyEditor->GetPropertyHandle())
-				.IsEnabled(IsEnabledAttrib)
-				.CustomResetToDefault(CustomResetToDefault);
-			ResetWidget = ResetButton;
-		}
-	};
-
 	TSharedPtr<SPropertyValueWidget> PropertyValue;
 
 	if( InCustomRow.IsValid() )
@@ -830,13 +791,12 @@ void FDetailPropertyRow::MakeValueWidget( FDetailWidgetRow& Row, const TSharedPt
 		[
 			SAssignNew( PropertyValue, SPropertyValueWidget, PropertyEditor, GetPropertyUtilities() )
 			.ShowPropertyButtons( false ) // We handle this ourselves
-			.OptionalResetWidget(ResetButton.IsValid() ? ResetButton.ToSharedRef() : SNullWidget::NullWidget)
 		];
 		MinWidth = PropertyValue->GetMinDesiredWidth();
 		MaxWidth = PropertyValue->GetMaxDesiredWidth();
 	}
 
-	if(bAddWidgetDecoration)
+	if( bAddWidgetDecoration )
 	{
 		if( bShowPropertyButtons )
 		{
@@ -867,29 +827,37 @@ void FDetailPropertyRow::MakeValueWidget( FDetailWidgetRow& Row, const TSharedPt
 				PropertyCustomizationHelpers::MakeEditConfigHierarchyButton(FSimpleDelegate::CreateSP(PropertyEditor.ToSharedRef(), &FPropertyEditor::EditConfigHierarchy))
 			];
 		}
-
-		if ((!PropertyValue.IsValid() || (PropertyValue.IsValid() && !PropertyValue->CreatedResetButton()))
-			&& ResetWidget.IsValid())
-		{
-			ValueWidget->AddSlot()
-			.Padding(4.0f, 0.0f)
-			.AutoWidth()
-			.VAlign(VAlign_Center)
-			.HAlign(HAlign_Left)
-			[
-				ResetWidget.ToSharedRef()
-			];
-		}
 	}
 
 	Row.ValueContent()
-	.HAlign( HorizontalAlignment )
-	.VAlign( VerticalAlignment )	
-	.MinDesiredWidth( MinWidth )
-	.MaxDesiredWidth( MaxWidth )
-	[
-		ValueWidget
-	];
+		.HAlign( HorizontalAlignment )
+		.VAlign( VerticalAlignment )	
+		.MinDesiredWidth( MinWidth )
+		.MaxDesiredWidth( MaxWidth )
+		[
+			ValueWidget
+		];
+	
+	// set edit condition handlers - use customized if provided
+	TAttribute<bool> EditConditionValue = CustomEditConditionValue;
+	if (!EditConditionValue.IsSet())
+	{
+		EditConditionValue = TAttribute<bool>(PropertyEditor.ToSharedRef(), &FPropertyEditor::IsEditConditionMet);
+	}
+
+	FOnBooleanValueChanged OnEditConditionValueChanged = CustomEditConditionValueChanged;
+	if (!OnEditConditionValueChanged.IsBound() && PropertyEditor->SupportsEditConditionToggle())
+	{
+		OnEditConditionValueChanged = FOnBooleanValueChanged::CreateLambda([this](bool bValue) 
+		{ 
+			if (PropertyEditor->IsEditConditionMet() != bValue)
+			{
+				PropertyEditor->ToggleEditConditionState();
+			}
+		});
+	}
+
+	Row.EditCondition(EditConditionValue, OnEditConditionValueChanged);
 }
 
 #undef LOCTEXT_NAMESPACE
diff --git a/Engine/Source/Editor/PropertyEditor/Private/DetailPropertyRow.h b/Engine/Source/Editor/PropertyEditor/Private/DetailPropertyRow.h
--- a/Engine/Source/Editor/PropertyEditor/Private/DetailPropertyRow.h
+++ b/Engine/Source/Editor/PropertyEditor/Private/DetailPropertyRow.h
@@ -193,7 +193,9 @@ private:
 	/** Custom widgets to use for this row instead of the default ones */
 	TSharedPtr<FDetailWidgetRow> CustomPropertyWidget;
 	/** User customized edit condition */
-	TSharedPtr<struct FCustomEditCondition> CustomEditCondition;
+	TAttribute<bool> CustomEditConditionValue;
+	/** User customized edit condition change handler. */
+	FOnBooleanValueChanged CustomEditConditionValueChanged;
 	/** User customized reset to default */
 	TOptional<FResetToDefaultOverride> CustomResetToDefault;
 	/** The category this row resides in */
diff --git a/Engine/Source/Editor/PropertyEditor/Private/DetailTreeNode.h b/Engine/Source/Editor/PropertyEditor/Private/DetailTreeNode.h
--- a/Engine/Source/Editor/PropertyEditor/Private/DetailTreeNode.h
+++ b/Engine/Source/Editor/PropertyEditor/Private/DetailTreeNode.h
@@ -16,7 +16,6 @@ class STableViewBase;
 class FDetailWidgetRow;
 class FPropertyNode;
 class FComplexPropertyNode;
-struct FDetailColumnSizeData;
 struct FDetailFilter;
 
 enum class ENodeVisibility : uint8
@@ -51,7 +50,7 @@ public:
 	 * @param OwnerTable		The table owner of the widget being generated
 	 * @param PropertyUtilities	Property utilities to help generate widgets
 	 */
-	virtual TSharedRef< ITableRow > GenerateWidgetForTableView(const TSharedRef<STableViewBase>& OwnerTable, const FDetailColumnSizeData& ColumnSizeData, bool bAllowFavoriteSystem) = 0;
+	virtual TSharedRef< ITableRow > GenerateWidgetForTableView(const TSharedRef<STableViewBase>& OwnerTable, bool bAllowFavoriteSystem) = 0;
 
 	virtual bool GenerateStandaloneWidget(FDetailWidgetRow& OutRow) const = 0;
 
diff --git a/Engine/Source/Editor/PropertyEditor/Private/IDetailsViewPrivate.h b/Engine/Source/Editor/PropertyEditor/Private/IDetailsViewPrivate.h
--- a/Engine/Source/Editor/PropertyEditor/Private/IDetailsViewPrivate.h
+++ b/Engine/Source/Editor/PropertyEditor/Private/IDetailsViewPrivate.h
@@ -140,5 +140,6 @@ public:
 	*/
 	virtual void RestoreExpandedItems(TSharedRef<FPropertyNode> StartNode) = 0;
 
-	virtual TSharedPtr<FEditConditionParser> GetEditConditionParser() const = 0;
+	/** Column width accessibility */
+	virtual FDetailColumnSizeData& GetColumnSizeData() = 0;
 };
diff --git a/Engine/Source/Editor/PropertyEditor/Private/Presentation/PropertyEditor/PropertyEditor.cpp b/Engine/Source/Editor/PropertyEditor/Private/Presentation/PropertyEditor/PropertyEditor.cpp
--- a/Engine/Source/Editor/PropertyEditor/Private/Presentation/PropertyEditor/PropertyEditor.cpp
+++ b/Engine/Source/Editor/PropertyEditor/Private/Presentation/PropertyEditor/PropertyEditor.cpp
@@ -44,25 +44,6 @@ FPropertyEditor::FPropertyEditor( const TSharedRef<FPropertyNode>& InPropertyNod
 
 	FProperty* Property = InPropertyNode->GetProperty();
 
-	if( Property )
-	{
-		static const FName EditConditionName = TEXT("EditCondition");
-
-		//see if the property supports some kind of edit condition and this isn't the "parent" property of a static array
-		if (Property->HasMetaData(EditConditionName) && !PropertyEditorHelpers::IsStaticArray(PropertyNode.Get()))
-		{
-			TSharedPtr<FEditConditionParser> Parser = PropertyUtilities->GetEditConditionParser();
-			if (Parser.IsValid())
-			{
-				EditConditionExpression = Parser->Parse(Property->GetMetaData(EditConditionName));
-				if (EditConditionExpression.IsValid())
-				{
-					EditConditionContext = MakeShareable(new FEditConditionContext(PropertyNode.Get()));
-				}
-			}
-		}
-	}
-
 	PropertyHandle = PropertyEditorHelpers::GetPropertyHandle( InPropertyNode, PropertyUtilities->GetNotifyHook(), PropertyUtilities );
 	check( PropertyHandle.IsValid() && PropertyHandle->IsValidHandle() );
 }
@@ -70,14 +51,9 @@ FPropertyEditor::FPropertyEditor( const TSharedRef<FPropertyNode>& InPropertyNod
 
 FText FPropertyEditor::GetDisplayName() const
 {
-	FCategoryPropertyNode* CategoryNode = PropertyNode->AsCategoryNode();
 	FItemPropertyNode* ItemPropertyNode = PropertyNode->AsItemPropertyNode();
 
-	if ( CategoryNode != NULL )
-	{
-		return CategoryNode->GetDisplayName();
-	}
-	else if ( ItemPropertyNode != NULL )
+	if ( ItemPropertyNode != NULL )
 	{
 		return ItemPropertyNode->GetDisplayName();
 	}
@@ -450,67 +426,35 @@ bool FPropertyEditor::IsEditConst() const
 	return PropertyNode->IsEditConst();
 }
 
-void FPropertyEditor::ToggleEditConditionState()
+bool FPropertyEditor::SupportsEditConditionToggle() const
 {
-	const FScopedTransaction Transaction(FText::Format(LOCTEXT("SetEditConditionState", "Set {0} edit condition state "), PropertyNode->GetDisplayName()));
-
-	PropertyNode->NotifyPreChange( PropertyNode->GetProperty(), PropertyUtilities->GetNotifyHook() );
-
-	const FBoolProperty* EditConditionProperty = EditConditionContext->GetSingleBoolProperty(EditConditionExpression);
-	check(EditConditionProperty != nullptr);
+	return PropertyNode->SupportsEditConditionToggle();
+}
 
-	FPropertyNode* ParentNode = PropertyNode->GetParentNode();
-	check(ParentNode != nullptr);
+bool FPropertyEditor::HasEditCondition() const
+{
+	return PropertyNode->HasEditCondition();
+}
 
-	bool OldValue = true;
+bool FPropertyEditor::IsEditConditionMet() const
+{
+	return PropertyNode->IsEditConditionMet();
+}
 
-	FComplexPropertyNode* ComplexParentNode = PropertyNode->FindComplexParent();
-	for (int32 Index = 0; Index < ComplexParentNode->GetInstancesNum(); ++Index)
-	{
-		// ComplexParentNode points to the top-level object 
-		// ParentNode can point to a struct inside that object (which is stored as an FItemPropertyNode)
-		// We need all three pointers to get the value pointer
-		uint8* BaseAddress = ComplexParentNode->GetMemoryOfInstance(Index);
-		uint8* ParentOffset = ParentNode->GetValueAddress(BaseAddress, PropertyNode->HasNodeFlags(EPropertyNodeFlags::IsSparseClassData) != 0);
+bool FPropertyEditor::IsOnlyVisibleWhenEditConditionMet() const
+{
+	return PropertyNode->IsOnlyVisibleWhenEditConditionMet();
+}
 
-		uint8* ValuePtr = ComplexParentNode->GetValuePtrOfInstance(Index, EditConditionProperty, ParentNode);
+void FPropertyEditor::ToggleEditConditionState()
+{
+	const FScopedTransaction Transaction(FText::Format(LOCTEXT("SetEditConditionState", "Set {0} edit condition state "), PropertyNode->GetDisplayName()));
 
-		// SPARSEDATA_TODO: these two lines should go away once we're really confident the pointer math is all correct
-		uint8* OldValuePtr = EditConditionProperty->ContainerPtrToValuePtr<uint8>(ParentOffset);
-		check(OldValuePtr == ValuePtr || PropertyNode->HasNodeFlags(EPropertyNodeFlags::IsSparseClassData));
+	PropertyNode->NotifyPreChange( PropertyNode->GetProperty(), PropertyUtilities->GetNotifyHook() );
 
-		OldValue &= EditConditionProperty->GetPropertyValue(ValuePtr);
-		EditConditionProperty->SetPropertyValue(ValuePtr, !OldValue);
-	}
+	PropertyNode->ToggleEditConditionState();
 
-	// Propagate the value change to any instances if we're editing a template object
-	FObjectPropertyNode* ObjectNode = PropertyNode->FindObjectItemParent();
-	if (ObjectNode != nullptr)
-	{
-		for (int32 ObjIndex = 0; ObjIndex < ObjectNode->GetNumObjects(); ++ObjIndex)
-		{
-			TWeakObjectPtr<UObject> ObjectWeakPtr = ObjectNode->GetUObject(ObjIndex);
-			UObject* Object = ObjectWeakPtr.Get();
-			if (Object != nullptr && Object->IsTemplate())
-			{
-				TArray<UObject*> ArchetypeInstances;
-				Object->GetArchetypeInstances(ArchetypeInstances);
-				for (int32 InstanceIndex = 0; InstanceIndex < ArchetypeInstances.Num(); ++InstanceIndex)
-				{
-					uint8* ArchetypeBaseOffset = ComplexParentNode->GetValueAddressFromObject(ArchetypeInstances[InstanceIndex]);
-					uint8* ArchetypeParentOffset = ParentNode->GetValueAddress(ArchetypeBaseOffset, PropertyNode->HasNodeFlags(EPropertyNodeFlags::IsSparseClassData) != 0);
-					uint8* ArchetypeValueAddr = EditConditionProperty->ContainerPtrToValuePtr<uint8>(ArchetypeParentOffset);
-
-					// Only propagate if the current value on the instance matches the previous value on the template.
-					const bool CurValue = EditConditionProperty->GetPropertyValue(ArchetypeValueAddr);
-					if (OldValue == CurValue)
-					{
-						EditConditionProperty->SetPropertyValue(ArchetypeValueAddr, !OldValue);
-					}
-				}
-			}
-		}
-	}
+	const FComplexPropertyNode* ComplexParentNode = PropertyNode->FindComplexParent();
 
 	TArray<TMap<FString,int32>> ArrayIndicesPerObject;
 	ArrayIndicesPerObject.AddDefaulted(ComplexParentNode->GetInstancesNum());
@@ -572,7 +516,7 @@ void FPropertyEditor::OnGetActorFiltersForSceneOutliner( TSharedPtr<FSceneOutlin
 
 bool FPropertyEditor::IsPropertyEditingEnabled() const
 {
-	return ( PropertyUtilities->IsPropertyEditingEnabled() ) && (!HasEditCondition() || IsEditConditionMet());
+	return PropertyUtilities->IsPropertyEditingEnabled() && (!PropertyNode->HasEditCondition() || PropertyNode->IsEditConditionMet());
 }
 
 void FPropertyEditor::ForceRefresh()
@@ -585,82 +529,6 @@ void FPropertyEditor::RequestRefresh()
 	PropertyUtilities->RequestRefresh();
 }
 
-bool FPropertyEditor::IsOnlyVisibleWhenEditConditionMet() const
-{
-	static const FName Name_EditConditionHides("EditConditionHides");
-	FProperty* Property = PropertyNode->GetProperty();
-	if (Property && Property->HasMetaData(Name_EditConditionHides))
-	{
-		return HasEditCondition();
-	}
-
-	return false;
-}
-
-bool FPropertyEditor::HasEditCondition() const 
-{ 
-	return EditConditionExpression.IsValid();
-}
-
-bool FPropertyEditor::IsEditConditionMet() const 
-{ 
-	if (HasEditCondition())
-	{
-		TSharedPtr<FEditConditionParser> EditConditionParser = PropertyUtilities->GetEditConditionParser();
-		if (EditConditionParser.IsValid())
-		{
-			TOptional<bool> Result = EditConditionParser->Evaluate(*EditConditionExpression.Get(), *EditConditionContext.Get());
-			if (Result.IsSet())
-			{
-				return Result.GetValue();
-			}
-		}
-	}
-
-	return true;
-}
-
-bool FPropertyEditor::SupportsEditConditionToggle() const
-{
-	FProperty* Property = PropertyNode->GetProperty();
-
-	static const FName Name_HideEditConditionToggle("HideEditConditionToggle");
-	if (EditConditionExpression.IsValid() && !Property->HasMetaData(Name_HideEditConditionToggle))
-	{
-		const FBoolProperty* ConditionalProperty = EditConditionContext->GetSingleBoolProperty(EditConditionExpression);
-		if (ConditionalProperty != nullptr)
-		{
-			// There are 2 valid states for inline edit conditions:
-			// 1. The property is marked as editable and has InlineEditConditionToggle set. 
-			// 2. The property is not marked as editable and does not have InlineEditConditionToggle set.
-			// In both cases, the original property will be hidden and only show up as a toggle.
-
-			static const FName Name_InlineEditConditionToggle("InlineEditConditionToggle");
-			const bool bIsInlineEditCondition = ConditionalProperty->HasMetaData(Name_InlineEditConditionToggle);
-			const bool bIsEditable = ConditionalProperty->HasAllPropertyFlags(CPF_Edit);
-
-			if (bIsInlineEditCondition == bIsEditable)
-			{
-				return true;
-			}
-
-			if (bIsInlineEditCondition && !bIsEditable)
-			{
-				UE_LOG(LogPropertyEditor, Warning, TEXT("Property being used as inline edit condition is not editable, but has redundant InlineEditConditionToggle flag. Field \"%s\" in class \"%s\"."), *ConditionalProperty->GetNameCPP(), *Property->GetOwnerStruct()->GetName());
-				return true;
-			}
-
-			// The property is already shown, and not marked as inline edit condition.
-			if (!bIsInlineEditCondition && bIsEditable)
-			{
-				return false;
-			}
-		}
-	}
-
-	return false;
-}
-
 void FPropertyEditor::AddPropertyEditorChild( const TSharedRef<FPropertyEditor>& Child )
 {
 	ChildPropertyEditors.Add( Child );
diff --git a/Engine/Source/Editor/PropertyEditor/Private/Presentation/PropertyEditor/PropertyEditor.h b/Engine/Source/Editor/PropertyEditor/Private/Presentation/PropertyEditor/PropertyEditor.h
--- a/Engine/Source/Editor/PropertyEditor/Private/Presentation/PropertyEditor/PropertyEditor.h
+++ b/Engine/Source/Editor/PropertyEditor/Private/Presentation/PropertyEditor/PropertyEditor.h
@@ -40,12 +40,15 @@ public:
 
 	void RequestRefresh();
 
-	bool SupportsEditConditionToggle() const;
-	void ToggleEditConditionState();
-
 	/**	@return Whether the property is editconst */
 	bool IsEditConst() const;
 
+	/** @return Whether this property should have an edit condition toggle. */
+	bool SupportsEditConditionToggle() const;
+
+	/** Toggle the current state of the edit condition if this SupportsEditConditionToggle() */
+	void ToggleEditConditionState();
+
 	/**	@return Whether the property has a condition which must be met before allowing editing of it's value */
 	bool HasEditCondition() const;
 
@@ -123,8 +126,4 @@ private:
 	 
 	/** The property view where this widget resides */
 	TSharedRef< class IPropertyUtilities > PropertyUtilities;
-
-	/** Edit condition expression used to determine if this property editor can modify its property */
-	TSharedPtr<FEditConditionExpression> EditConditionExpression;
-	TSharedPtr<FEditConditionContext> EditConditionContext;
 };
diff --git a/Engine/Source/Editor/PropertyEditor/Private/Presentation/PropertyTable/PropertyTable.cpp b/Engine/Source/Editor/PropertyEditor/Private/Presentation/PropertyTable/PropertyTable.cpp
--- a/Engine/Source/Editor/PropertyEditor/Private/Presentation/PropertyTable/PropertyTable.cpp
+++ b/Engine/Source/Editor/PropertyEditor/Private/Presentation/PropertyTable/PropertyTable.cpp
@@ -113,11 +113,6 @@ TSharedPtr<class FAssetThumbnailPool> FPropertyTable::GetThumbnailPool() const
 	return NULL;
 }
 
-TSharedPtr<FEditConditionParser> FPropertyTable::GetEditConditionParser() const
-{
-	return EditConditionParser;
-}
-
 bool FPropertyTable::GetIsUserAllowedToChangeRoot()
 {
 	return AllowUserToChangeRoot;
diff --git a/Engine/Source/Editor/PropertyEditor/Private/Presentation/PropertyTable/PropertyTable.h b/Engine/Source/Editor/PropertyEditor/Private/Presentation/PropertyTable/PropertyTable.h
--- a/Engine/Source/Editor/PropertyEditor/Private/Presentation/PropertyTable/PropertyTable.h
+++ b/Engine/Source/Editor/PropertyEditor/Private/Presentation/PropertyTable/PropertyTable.h
@@ -29,7 +29,6 @@ public:
 	virtual void EnqueueDeferredAction( FSimpleDelegate DeferredAction ) override;
 	virtual TSharedPtr<class FAssetThumbnailPool> GetThumbnailPool() const override;
 	virtual void NotifyFinishedChangingProperties(const FPropertyChangedEvent& PropertyChangedEvent) override {}
-	virtual TSharedPtr<FEditConditionParser> GetEditConditionParser() const override;
 
 	virtual bool GetIsUserAllowedToChangeRoot() override;
 	virtual void SetIsUserAllowedToChangeRoot( bool InAllowUserToChangeRoot ) override;
diff --git a/Engine/Source/Editor/PropertyEditor/Private/PropertyCustomizationHelpers.cpp b/Engine/Source/Editor/PropertyEditor/Private/PropertyCustomizationHelpers.cpp
--- a/Engine/Source/Editor/PropertyEditor/Private/PropertyCustomizationHelpers.cpp
+++ b/Engine/Source/Editor/PropertyEditor/Private/PropertyCustomizationHelpers.cpp
@@ -495,17 +495,6 @@ void SObjectPropertyEntryBox::Construct( const FArguments& InArgs )
 		}
 	}
 
-	TSharedPtr<SResetToDefaultPropertyEditor> ResetButton = nullptr;
-
-	if (InArgs._CustomResetToDefault.IsSet() || (PropertyHandle.IsValid() && !PropertyHandle->HasMetaData(TEXT("NoResetToDefault")) && !PropertyHandle->IsResetToDefaultCustomized()))
-	{
-		SAssignNew(ResetButton, SResetToDefaultPropertyEditor, PropertyHandle)
-			.IsEnabled(this, &SObjectPropertyEntryBox::IsEnabled)
-			.CustomResetToDefault(InArgs._CustomResetToDefault);		
-	};
-
-	TSharedRef<SWidget> ResetWidget = ResetButton.IsValid() ? ResetButton.ToSharedRef() : SNullWidget::NullWidget;
-
 	ChildSlot
 	[	
 		SNew(SHorizontalBox)
@@ -534,10 +523,6 @@ void SObjectPropertyEntryBox::Construct( const FArguments& InArgs )
 				[
 					InArgs._CustomContentSlot.Widget
 				]
-				.ResetToDefaultSlot()
-				[
-					ResetWidget
-				]
 		]
 	];
 }
diff --git a/Engine/Source/Editor/PropertyEditor/Private/PropertyEditorHelpers.cpp b/Engine/Source/Editor/PropertyEditor/Private/PropertyEditorHelpers.cpp
--- a/Engine/Source/Editor/PropertyEditor/Private/PropertyEditorHelpers.cpp
+++ b/Engine/Source/Editor/PropertyEditor/Private/PropertyEditorHelpers.cpp
@@ -60,17 +60,6 @@ void SPropertyNameWidget::Construct( const FArguments& InArgs, TSharedPtr<FPrope
 		]
 	
 	];
-
-	if( InArgs._DisplayResetToDefault && !PropertyEditor->GetPropertyHandle()->HasMetaData(TEXT("NoResetToDefault")) )
-	{
-		HorizontalBox->AddSlot()
-		.AutoWidth()
-		.VAlign(VAlign_Center)
-		.Padding(2,1)
-		[
-			SNew( SResetToDefaultPropertyEditor, PropertyEditor->GetPropertyHandle())
-		];
-	}
 }
 
 void SPropertyValueWidget::Construct( const FArguments& InArgs, TSharedPtr<FPropertyEditor> PropertyEditor, TSharedPtr<IPropertyUtilities> InPropertyUtilities )
@@ -78,11 +67,9 @@ void SPropertyValueWidget::Construct( const FArguments& InArgs, TSharedPtr<FProp
 	MinDesiredWidth = 0.0f;
 	MaxDesiredWidth = 0.0f;
 
-	bCreatedResetButton = false;
 	SetEnabled( TAttribute<bool>( PropertyEditor.ToSharedRef(), &FPropertyEditor::IsPropertyEditingEnabled ) );
 
-
-	ValueEditorWidget = ConstructPropertyEditorWidget( PropertyEditor, InPropertyUtilities, InArgs._OptionalResetWidget );
+	ValueEditorWidget = ConstructPropertyEditorWidget( PropertyEditor, InPropertyUtilities );
 
 	if ( !ValueEditorWidget->GetToolTip().IsValid() )
 	{
@@ -134,7 +121,7 @@ void SPropertyValueWidget::Construct( const FArguments& InArgs, TSharedPtr<FProp
 
 }
 
-TSharedRef<SWidget> SPropertyValueWidget::ConstructPropertyEditorWidget( TSharedPtr<FPropertyEditor>& PropertyEditor, TSharedPtr<IPropertyUtilities> InPropertyUtilities, TSharedRef<SWidget> InResetDefaultWidget)
+TSharedRef<SWidget> SPropertyValueWidget::ConstructPropertyEditorWidget( TSharedPtr<FPropertyEditor>& PropertyEditor, TSharedPtr<IPropertyUtilities> InPropertyUtilities )
 {
 	const TSharedRef<FPropertyEditor> PropertyEditorRef = PropertyEditor.ToSharedRef();
 	//const TSharedRef<IPropertyUtilities> PropertyUtilitiesRef = InPropertyUtilities.ToSharedRef();
@@ -192,15 +179,8 @@ TSharedRef<SWidget> SPropertyValueWidget::ConstructPropertyEditorWidget( TShared
 		{
 			TSharedRef<SPropertyEditorAsset> AssetWidget = 
 				SAssignNew( PropertyWidget, SPropertyEditorAsset, PropertyEditorRef )
-				.ThumbnailPool( InPropertyUtilities.IsValid() ? InPropertyUtilities->GetThumbnailPool() : nullptr )
-				.ResetToDefaultSlot()
-				[
-					InResetDefaultWidget
-				];
-			if (InResetDefaultWidget != SNullWidget::NullWidget)
-			{
-				bCreatedResetButton = true;
-			}
+				.ThumbnailPool( InPropertyUtilities.IsValid() ? InPropertyUtilities->GetThumbnailPool() : nullptr );
+			
 			AssetWidget->GetDesiredWidth( MinDesiredWidth, MaxDesiredWidth );
 		}
 		else if ( SPropertyEditorNumeric<float>::Supports( PropertyEditorRef ) )
@@ -346,10 +326,10 @@ TSharedRef<SWidget> SPropertyValueWidget::ConstructPropertyEditorWidget( TShared
 	return PropertyWidget.ToSharedRef();
 }
 
-void SEditConditionWidget::Construct( const FArguments& Args, TSharedPtr<FPropertyEditor> InPropertyEditor )
+void SEditConditionWidget::Construct( const FArguments& Args )
 {
-	PropertyEditor = InPropertyEditor;
-	CustomEditCondition = Args._CustomEditCondition;
+	EditConditionValue = Args._EditConditionValue;
+	OnEditConditionValueChanged = Args._OnEditConditionValueChanged;
 
 	SetVisibility(HasEditConditionToggle() ? EVisibility::Visible : EVisibility::Collapsed);
 
@@ -366,29 +346,21 @@ void SEditConditionWidget::Construct( const FArguments& Args, TSharedPtr<FProper
 
 bool SEditConditionWidget::HasEditConditionToggle() const
 {
-	return (PropertyEditor.IsValid() && PropertyEditor->HasEditCondition() && PropertyEditor->SupportsEditConditionToggle())
-		|| (CustomEditCondition.OnEditConditionValueChanged.IsBound());
+	return OnEditConditionValueChanged.IsBound();
 }
 
 void SEditConditionWidget::OnEditConditionCheckChanged( ECheckBoxState CheckState )
 {
-	FScopedTransaction EditConditionChangedTransaction(FText::Format(LOCTEXT("UpdatedEditConditionFmt", "{0} Edit Condition Changed"), PropertyEditor->GetDisplayName()));
-
-	if (PropertyEditor.IsValid() && PropertyEditor->HasEditCondition() && PropertyEditor->SupportsEditConditionToggle())
-	{
-		PropertyEditor->ToggleEditConditionState();
-	}
-	else
-	{
-		CustomEditCondition.OnEditConditionValueChanged.ExecuteIfBound(CheckState == ECheckBoxState::Checked);
-	}
+	FScopedTransaction EditConditionChangedTransaction(LOCTEXT("UpdatedEditConditionFmt", "Edit Condition Changed"));
+	
+	OnEditConditionValueChanged.ExecuteIfBound(CheckState == ECheckBoxState::Checked);
 }
 
 ECheckBoxState SEditConditionWidget::OnGetEditConditionCheckState() const
 {
-	bool bEditConditionMet = (PropertyEditor.IsValid() && PropertyEditor->HasEditCondition() && PropertyEditor->IsEditConditionMet())
-		|| CustomEditCondition.EditConditionValue.Get();
-	return bEditConditionMet ? ECheckBoxState::Checked : ECheckBoxState::Unchecked;
+	checkSlow(HasEditConditionToggle());
+
+	return EditConditionValue.Get() ? ECheckBoxState::Checked : ECheckBoxState::Unchecked;
 }
 
 namespace PropertyEditorHelpers
diff --git a/Engine/Source/Editor/PropertyEditor/Private/PropertyEditorHelpers.h b/Engine/Source/Editor/PropertyEditor/Private/PropertyEditorHelpers.h
--- a/Engine/Source/Editor/PropertyEditor/Private/PropertyEditorHelpers.h
+++ b/Engine/Source/Editor/PropertyEditor/Private/PropertyEditorHelpers.h
@@ -43,10 +43,8 @@ class SPropertyNameWidget : public SCompoundWidget
 {
 public:
 	SLATE_BEGIN_ARGS( SPropertyNameWidget )
-		:_DisplayResetToDefault(true)
 	{}
 		SLATE_EVENT( FOnClicked, OnDoubleClicked )
-		SLATE_ARGUMENT( bool, DisplayResetToDefault )
 	SLATE_END_ARGS()
 
 	void Construct( const FArguments& InArgs, TSharedPtr<FPropertyEditor> PropertyEditor );
@@ -59,11 +57,9 @@ class SPropertyValueWidget : public SCompoundWidget
 {
 public:
 	SLATE_BEGIN_ARGS( SPropertyValueWidget )
-		: _ShowPropertyButtons( true ),
-		_OptionalResetWidget(SNullWidget::NullWidget)
+		: _ShowPropertyButtons( true )
 	{}
 		SLATE_ARGUMENT( bool, ShowPropertyButtons )
-		SLATE_ARGUMENT( TSharedRef<SWidget>, OptionalResetWidget)
 	SLATE_END_ARGS()
 
 	void Construct( const FArguments& InArgs, TSharedPtr<FPropertyEditor> InPropertyEditor, TSharedPtr<IPropertyUtilities> InPropertyUtilities );
@@ -74,36 +70,25 @@ public:
 	/** @return The maximum desired with if this property value */
 	float GetMaxDesiredWidth() const { return MaxDesiredWidth; }
 
-	/** @return Whether this widget handles its own reset button */
-	bool CreatedResetButton() const { return bCreatedResetButton; }
-
 private:
-	TSharedRef<SWidget> ConstructPropertyEditorWidget( TSharedPtr<FPropertyEditor>& PropertyEditor, TSharedPtr<IPropertyUtilities> InPropertyUtilities, TSharedRef<SWidget> InOptionalResetDefaultWidget = SNullWidget::NullWidget);
+	TSharedRef<SWidget> ConstructPropertyEditorWidget( TSharedPtr<FPropertyEditor>& PropertyEditor, TSharedPtr<IPropertyUtilities> InPropertyUtilities );
 private:
 	TSharedPtr< SWidget > ValueEditorWidget;
 	/** The minimum desired with if this property value */
 	float MinDesiredWidth;
 	/** The maximum desired with if this property value */
 	float MaxDesiredWidth;
-	/** Whether or not this value widget handled its own reset button */
-	bool bCreatedResetButton;
-};
-
-
-struct FCustomEditCondition
-{
-	TAttribute<bool> EditConditionValue;
-	FOnBooleanValueChanged OnEditConditionValueChanged;
 };
 
 class SEditConditionWidget : public SCompoundWidget
 {
 	SLATE_BEGIN_ARGS( SEditConditionWidget )
 	{}
-		SLATE_ARGUMENT( FCustomEditCondition, CustomEditCondition )
+		SLATE_ATTRIBUTE( bool, EditConditionValue )
+		SLATE_EVENT( FOnBooleanValueChanged, OnEditConditionValueChanged )
 	SLATE_END_ARGS()
 
-	void Construct( const FArguments& Args, TSharedPtr<FPropertyEditor> InPropertyEditor );
+	void Construct( const FArguments& Args );
 
 private:
 	void OnEditConditionCheckChanged( ECheckBoxState CheckState );
@@ -111,8 +96,8 @@ private:
 	ECheckBoxState OnGetEditConditionCheckState() const;
 
 private:
-	TSharedPtr<FPropertyEditor> PropertyEditor;
-	FCustomEditCondition CustomEditCondition;
+	TAttribute<bool> EditConditionValue;
+	FOnBooleanValueChanged OnEditConditionValueChanged;
 };
 
 namespace PropertyEditorHelpers
diff --git a/Engine/Source/Editor/PropertyEditor/Private/PropertyHandleImpl.cpp b/Engine/Source/Editor/PropertyEditor/Private/PropertyHandleImpl.cpp
--- a/Engine/Source/Editor/PropertyEditor/Private/PropertyHandleImpl.cpp
+++ b/Engine/Source/Editor/PropertyEditor/Private/PropertyHandleImpl.cpp
@@ -2215,8 +2215,7 @@ TSharedRef<SWidget> FPropertyHandleBase::CreatePropertyNameWidget( const FText&
 
 		TSharedPtr<FPropertyEditor> PropertyEditor = FPropertyEditor::Create( Implementation->GetPropertyNode().ToSharedRef(), Implementation->GetPropertyUtilities().ToSharedRef() );
 
-		return SNew( SPropertyNameWidget, PropertyEditor )
-				.DisplayResetToDefault( bDisplayResetToDefault );
+		return SNew( SPropertyNameWidget, PropertyEditor );
 	}
 
 	return SNullWidget::NullWidget;
@@ -2495,16 +2494,16 @@ TSharedPtr<FPropertyNode> FPropertyHandleBase::GetPropertyNode() const
 	return Implementation->GetPropertyNode();
 }
 
-void FPropertyHandleBase::OnCustomResetToDefault(const FResetToDefaultOverride& OnCustomResetToDefault)
+void FPropertyHandleBase::OnCustomResetToDefault(const FResetToDefaultOverride& CustomResetToDefault)
 {
-	if (OnCustomResetToDefault.OnResetToDefaultClicked().IsBound())
+	if (CustomResetToDefault.HasResetToDefaultHandler())
 	{
 		FScopedTransaction Transaction(LOCTEXT("PropertyCustomResetToDefault", "Custom Reset to Default"));
 		if (Implementation->GetPropertyUtilities().IsValid() && Implementation->GetPropertyUtilities()->GetNotifyHook() != nullptr)
 		{
 			Implementation->GetPropertyNode()->NotifyPreChange(Implementation->GetPropertyNode()->GetProperty(), Implementation->GetPropertyUtilities()->GetNotifyHook());
 		}
-		OnCustomResetToDefault.OnResetToDefaultClicked().Execute(SharedThis(this));
+		CustomResetToDefault.OnResetToDefaultClicked(SharedThis(this));
 
 		// Call PostEditchange on all the objects
 		FPropertyChangedEvent ChangeEvent(Implementation->GetPropertyNode()->GetProperty());
diff --git a/Engine/Source/Editor/PropertyEditor/Private/PropertyNode.cpp b/Engine/Source/Editor/PropertyEditor/Private/PropertyNode.cpp
--- a/Engine/Source/Editor/PropertyEditor/Private/PropertyNode.cpp
+++ b/Engine/Source/Editor/PropertyEditor/Private/PropertyNode.cpp
@@ -1,13 +1,12 @@
 // Copyright Epic Games, Inc. All Rights Reserved.
 
-
 #include "PropertyNode.h"
 #include "Misc/ConfigCacheIni.h"
-#include "UObject/MetaData.h"
 #include "Serialization/ArchiveReplaceObjectRef.h"
 #include "Components/ActorComponent.h"
 #include "Editor/UnrealEdEngine.h"
 #include "Engine/UserDefinedStruct.h"
+#include "EditConditionContext.h"
 #include "UnrealEdGlobals.h"
 #include "ScopedTransaction.h"
 #include "PropertyRestriction.h"
@@ -17,18 +16,21 @@
 #include "Editor.h"
 #include "ObjectPropertyNode.h"
 #include "PropertyHandleImpl.h"
+#include "PropertyTextUtilities.h"
 #include "EditorSupportDelegates.h"
 #include "UObject/ConstructorHelpers.h"
 
 #include "Framework/Notifications/NotificationManager.h"
 #include "Widgets/Notifications/SNotificationList.h"
-#include "PropertyTextUtilities.h"
+#include "UObject/MetaData.h"
 #include "UObject/TextProperty.h"
 #include "UObject/EnumProperty.h"
 #include "UObject/UnrealType.h"
 
 #define LOCTEXT_NAMESPACE "PropertyNode"
 
+FEditConditionParser FPropertyNode::EditConditionParser;
+
 FPropertySettings& FPropertySettings::Get()
 {
 	static FPropertySettings Settings;
@@ -77,14 +79,11 @@ FPropertyNode::FPropertyNode(void)
 {
 }
 
-
 FPropertyNode::~FPropertyNode(void)
 {
 	DestroyTree();
 }
 
-
-
 void FPropertyNode::InitNode(const FPropertyNodeInitParams& InitParams)
 {
 	//Dismantle the previous tree
@@ -209,6 +208,21 @@ void FPropertyNode::InitNode(const FPropertyNodeInitParams& InitParams)
 	InitExpansionFlags();
 
 	FProperty* MyProperty = Property.Get();
+	if (MyProperty)
+	{
+		const FString& EditConditionString = MyProperty->GetMetaData(TEXT("EditCondition"));
+
+		// see if the property supports some kind of edit condition and this isn't the "parent" property of a static array
+		const bool bIsStaticArrayParent = MyProperty->ArrayDim > 1 && GetArrayIndex() != -1;
+		if (!EditConditionString.IsEmpty() && !bIsStaticArrayParent)
+		{
+			EditConditionExpression = EditConditionParser.Parse(EditConditionString);
+			if (EditConditionExpression.IsValid())
+			{
+				EditConditionContext = MakeShareable(new FEditConditionContext(*this));
+			}
+		}
+	}
 
 	bool bRequiresValidation = bIsEditInlineNew || bShowInnerObjectProperties || ( MyProperty && (MyProperty->IsA<FArrayProperty>() || MyProperty->IsA<FSetProperty>() || MyProperty->IsA<FMapProperty>() ));
 
@@ -217,7 +231,7 @@ void FPropertyNode::InitNode(const FPropertyNodeInitParams& InitParams)
 
 	SetNodeFlags( EPropertyNodeFlags::RequiresValidation, bRequiresValidation );
 
-	if ( InitParams.bAllowChildren )
+	if (InitParams.bAllowChildren)
 	{
 		RebuildChildren();
 	}
@@ -279,7 +293,6 @@ void FPropertyNode::ClearCachedReadAddresses( bool bRecursive )
 	}
 }
 
-
 // Follows the chain of items upwards until it finds the object window that houses this item.
 FComplexPropertyNode* FPropertyNode::FindComplexParent()
 {
@@ -877,19 +890,19 @@ bool FPropertyNode::IsPropertyConst() const
 /** @return whether this window's property is constant (can't be edited by the user) */
 bool FPropertyNode::IsEditConst() const
 {
-	if( bUpdateEditConstState )
+	if (bUpdateEditConstState)
 	{
 		// Ask the objects whether this property can be changed
 		const FObjectPropertyNode* ObjectPropertyNode = FindObjectItemParent();
 
 		bIsEditConst = IsPropertyConst();
-		if(!bIsEditConst && Property != nullptr && ObjectPropertyNode)
+		if (!bIsEditConst && Property != nullptr && ObjectPropertyNode)
 		{
 			// travel up the chain to see if this property's owner struct is editconst - if it is, so is this property
 			FPropertyNode* NextParent = ParentNode;
-			while(NextParent != nullptr && CastField<FStructProperty>(NextParent->GetProperty()) != NULL)
+			while (NextParent != nullptr && CastField<FStructProperty>(NextParent->GetProperty()) != NULL)
 			{
-				if(NextParent->IsEditConst())
+				if (NextParent->IsEditConst())
 				{
 					bIsEditConst = true;
 					break;
@@ -897,14 +910,14 @@ bool FPropertyNode::IsEditConst() const
 				NextParent = NextParent->ParentNode;
 			}
 
-			if(!bIsEditConst)
+			if (!bIsEditConst)
 			{
-				for(TPropObjectConstIterator CurObjectIt(ObjectPropertyNode->ObjectConstIterator()); CurObjectIt; ++CurObjectIt)
+				for (TPropObjectConstIterator CurObjectIt(ObjectPropertyNode->ObjectConstIterator()); CurObjectIt; ++CurObjectIt)
 				{
 					const TWeakObjectPtr<UObject> CurObject = *CurObjectIt;
-					if(CurObject.IsValid())
+					if (CurObject.IsValid())
 					{
-						if(!CurObject->CanEditChange(Property.Get()))
+						if (!CurObject->CanEditChange(Property.Get()))
 						{
 							// At least one of the objects didn't like the idea of this property being changed.
 							bIsEditConst = true;
@@ -915,13 +928,149 @@ bool FPropertyNode::IsEditConst() const
 			}
 		}
 
+		// check edit condition
+		if (!bIsEditConst && HasEditCondition())
+		{
+			bIsEditConst = !IsEditConditionMet();
+		}
+
 		bUpdateEditConstState = false;
 	}
 
-
 	return bIsEditConst;
 }
 
+bool FPropertyNode::HasEditCondition() const 
+{ 
+	return EditConditionExpression.IsValid();
+}
+
+bool FPropertyNode::IsEditConditionMet() const 
+{ 
+	if (HasEditCondition())
+	{
+		TOptional<bool> Result = EditConditionParser.Evaluate(*EditConditionExpression.Get(), *EditConditionContext.Get());
+		if (Result.IsSet())
+		{
+			return Result.GetValue();
+		}
+	}
+
+	return true;
+}
+
+bool FPropertyNode::SupportsEditConditionToggle() const
+{
+	if (!Property.IsValid())
+	{
+		return false;
+	}
+
+	FProperty* MyProperty = Property.Get();
+
+	static const FName Name_HideEditConditionToggle("HideEditConditionToggle");
+	if (EditConditionExpression.IsValid() && !Property->HasMetaData(Name_HideEditConditionToggle))
+	{
+		const FBoolProperty* ConditionalProperty = EditConditionContext->GetSingleBoolProperty(EditConditionExpression);
+		if (ConditionalProperty != nullptr)
+		{
+			// There are 2 valid states for inline edit conditions:
+			// 1. The property is marked as editable and has InlineEditConditionToggle set. 
+			// 2. The property is not marked as editable and does not have InlineEditConditionToggle set.
+			// In both cases, the original property will be hidden and only show up as a toggle.
+
+			static const FName Name_InlineEditConditionToggle("InlineEditConditionToggle");
+			const bool bIsInlineEditCondition = ConditionalProperty->HasMetaData(Name_InlineEditConditionToggle);
+			const bool bIsEditable = ConditionalProperty->HasAllPropertyFlags(CPF_Edit);
+
+			if (bIsInlineEditCondition == bIsEditable)
+			{
+				return true;
+			}
+
+			if (bIsInlineEditCondition && !bIsEditable)
+			{
+				UE_LOG(LogPropertyNode, Warning, TEXT("Property being used as inline edit condition is not editable, but has redundant InlineEditConditionToggle flag. Field \"%s\" in class \"%s\"."), *ConditionalProperty->GetNameCPP(), *Property->GetOwnerStruct()->GetName());
+				return true;
+			}
+
+			// The property is already shown, and not marked as inline edit condition.
+			if (!bIsInlineEditCondition && bIsEditable)
+			{
+				return false;
+			}
+		}
+	}
+
+	return false;
+}
+
+void FPropertyNode::ToggleEditConditionState()
+{
+	const FBoolProperty* EditConditionProperty = EditConditionContext->GetSingleBoolProperty(EditConditionExpression);
+	check(EditConditionProperty != nullptr);
+
+	FPropertyNode* MyParentNode = ParentNodeWeakPtr.Pin().Get();
+	check(MyParentNode != nullptr);
+
+	bool OldValue = true; 
+	bool IsSparseClassData = HasNodeFlags(EPropertyNodeFlags::IsSparseClassData) != 0;
+
+	FComplexPropertyNode* ComplexParentNode = FindComplexParent();
+	for (int32 Index = 0; Index < ComplexParentNode->GetInstancesNum(); ++Index)
+	{
+		// ComplexParentNode points to the top-level object 
+		// ParentNode can point to a struct inside that object (which is stored as an FItemPropertyNode)
+		// We need all three pointers to get the value pointer
+		uint8* BaseAddress = ComplexParentNode->GetMemoryOfInstance(Index);
+		uint8* ParentOffset = MyParentNode->GetValueAddress(BaseAddress, IsSparseClassData);
+
+		uint8* ValuePtr = ComplexParentNode->GetValuePtrOfInstance(Index, EditConditionProperty, MyParentNode);
+
+		OldValue &= EditConditionProperty->GetPropertyValue(ValuePtr);
+		EditConditionProperty->SetPropertyValue(ValuePtr, !OldValue);
+	}
+
+	// Propagate the value change to any instances if we're editing a template object
+	FObjectPropertyNode* ObjectNode = FindObjectItemParent();
+	if (ObjectNode != nullptr)
+	{
+		for (int32 ObjIndex = 0; ObjIndex < ObjectNode->GetNumObjects(); ++ObjIndex)
+		{
+			TWeakObjectPtr<UObject> ObjectWeakPtr = ObjectNode->GetUObject(ObjIndex);
+			UObject* Object = ObjectWeakPtr.Get();
+			if (Object != nullptr && Object->IsTemplate())
+			{
+				TArray<UObject*> ArchetypeInstances;
+				Object->GetArchetypeInstances(ArchetypeInstances);
+				for (int32 InstanceIndex = 0; InstanceIndex < ArchetypeInstances.Num(); ++InstanceIndex)
+				{
+					uint8* ArchetypeBaseOffset = ComplexParentNode->GetValueAddressFromObject(ArchetypeInstances[InstanceIndex]);
+					uint8* ArchetypeParentOffset = MyParentNode->GetValueAddress(ArchetypeBaseOffset, IsSparseClassData);
+					uint8* ArchetypeValueAddr = EditConditionProperty->ContainerPtrToValuePtr<uint8>(ArchetypeParentOffset);
+
+					// Only propagate if the current value on the instance matches the previous value on the template.
+					const bool CurValue = EditConditionProperty->GetPropertyValue(ArchetypeValueAddr);
+					if (OldValue == CurValue)
+					{
+						EditConditionProperty->SetPropertyValue(ArchetypeValueAddr, !OldValue);
+					}
+				}
+			}
+		}
+	}
+}
+
+bool FPropertyNode::IsOnlyVisibleWhenEditConditionMet() const
+{
+	static const FName Name_EditConditionHides("EditConditionHides");
+	if (Property.IsValid() && Property->HasMetaData(Name_EditConditionHides))
+	{
+		return HasEditCondition();
+	}
+
+	return false;
+}
 
 /**
  * Appends my path, including an array index (where appropriate)
diff --git a/Engine/Source/Editor/PropertyEditor/Private/PropertyNode.h b/Engine/Source/Editor/PropertyEditor/Private/PropertyNode.h
--- a/Engine/Source/Editor/PropertyEditor/Private/PropertyNode.h
+++ b/Engine/Source/Editor/PropertyEditor/Private/PropertyNode.h
@@ -7,8 +7,11 @@
 #include "UObject/UnrealType.h"
 #include "PropertyPath.h"
 #include "PropertyEditorModule.h"
+#include "EditConditionParser.h"
 
 class FComplexPropertyNode;
+class FEditConditionContext;
+class FEditConditionExpression;
 class FNotifyHook;
 class FObjectPropertyNode;
 class FPropertyItemValueDataTrackerSlate;
@@ -909,6 +912,22 @@ public:
 	 * Broadcasts reset to default property changes
 	 */
 	void BroadcastPropertyResetToDefault();
+	
+	/** @return Whether this property should have an edit condition toggle. */
+	bool SupportsEditConditionToggle() const;
+
+	/** Toggle the current state of the edit condition if this SupportsEditConditionToggle() */
+	void ToggleEditConditionState();
+
+	/**	@return Whether the property has a condition which must be met before allowing editing of it's value */
+	bool HasEditCondition() const;
+
+	/**	@return Whether the condition has been met to allow editing of this property's value */
+	bool IsEditConditionMet() const;
+
+	/**	@return Whether this property derives its visibility from its edit condition */
+	bool IsOnlyVisibleWhenEditConditionMet() const;
+
 protected:
 
 	// Returns a pointer to the starting point of the structure that contains the property this node uses.
@@ -1002,6 +1021,9 @@ protected:
 	static bool DoesChildPropertyRequireValidation(FProperty* InChildProp);
 
 protected:
+
+	static FEditConditionParser EditConditionParser;
+
 	/**
 	 * The node that is the parent of this node or nullptr for the root
 	 */
@@ -1083,6 +1105,10 @@ protected:
 	*/
 	FString PropertyPath;
 
+	/** Edit condition expression used to determine if this property editor can modify its property */
+	TSharedPtr<FEditConditionExpression> EditConditionExpression;
+	TSharedPtr<FEditConditionContext> EditConditionContext;
+
 	/**
 	* Cached state of flags that are expensive to update
 	* These update when values are changed in the details panel
diff --git a/Engine/Source/Editor/PropertyEditor/Private/PropertyRowGenerator.cpp b/Engine/Source/Editor/PropertyEditor/Private/PropertyRowGenerator.cpp
--- a/Engine/Source/Editor/PropertyEditor/Private/PropertyRowGenerator.cpp
+++ b/Engine/Source/Editor/PropertyEditor/Private/PropertyRowGenerator.cpp
@@ -19,7 +19,6 @@ class FPropertyRowGeneratorUtilities : public IPropertyUtilities
 public:
 	FPropertyRowGeneratorUtilities(FPropertyRowGenerator& InGenerator)
 		: Generator(&InGenerator)
-		, EditConditionParser(new FEditConditionParser)
 	{
 	}
 
@@ -91,14 +90,8 @@ public:
 		return Generator != nullptr && Generator->HasClassDefaultObject();
 	}
 
-	virtual TSharedPtr<FEditConditionParser> GetEditConditionParser() const override
-	{
-		return EditConditionParser;
-	}
-
 private:
 	FPropertyRowGenerator* Generator;
-	TSharedPtr<FEditConditionParser> EditConditionParser;
 };
 
 
diff --git a/Engine/Source/Editor/PropertyEditor/Private/SDetailSingleItemRow.cpp b/Engine/Source/Editor/PropertyEditor/Private/SDetailSingleItemRow.cpp
--- a/Engine/Source/Editor/PropertyEditor/Private/SDetailSingleItemRow.cpp
+++ b/Engine/Source/Editor/PropertyEditor/Private/SDetailSingleItemRow.cpp
@@ -16,7 +16,7 @@
 #include "PropertyEditorModule.h"
 #include "Modules/ModuleInterface.h"
 #include "Modules/ModuleManager.h"
-
+#include "SResetToDefaultPropertyEditor.h"
 
 void SConstrainedBox::Construct(const FArguments& InArgs)
 {
@@ -54,8 +54,8 @@ FVector2D SConstrainedBox::ComputeDesiredSize(float LayoutScaleMultiplier) const
 
 namespace DetailWidgetConstants
 {
-	const FMargin LeftRowPadding( 0.0f, 2.5f, 2.0f, 2.5f );
-	const FMargin RightRowPadding( 3.0f, 2.5f, 2.0f, 2.5f );
+	const FMargin LeftRowPadding( 0.0f, 2.5f, 12.0f, 2.5f );
+	const FMargin RightRowPadding( 12.0f, 2.5f, 2.0f, 2.5f );
 }
 
 namespace SDetailSingleItemRow_Helper
@@ -80,12 +80,12 @@ namespace SDetailSingleItemRow_Helper
 	}
 }
 
-FReply SDetailSingleItemRow::OnFavoriteToggle()
+void SDetailSingleItemRow::OnFavoriteMenuToggle()
 {
 	if (Customization->GetPropertyNode().IsValid() && Customization->GetPropertyNode()->CanDisplayFavorite())
 	{
-		bool toggle = !Customization->GetPropertyNode()->IsFavorite();
-		Customization->GetPropertyNode()->SetFavorite(toggle);
+		bool bToggle = !Customization->GetPropertyNode()->IsFavorite();
+		Customization->GetPropertyNode()->SetFavorite(bToggle);
 		if (OwnerTreeNode.IsValid())
 		{
 			//////////////////////////////////////////////////////////////////////////
@@ -104,28 +104,27 @@ FReply SDetailSingleItemRow::OnFavoriteToggle()
 			}
 			
 			//Get the number of favorite child (simple and advance) to know if the favorite category will be create or remove
-			FString CategoryFavoritesName = TEXT("Favorites");
-			FName CatFavName = *CategoryFavoritesName;
 			int32 SimplePropertiesNum = 0;
 			int32 AdvancePropertiesNum = 0;
 
 			FDetailLayoutBuilderImpl& DetailLayout = OwnerTreeNode.Pin()->GetParentCategory()->GetParentLayoutImpl();
 
-			bool HasCategoryFavorite = DetailLayout.HasCategory(CatFavName);
+			const FName FavoritesCategoryName = TEXT("Favorites");
+			bool HasCategoryFavorite = DetailLayout.HasCategory(FavoritesCategoryName);
 			if(HasCategoryFavorite)
 			{
-				DetailLayout.DefaultCategory(CatFavName).GetCategoryInformation(SimplePropertiesNum, AdvancePropertiesNum);
+				DetailLayout.DefaultCategory(FavoritesCategoryName).GetCategoryInformation(SimplePropertiesNum, AdvancePropertiesNum);
 			}
 
-			//Check if the property we toggle is an advance property
+			// Check if the property we toggle is an advance property
 			bool IsAdvanceProperty = Customization->GetPropertyNode()->HasNodeFlags(EPropertyNodeFlags::IsAdvanced) == 0 ? false : true;
 
-			//Compute the scrolling offset by item
+			// Compute the scrolling offset by item
 			int32 ScrollingOffsetAdd = ExpandSize;
 			int32 ScrollingOffsetRemove = -ExpandSize;
 			if (HasCategoryFavorite)
 			{
-				//Adding the advance button in a category add 1 item
+				// Adding the advance button in a category add 1 item
 				ScrollingOffsetAdd += (IsAdvanceProperty && AdvancePropertiesNum == 0) ? 1 : 0;
 
 				if (IsAdvanceProperty && AdvancePropertiesNum == 1)
@@ -141,22 +140,21 @@ FReply SDetailSingleItemRow::OnFavoriteToggle()
 			}
 			else
 			{
-				//Adding new category (2 items) adding advance button (1 item)
+				// Adding new category (2 items) adding advance button (1 item)
 				ScrollingOffsetAdd += IsAdvanceProperty ? 3 : 2;
 				
-				//We should never remove an item from favorite if there is no favorite category
-				//Set the remove offset to 0
+				// We should never remove an item from favorite if there is no favorite category
+				// Set the remove offset to 0
 				ScrollingOffsetRemove = 0;
 			}
 
-			//Apply the calculated offset
-			OwnerTreeNode.Pin()->GetDetailsView()->MoveScrollOffset(toggle ? ScrollingOffsetAdd : ScrollingOffsetRemove);
+			// Apply the calculated offset
+			OwnerTreeNode.Pin()->GetDetailsView()->MoveScrollOffset(bToggle ? ScrollingOffsetAdd : ScrollingOffsetRemove);
 
-			//Refresh the tree
+			// Refresh the tree
 			OwnerTreeNode.Pin()->GetDetailsView()->ForceRefresh();
 		}
 	}
-	return FReply::Handled();
 }
 
 void SDetailSingleItemRow::OnArrayDragEnter(const FDragDropEvent& DragDropEvent)
@@ -306,177 +304,158 @@ TSharedPtr<IPropertyHandle> SDetailSingleItemRow::GetPropertyHandle() const
 	return Handle;
 }
 
-const FSlateBrush* SDetailSingleItemRow::GetFavoriteButtonBrush() const
-{
-	if (Customization->GetPropertyNode().IsValid() && Customization->GetPropertyNode()->CanDisplayFavorite())
-	{
-		return FEditorStyle::GetBrush(Customization->GetPropertyNode()->IsFavorite() ? TEXT("DetailsView.PropertyIsFavorite") : IsHovered() ? TEXT("DetailsView.PropertyIsNotFavorite") : TEXT("DetailsView.NoFavoritesSystem"));
-	}
-	//Adding a transparent brush make sure all property are left align correctly
-	return FEditorStyle::GetBrush(TEXT("DetailsView.NoFavoritesSystem"));
-}
-
 void SDetailSingleItemRow::Construct( const FArguments& InArgs, FDetailLayoutCustomization* InCustomization, bool bHasMultipleColumns, TSharedRef<FDetailTreeNode> InOwnerTreeNode, const TSharedRef<STableViewBase>& InOwnerTableView )
 {
 	OwnerTreeNode = InOwnerTreeNode;
 	bAllowFavoriteSystem = InArgs._AllowFavoriteSystem;
-
-	ColumnSizeData = InArgs._ColumnSizeData;
-
-	TSharedRef<SWidget> Widget = SNullWidget::NullWidget;
 	Customization = InCustomization;
 
-	EHorizontalAlignment HorizontalAlignment = HAlign_Fill;
-	EVerticalAlignment VerticalAlignment = VAlign_Fill;
-
-	TAttribute<bool> NameWidgetEnabled;
+	TSharedRef<SWidget> Widget = SNullWidget::NullWidget;
 
 	FOnTableRowDragEnter ArrayDragDelegate;
 	FOnTableRowDragLeave ArrayDragLeaveDelegate;
 	FOnTableRowDrop ArrayDropDelegate;
 	FOnCanAcceptDrop ArrayAcceptDropDelegate;
 
+	FDetailColumnSizeData& ColumnSizeData = InOwnerTreeNode->GetDetailsView()->GetColumnSizeData();
+
 	const bool bIsValidTreeNode = InOwnerTreeNode->GetParentCategory().IsValid() && InOwnerTreeNode->GetParentCategory()->IsParentLayoutValid();
-	if(bIsValidTreeNode)
+	if (bIsValidTreeNode)
 	{
-		if(InCustomization->IsValidCustomization())
+		if (Customization->IsValidCustomization())
 		{
+			TAttribute<bool> IsPropertyEditingEnabled = InOwnerTreeNode->IsPropertyEditingEnabled();
+
 			FDetailWidgetRow Row = InCustomization->GetWidgetRow();
 
-			TSharedPtr<SWidget> NameWidget;
-			TSharedPtr<SWidget> ValueWidget;
+			TSharedPtr<SWidget> NameWidget, ValueWidget, ExtensionWidget;
 
 			NameWidget = Row.NameWidget.Widget;
-			if(Row.IsEnabledAttr.IsBound())
-			{
-				NameWidgetEnabled = Row.IsEnabledAttr;
-				NameWidget->SetEnabled(Row.IsEnabledAttr);
-			}
 
 			ValueWidget =
 				SNew(SConstrainedBox)
 				.MinWidth(Row.ValueWidget.MinWidth)
 				.MaxWidth(Row.ValueWidget.MaxWidth)
+				.IsEnabled(IsPropertyEditingEnabled)
 				[
 					Row.ValueWidget.Widget
 				];
 
-			TSharedRef<SWidget> ExtensionWidget = CreateExtensionWidget(ValueWidget.ToSharedRef(), *Customization, InOwnerTreeNode);
+			ExtensionWidget = CreateExtensionWidget(ValueWidget.ToSharedRef(), *Customization, InOwnerTreeNode);
 
-			if(Row.IsEnabledAttr.IsBound())
+			TAttribute<bool> NameWidgetEnabled;
+			if (Row.IsEnabledAttr.IsBound())
 			{
+				NameWidgetEnabled = Row.IsEnabledAttr;
+				NameWidget->SetEnabled(Row.IsEnabledAttr);
 				ValueWidget->SetEnabled(Row.IsEnabledAttr);
 				ExtensionWidget->SetEnabled(Row.IsEnabledAttr);
 			}
 
-			TArray<TSharedRef<SWidget>> GlobalWidgetExtensions;
-			FOnGenerateGlobalRowExtensionArgs Args{ GetPropertyHandle(), GetPropertyNode(), OwnerTreeNode };
-
-			FPropertyEditorModule& Module = FModuleManager::LoadModuleChecked<FPropertyEditorModule>("PropertyEditor");
-			Module.GetGlobalRowExtensionDelegate().Broadcast(Args, FOnGenerateGlobalRowExtensionArgs::EWidgetPosition::Left, GlobalWidgetExtensions);
+			TSharedRef<SHorizontalBox> RowBox = SNew(SHorizontalBox);
 
-			TSharedRef<SWidget> KeyFrameButton = CreateKeyframeButton(*Customization, InOwnerTreeNode);
-			TAttribute<bool> IsPropertyEditingEnabled = InOwnerTreeNode->IsPropertyEditingEnabled();
+			// create outer splitter
+			TSharedRef<SSplitter> OuterSplitter =
+				SNew(SSplitter)
+				.Style(FEditorStyle::Get(), "DetailsView.Splitter.Outer")
+				.PhysicalSplitterHandleSize(1.0f)
+				.HitDetectionSplitterHandleSize(5.0f);
 
-			bool const bEnableFavoriteSystem = IsEngineExitRequested() ? false : (GetDefault<UEditorExperimentalSettings>()->bEnableFavoriteSystem && bAllowFavoriteSystem);
+			Widget = OuterSplitter;
 
-			TSharedRef<SHorizontalBox> InternalLeftColumnRowBox =
+			// create Left column:
+			// | Left  | Name | Value | Right |
+			TSharedRef<SHorizontalBox> LeftColumnBox =
 				SNew(SHorizontalBox)
 				.Clipping(EWidgetClipping::OnDemand);
 
-			if(GlobalWidgetExtensions.Num() > 0)
-			{
-				TSharedRef<SHorizontalBox> LeftExtensions = SNew(SHorizontalBox);
-
-				for (TSharedRef<SWidget>& GlobalExtensionWidget : GlobalWidgetExtensions)
-				{
-					LeftExtensions->AddSlot()
-					.AutoWidth()
+			// edit condition widget
+			LeftColumnBox->AddSlot()
+				.Padding(0.0f, 0.0f, 3.0f, 0.0f)
+				.HAlign(HAlign_Left)
+				.VAlign(VAlign_Center)
+				.AutoWidth()
+				[
+					SNew(SConstrainedBox)
+					.MinWidth(20)
 					[
-						MoveTemp(GlobalExtensionWidget)
-					];
-				}
+						SNew(SEditConditionWidget)
+						.EditConditionValue(Row.EditConditionValue)
+						.OnEditConditionValueChanged(Row.OnEditConditionValueChanged)
+					]
+				];
 
-				InternalLeftColumnRowBox->AddSlot()
-					.Padding(0.0f, 0.0f)
-					.HAlign(HAlign_Left)
-					.VAlign(VAlign_Center)
-					.AutoWidth()
-					[
-						MoveTemp(LeftExtensions)
-					];
+			OuterSplitter->AddSlot()
+				.SizeRule(SSplitter::ESizeRule::SizeToContent)
+				[
+					LeftColumnBox
+				];
 
-			}
+			TSharedPtr<SSplitter> InnerSplitter;
 
-			if (bEnableFavoriteSystem)
+			// create inner splitter
+			OuterSplitter->AddSlot()
+				.Value(ColumnSizeData.PropertyColumnWidth)
+				.OnSlotResized(SSplitter::FOnSlotResized::CreateLambda([](float){})) // dummy handler, FDetailColumnSizeData handles this calculation
+			[
+				SAssignNew(InnerSplitter, SSplitter)
+				.Style(FEditorStyle::Get(), "DetailsView.Splitter")
+				.PhysicalSplitterHandleSize(1.0f)
+				.HitDetectionSplitterHandleSize(5.0f)
+				.HighlightedHandleIndex(ColumnSizeData.HoveredSplitterIndex)
+				.OnHandleHovered(ColumnSizeData.OnSplitterHandleHovered)
+			];
+
+			// create Name column:
+			// | Left  | Name | Value | Right |
+			TSharedRef<SHorizontalBox> NameColumnBox = SNew(SHorizontalBox)
+				.Clipping(EWidgetClipping::OnDemand);
+			
+			TSharedPtr<FPropertyNode> PropertyNode = Customization->GetPropertyNode();
+			if (PropertyNode.IsValid())
 			{
-				InternalLeftColumnRowBox->AddSlot()
-					.Padding(0.0f, 0.0f)
-					.HAlign(HAlign_Left)
-					.VAlign(VAlign_Center)
-					.AutoWidth()
-					[
-						SNew(SButton)
-						.HAlign(HAlign_Center)
+				if (PropertyNode->IsReorderable())
+				{
+					TSharedPtr<SDetailSingleItemRow> InRow = SharedThis(this);
+					TSharedRef<SWidget> Handle = PropertyEditorHelpers::MakePropertyReorderHandle(PropertyNode.ToSharedRef(), InRow);
+					Handle->SetEnabled(IsPropertyEditingEnabled);
+
+					NameColumnBox->AddSlot()
+						.Padding(0.0f, 0.0f, 10.0f, 0.0f)
+						.HAlign(HAlign_Right)
 						.VAlign(VAlign_Center)
-						.IsFocusable(false)
-						.ButtonStyle(FEditorStyle::Get(), "NoBorder")
-						.OnClicked(this, &SDetailSingleItemRow::OnFavoriteToggle)
 						[
-							SNew(SImage).Image(this, &SDetailSingleItemRow::GetFavoriteButtonBrush)
-						]
-					];
-			}
-			TSharedPtr<SOverlay> LeftSideOverlay;
-			LeftSideOverlay = SNew(SOverlay);
-			LeftSideOverlay->AddSlot()
-				.Padding(3.0f, 0.0f)
-				.HAlign(HAlign_Left)
-				.VAlign(VAlign_Center)
-				[
-					SNew(SExpanderArrow, SharedThis(this))
-					.BaseIndentLevel(1)
-				];
+							Handle
+						];
 
-			TSharedPtr<FPropertyNode> PropertyNode = Customization->GetPropertyNode();
-			if (PropertyNode.IsValid() && PropertyNode->IsReorderable())
-			{
-				TSharedPtr<SDetailSingleItemRow> InRow = SharedThis(this);
-				TSharedRef<SWidget> Handle = PropertyEditorHelpers::MakePropertyReorderHandle(PropertyNode.ToSharedRef(), InRow);
-				Handle->SetEnabled(IsPropertyEditingEnabled);
-				LeftSideOverlay->AddSlot()
-					.Padding(0.0f, 0.0f, 10.0f, 0.0f)
-					.HAlign(HAlign_Right)
-					.VAlign(VAlign_Center)
-					[
-						Handle
-					];
-				ArrayDragDelegate = FOnTableRowDragEnter::CreateSP(this, &SDetailSingleItemRow::OnArrayDragEnter);
-				ArrayDragLeaveDelegate = FOnTableRowDragLeave::CreateSP(this, &SDetailSingleItemRow::OnArrayDragLeave);
-				ArrayDropDelegate = FOnTableRowDrop::CreateSP(this, &SDetailSingleItemRow::OnArrayDrop);
-				ArrayAcceptDropDelegate = FOnCanAcceptDrop::CreateSP(this, &SDetailSingleItemRow::OnArrayCanAcceptDrop);
-				SwappablePropertyNode = PropertyNode;
-			}
-			else if (PropertyNode.IsValid() 
-				&& CastField<FArrayProperty>(PropertyNode->GetProperty()) != nullptr // Is an array
-				&& CastField<FObjectProperty>(CastField<FArrayProperty>(PropertyNode->GetProperty())->Inner) != nullptr) // Is an object array
-			{
-				ArrayDragDelegate = FOnTableRowDragEnter::CreateSP(this, &SDetailSingleItemRow::OnArrayDragEnter);
-				ArrayDragLeaveDelegate = FOnTableRowDragLeave::CreateSP(this, &SDetailSingleItemRow::OnArrayDragLeave);
-				ArrayDropDelegate = FOnTableRowDrop::CreateSP(this, &SDetailSingleItemRow::OnArrayHeaderDrop);
-				ArrayAcceptDropDelegate = FOnCanAcceptDrop::CreateSP(this, &SDetailSingleItemRow::OnArrayCanAcceptDrop);
+					ArrayDragDelegate = FOnTableRowDragEnter::CreateSP(this, &SDetailSingleItemRow::OnArrayDragEnter);
+					ArrayDragLeaveDelegate = FOnTableRowDragLeave::CreateSP(this, &SDetailSingleItemRow::OnArrayDragLeave);
+					ArrayDropDelegate = FOnTableRowDrop::CreateSP(this, &SDetailSingleItemRow::OnArrayDrop);
+					ArrayAcceptDropDelegate = FOnCanAcceptDrop::CreateSP(this, &SDetailSingleItemRow::OnArrayCanAcceptDrop);
+					SwappablePropertyNode = PropertyNode;
+				}
+				else if (CastField<FArrayProperty>(PropertyNode->GetProperty()) != nullptr // Is an array
+					&& CastField<FObjectProperty>(CastField<FArrayProperty>(PropertyNode->GetProperty())->Inner) != nullptr) // Is an object array
+				{
+					ArrayDragDelegate = FOnTableRowDragEnter::CreateSP(this, &SDetailSingleItemRow::OnArrayDragEnter);
+					ArrayDragLeaveDelegate = FOnTableRowDragLeave::CreateSP(this, &SDetailSingleItemRow::OnArrayDragLeave);
+					ArrayDropDelegate = FOnTableRowDrop::CreateSP(this, &SDetailSingleItemRow::OnArrayHeaderDrop);
+					ArrayAcceptDropDelegate = FOnCanAcceptDrop::CreateSP(this, &SDetailSingleItemRow::OnArrayCanAcceptDrop);
+				}
 			}
-			
-			InternalLeftColumnRowBox->AddSlot()
-				.Padding(0.0f, 0.0f)
-				.HAlign(HAlign_Left)
+
+			// expander arrow
+			NameColumnBox->AddSlot()
+				.HAlign(HAlign_Right)
 				.VAlign(VAlign_Center)
+				.Padding(0,0,3,0)
 				.AutoWidth()
 				[
-					LeftSideOverlay.ToSharedRef()
+					SNew(SExpanderArrow, SharedThis(this))
+					.BaseIndentLevel(1)
 				];
 
-			if(bHasMultipleColumns)
+			if (bHasMultipleColumns)
 			{
 				// If the NameWidget has already been disabled, don't re-enable it if IsPropertyEditingEnabled is true.
 				NameWidget->SetEnabled(NameWidgetEnabled.IsBound() ?
@@ -487,111 +466,143 @@ void SDetailSingleItemRow::Construct( const FArguments& InArgs, FDetailLayoutCus
 						}))
 					: IsPropertyEditingEnabled);
 
-				TSharedPtr<SHorizontalBox> HBox;
-
-				InternalLeftColumnRowBox->AddSlot()
+				NameColumnBox->AddSlot()
 					.HAlign(Row.NameWidget.HorizontalAlignment)
 					.VAlign(Row.NameWidget.VerticalAlignment)
 					.Padding(DetailWidgetConstants::LeftRowPadding)
 					[
 						NameWidget.ToSharedRef()
 					];
-				InternalLeftColumnRowBox->AddSlot()
-					.Padding(3.0f, 0.0f)
-					.HAlign(HAlign_Right)
-					.VAlign(VAlign_Center)
-					.AutoWidth()
-					[
-						KeyFrameButton
-					];
+		
 
-				GlobalWidgetExtensions.Reset();
-				Module.GetGlobalRowExtensionDelegate().Broadcast(Args, FOnGenerateGlobalRowExtensionArgs::EWidgetPosition::Right, GlobalWidgetExtensions);
-				if (GlobalWidgetExtensions.Num() > 0)
-				{
-					TSharedRef<SHorizontalBox> RightExtensions = SNew(SHorizontalBox);
-
-					for (TSharedRef<SWidget>& GlobalRowExtension: GlobalWidgetExtensions)
-					{
-						RightExtensions->AddSlot()
-						.AutoWidth()
-						[
-							MoveTemp(GlobalRowExtension)
-						];
-					}
-
-					InternalLeftColumnRowBox->AddSlot()
-						.Padding(0.0f, 0.0f)
-						.HAlign(HAlign_Right)
-						.VAlign(VAlign_Center)
-						.AutoWidth()
-						[
-							MoveTemp(RightExtensions)
-						];
-				}
-
-				TSharedRef<SSplitter> Splitter =
-					SNew(SSplitter)
-					.Style(FEditorStyle::Get(), "DetailsView.Splitter")
-					.PhysicalSplitterHandleSize(1.0f)
-					.HitDetectionSplitterHandleSize(5.0f)
-
-					+ SSplitter::Slot()
-					.Value(ColumnSizeData.LeftColumnWidth)
-					.OnSlotResized(SSplitter::FOnSlotResized::CreateSP(this, &SDetailSingleItemRow::OnLeftColumnResized))
+				InnerSplitter->AddSlot()
+					.Value(ColumnSizeData.NameColumnWidth)
+					.OnSlotResized(ColumnSizeData.OnNameColumnResized)
 					[
-						InternalLeftColumnRowBox
-					]
+						NameColumnBox
+					];
 
-					+ SSplitter::Slot()
-					.Value(ColumnSizeData.RightColumnWidth)
-					.OnSlotResized(ColumnSizeData.OnWidthChanged)
+				// create Value column:
+				// | Left  | Name | Value | Right |
+				InnerSplitter->AddSlot()
+					.Value(ColumnSizeData.ValueColumnWidth)
+					.OnSlotResized(ColumnSizeData.OnValueColumnResized) 
 					[
-						SNew(SBox)
-						.IsEnabled(IsPropertyEditingEnabled)
+						SNew(SHorizontalBox)
+						.Clipping(EWidgetClipping::OnDemand)
+						+ SHorizontalBox::Slot()
+						.HAlign(Row.ValueWidget.HorizontalAlignment)
+						.VAlign(Row.ValueWidget.VerticalAlignment)
+						.Padding(DetailWidgetConstants::RightRowPadding)
 						[
-							SNew(SHorizontalBox)
-							+SHorizontalBox::Slot()
-							.FillWidth(1.0f)
-							[
-								SNew(SHorizontalBox)
-								+ SHorizontalBox::Slot()
-								.Padding(DetailWidgetConstants::RightRowPadding)
-								.HAlign(Row.ValueWidget.HorizontalAlignment)
-								.VAlign(Row.ValueWidget.VerticalAlignment)
-								[
-									ValueWidget.ToSharedRef()
-								]
-							]
-							+SHorizontalBox::Slot()
-							.AutoWidth()
-							.VAlign(VAlign_Center)
-							[
-								ExtensionWidget
-							]
+							ValueWidget.ToSharedRef()
 						]
 					];
-				Widget = Splitter;
 			}
 			else
 			{
-				InternalLeftColumnRowBox->SetEnabled(IsPropertyEditingEnabled);
-				InternalLeftColumnRowBox->AddSlot()
+				NameColumnBox->SetEnabled(IsPropertyEditingEnabled);
+				NameColumnBox->AddSlot()
 					.HAlign(Row.WholeRowWidget.HorizontalAlignment)
 					.VAlign(Row.WholeRowWidget.VerticalAlignment)
 					.Padding(DetailWidgetConstants::LeftRowPadding)
 					[
 						Row.WholeRowWidget.Widget
 					];
-				InternalLeftColumnRowBox->AddSlot()
-					.Padding(3.0f, 0.0f)
-					.HAlign(HAlign_Right)
+
+				InnerSplitter->AddSlot()
+					[
+						NameColumnBox
+					];
+			}
+
+			TSharedRef<SHorizontalBox> RightColumnBox = SNew(SHorizontalBox)
+				.Clipping(EWidgetClipping::OnDemand);
+
+			// reset to default widget
+			TSharedPtr<IPropertyHandle> PropertyHandle = GetPropertyHandle();
+			if (PropertyHandle.IsValid() && !PropertyHandle->HasMetaData("NoResetToDefualt") && !PropertyHandle->GetInstanceMetaData("NoResetToDefault"))
+			{
+				TSharedPtr<SWidget> ResetWidget = nullptr;
+
+				if (PropertyHandle->IsResetToDefaultCustomized())
+				{
+					// FIXME: Workaround for JIRA UE-73210.
+					// We had an oscillating SPropertyValueWidget width while dragging a UMG widget in the designer.
+					// The way drag&drop is implemented (SDesignerView::ProcessDropAndAddWidget), a new UCanvasPanelSlot gets 
+					// recreated every frame, so the details panel gets refreshed every frame. Since new property rows are created 
+					// before old ones are destroyed in the details panel, the HasCustomResetToDefault flag on the property node 
+					// toggles from frame to frame, so we alternate between having a ResetToDefaultPropertyEditor and not having one.
+					// By having a spacer fill the blank, the property row layout doesn't change while dragging, but we still see 
+					// a flashing yellow reset arrow (when visible).
+					const FSlateBrush* DiffersFromDefaultBrush = FEditorStyle::GetBrush("PropertyWindow.DiffersFromDefault");
+					ResetWidget = SNew(SSpacer)
+						.Size(DiffersFromDefaultBrush != nullptr ? DiffersFromDefaultBrush->ImageSize : FVector2D(8.0f, 8.0f));
+				}
+				else
+				{
+					SAssignNew(ResetWidget, SResetToDefaultPropertyEditor, PropertyHandle)
+						.IsEnabled(IsPropertyEditingEnabled)
+						.CustomResetToDefault(Row.CustomResetToDefault);
+				}
+
+				RightColumnBox->AddSlot()
+					.HAlign(HAlign_Left)
 					.VAlign(VAlign_Center)
+					.Padding(10,0,0,0)
+					.AutoWidth()
 					[
-						KeyFrameButton
+						ResetWidget.ToSharedRef()
 					];
-				Widget = InternalLeftColumnRowBox;
+			} 
+
+			// keyframe button
+			RightColumnBox->AddSlot()
+				.HAlign(HAlign_Left)
+				.VAlign(VAlign_Center)
+				.Padding(5,0,0,0)
+				.AutoWidth()
+				[
+					CreateKeyframeButton(*Customization, InOwnerTreeNode)
+				];
+
+			// extension widget
+			RightColumnBox->AddSlot()
+				.HAlign(HAlign_Left)
+				.VAlign(VAlign_Center)
+				.Padding(5,0,0,0)
+				.AutoWidth()
+				[
+					ExtensionWidget.ToSharedRef()
+				];
+
+			// fetch global extension widgets 
+			FPropertyEditorModule& PropertyEditorModule = FModuleManager::Get().GetModuleChecked<FPropertyEditorModule>("PropertyEditor");
+
+			TArray<TSharedRef<SWidget>> GlobalWidgetExtensions;
+			FOnGenerateGlobalRowExtensionArgs Args { GetPropertyHandle(), GetPropertyNode(), OwnerTreeNode };
+			PropertyEditorModule.GetGlobalRowExtensionDelegate().Broadcast(Args, GlobalWidgetExtensions);
+
+			if (GlobalWidgetExtensions.Num() > 0)
+			{
+				for (TSharedRef<SWidget>& GlobalRowExtension : GlobalWidgetExtensions)
+				{
+					RightColumnBox->AddSlot()
+						.HAlign(HAlign_Left)
+						.VAlign(VAlign_Center)
+						.AutoWidth()
+						[
+							MoveTemp(GlobalRowExtension)
+						];
+				}
 			}
+
+			OuterSplitter->AddSlot()
+				.Value(ColumnSizeData.RightColumnWidth)
+				.OnSlotResized(ColumnSizeData.OnRightColumnResized)
+			[
+				RightColumnBox
+			];
 		}
 	}
 	else
@@ -602,12 +613,17 @@ void SDetailSingleItemRow::Construct( const FArguments& InArgs, FDetailLayoutCus
 	}
 
 	this->ChildSlot
-	[	
+	[
 		SNew( SBorder )
-		.BorderImage( this, &SDetailSingleItemRow::GetBorderImage )
-		.Padding( FMargin( 0.0f, 0.0f, SDetailTableRowBase::ScrollbarPaddingSize, 0.0f ) )
+		.BorderImage( FEditorStyle::GetBrush( "DetailsView.GridLine" ) )
+		.Padding( FMargin( 0, 0, 0, 1 ) )
 		[
-			Widget
+			SNew( SBorder )
+			.BorderImage( this, &SDetailSingleItemRow::GetBorderImage )
+			.Padding( FMargin( 0.0f, 0.0f, SDetailTableRowBase::ScrollbarPaddingSize, 0.0f ) )
+			[
+				Widget
+			]
 		]
 	];
 
@@ -623,7 +639,6 @@ void SDetailSingleItemRow::Construct( const FArguments& InArgs, FDetailLayoutCus
 	);
 }
 
-
 bool SDetailSingleItemRow::OnContextMenuOpening(FMenuBuilder& MenuBuilder)
 {
 	const bool bIsCopyPasteBound = Customization->GetWidgetRow().IsCopyPasteBound();
@@ -631,7 +646,7 @@ bool SDetailSingleItemRow::OnContextMenuOpening(FMenuBuilder& MenuBuilder)
 	FUIAction CopyAction;
 	FUIAction PasteAction;
 
-	if(bIsCopyPasteBound)
+	if (bIsCopyPasteBound)
 	{
 		CopyAction = Customization->GetWidgetRow().CopyMenuAction;
 		PasteAction = Customization->GetWidgetRow().PasteMenuAction;
@@ -646,6 +661,13 @@ bool SDetailSingleItemRow::OnContextMenuOpening(FMenuBuilder& MenuBuilder)
 			PasteAction.ExecuteAction = FExecuteAction::CreateSP(this, &SDetailSingleItemRow::OnPasteProperty);
 			PasteAction.CanExecuteAction = FCanExecuteAction::CreateSP(this, &SDetailSingleItemRow::CanPasteProperty);
 		}
+		else
+		{
+			CopyAction.ExecuteAction = FExecuteAction::CreateLambda([](){});
+			CopyAction.CanExecuteAction = FCanExecuteAction::CreateLambda([]() { return false; });
+			PasteAction.ExecuteAction = FExecuteAction::CreateLambda([](){});
+			PasteAction.CanExecuteAction = FCanExecuteAction::CreateLambda([]() { return false; });
+		}
 	}
 
 	bool bAddedMenuEntry = false;
@@ -668,10 +690,33 @@ bool SDetailSingleItemRow::OnContextMenuOpening(FMenuBuilder& MenuBuilder)
 			NSLOCTEXT("PropertyView", "PasteProperty_ToolTip", "Paste the copied value here"),
 			FSlateIcon(FCoreStyle::Get().GetStyleSetName(), "GenericCommands.Paste"),
 			PasteAction);
+	}
+
+	FUIAction FavoriteAction;
+	FavoriteAction.ExecuteAction = FExecuteAction::CreateSP(this, &SDetailSingleItemRow::OnFavoriteMenuToggle);
+	FavoriteAction.CanExecuteAction = FCanExecuteAction::CreateLambda([this]()
+		{
+			return Customization->GetPropertyNode().IsValid() && Customization->GetPropertyNode()->CanDisplayFavorite();
+		});
 
-		bAddedMenuEntry = true;
+	FText FavoriteText = NSLOCTEXT("PropertyView", "FavoriteProperty", "Add to Favorites");
+	FText FavoriteTooltipText = NSLOCTEXT("PropertyView", "FavoriteProperty_ToolTip", "Add this property to your favorites.");
+	FName FavoriteIcon = "DetailsView.PropertyIsFavorite";
+
+	bool IsFavorite = Customization->GetPropertyNode().IsValid() && Customization->GetPropertyNode()->IsFavorite();
+	if (IsFavorite)
+	{
+		FavoriteText = NSLOCTEXT("PropertyView", "RemoveFavoriteProperty", "Remove from Favorites");
+		FavoriteTooltipText = NSLOCTEXT("PropertyView", "RemoveFavoriteProperty_ToolTip", "Remove this property from your favorites.");
+		FavoriteIcon = "DetailsView.PropertyIsNotFavorite";
 	}
 
+	MenuBuilder.AddMenuEntry(
+		FavoriteText,
+		FavoriteTooltipText,
+		FSlateIcon(FEditorStyle::Get().GetStyleSetName(), FavoriteIcon),
+		FavoriteAction);
+
 	const TArray<FDetailWidgetRow::FCustomMenuData>& CustomMenuActions = Customization->GetWidgetRow().CustomMenuItems;
 	if (CustomMenuActions.Num() > 0)
 	{
@@ -689,18 +734,11 @@ bool SDetailSingleItemRow::OnContextMenuOpening(FMenuBuilder& MenuBuilder)
 				CustomMenuData.Tooltip,
 				CustomMenuData.SlateIcon,
 				CustomMenuData.Action);
-			bAddedMenuEntry = true;
 		}
 
 	}
 
-	return bAddedMenuEntry;
-}
-
-void SDetailSingleItemRow::OnLeftColumnResized( float InNewWidth )
-{
-	// This has to be bound or the splitter will take it upon itself to determine the size
-	// We do nothing here because it is handled by the column size data
+	return true;
 }
 
 void SDetailSingleItemRow::OnCopyProperty()
@@ -807,7 +845,7 @@ bool SDetailSingleItemRow::CanPasteProperty() const
 		FPropertyEditor* PropertyEditor = PropertyRow->GetPropertyEditor().Get();
 		if (PropertyEditor)
 		{
-			return !PropertyEditor->IsEditConst() && (!PropertyEditor->HasEditCondition() || PropertyEditor->IsEditConditionMet());
+			return !PropertyEditor->IsEditConst();
 		}
 	}
 
@@ -849,7 +887,7 @@ TSharedRef<SWidget> SDetailSingleItemRow::CreateExtensionWidget(TSharedRef<SWidg
 	TSharedPtr<SWidget> ExtensionWidget = SNullWidget::NullWidget;
 
 	if(InTreeNode->GetParentCategory().IsValid())
-	{
+	{ 
 		IDetailsViewPrivate* DetailsView = InTreeNode->GetDetailsView();
 		TSharedPtr<IDetailPropertyExtensionHandler> ExtensionHandler = DetailsView->GetExtensionHandler();
 
diff --git a/Engine/Source/Editor/PropertyEditor/Private/SDetailSingleItemRow.h b/Engine/Source/Editor/PropertyEditor/Private/SDetailSingleItemRow.h
--- a/Engine/Source/Editor/PropertyEditor/Private/SDetailSingleItemRow.h
+++ b/Engine/Source/Editor/PropertyEditor/Private/SDetailSingleItemRow.h
@@ -32,6 +32,7 @@ public:
 
 	void Construct(const FArguments& InArgs);
 	virtual FVector2D ComputeDesiredSize(float LayoutScaleMultiplier) const override;
+
 private:
 	TAttribute< TOptional<float> > MinWidth;
 	TAttribute< TOptional<float> > MaxWidth;
@@ -53,7 +54,6 @@ public:
 		return FReply::Handled().DetectDrag(SharedThis(this), EKeys::LeftMouseButton);
 	};
 
-
 	FReply OnDragDetected(const FGeometry& MyGeometry, const FPointerEvent& MouseEvent) override;
 	TSharedPtr<class FArrayRowDragDropOp> CreateDragDropOperation(TSharedPtr<SDetailSingleItemRow> InRow);
 
@@ -61,8 +61,6 @@ private:
 	TWeakPtr<SDetailSingleItemRow> ParentRow;
 };
 
-
-
 /**
  * A widget for details that span the entire tree row and have no columns                                                              
  */
@@ -85,8 +83,8 @@ public:
 
 protected:
 	virtual bool OnContextMenuOpening( FMenuBuilder& MenuBuilder ) override;
+
 private:
-	void OnLeftColumnResized( float InNewWidth );
 	void OnCopyProperty();
 	void OnPasteProperty();
 	bool CanPasteProperty() const;
@@ -97,8 +95,7 @@ private:
 	FReply OnAddKeyframeClicked();
 	bool IsHighlighted() const;
 
-	const FSlateBrush* GetFavoriteButtonBrush() const;
-	FReply OnFavoriteToggle();
+	void OnFavoriteMenuToggle();
 
 	void OnArrayDragEnter(const FDragDropEvent& DragDropEvent);
 	void OnArrayDragLeave(const FDragDropEvent& DragDropEvent);
@@ -112,14 +109,14 @@ private:
 	/** Checks if the current drop event is being dropped into a valid location
 	 */
 	bool CheckValidDrop(const TSharedPtr<SDetailSingleItemRow> RowPtr) const;
-
+	
 	TSharedPtr<FPropertyNode> GetPropertyNode() const;
 	TSharedPtr<IPropertyHandle> GetPropertyHandle() const;
+
 private:
 	TWeakPtr<IDetailKeyframeHandler> KeyframeHandler;
 	/** Customization for this widget */
 	FDetailLayoutCustomization* Customization;
-	FDetailColumnSizeData ColumnSizeData;
 	bool bAllowFavoriteSystem;
 	bool bIsHoveredDragTarget;
 	bool bIsDragDropObject;
diff --git a/Engine/Source/Editor/PropertyEditor/Private/SDetailsView.cpp b/Engine/Source/Editor/PropertyEditor/Private/SDetailsView.cpp
--- a/Engine/Source/Editor/PropertyEditor/Private/SDetailsView.cpp
+++ b/Engine/Source/Editor/PropertyEditor/Private/SDetailsView.cpp
@@ -47,13 +47,7 @@ void SDetailsView::Construct(const FArguments& InArgs, const FDetailsViewArgs& I
 
 	PropertyUtilities = MakeShareable( new FPropertyDetailsUtilities( *this ) );
 	PropertyGenerationUtilities = MakeShareable( new FDetailsViewPropertyGenerationUtilities(*this) );
-	
-	ColumnWidth = DetailsViewArgs.ColumnWidth;
 
-	ColumnSizeData.LeftColumnWidth = TAttribute<float>( this, &SDetailsView::OnGetLeftColumnWidth );
-	ColumnSizeData.RightColumnWidth = TAttribute<float>( this, &SDetailsView::OnGetRightColumnWidth );
-	ColumnSizeData.OnWidthChanged = SSplitter::FOnSlotResized::CreateSP( this, &SDetailsView::OnSetColumnWidth );
-	
 	// We want the scrollbar to always be visible when objects are selected, but not when there is no selection - however:
 	//  - We can't use AlwaysShowScrollbar for this, as this will also show the scrollbar when nothing is selected
 	//  - We can't use the Visibility construction parameter, as it gets translated into user visibility and can hide the scrollbar even when objects are selected
diff --git a/Engine/Source/Editor/PropertyEditor/Private/SDetailsViewBase.cpp b/Engine/Source/Editor/PropertyEditor/Private/SDetailsViewBase.cpp
--- a/Engine/Source/Editor/PropertyEditor/Private/SDetailsViewBase.cpp
+++ b/Engine/Source/Editor/PropertyEditor/Private/SDetailsViewBase.cpp
@@ -31,7 +31,6 @@ SDetailsViewBase::~SDetailsViewBase()
 	ThumbnailPool.Reset();
 }
 
-
 void SDetailsViewBase::OnGetChildrenForDetailTree(TSharedRef<FDetailTreeNode> InTreeNode, TArray< TSharedRef<FDetailTreeNode> >& OutChildren)
 {
 	InTreeNode->GetChildren(OutChildren);
@@ -39,7 +38,7 @@ void SDetailsViewBase::OnGetChildrenForDetailTree(TSharedRef<FDetailTreeNode> In
 
 TSharedRef<ITableRow> SDetailsViewBase::OnGenerateRowForDetailTree(TSharedRef<FDetailTreeNode> InTreeNode, const TSharedRef<STableViewBase>& OwnerTable)
 {
-	return InTreeNode->GenerateWidgetForTableView(OwnerTable, ColumnSizeData, DetailsViewArgs.bAllowFavoriteSystem);
+	return InTreeNode->GenerateWidgetForTableView(OwnerTable, DetailsViewArgs.bAllowFavoriteSystem);
 }
 
 void SDetailsViewBase::SetRootExpansionStates(const bool bExpand, const bool bRecurse)
@@ -388,7 +387,7 @@ void SDetailsViewBase::UpdateSinglePropertyMap(TSharedPtr<FComplexPropertyNode>
 	TSharedPtr<FComplexPropertyNode> RootPropertyNode = InRootPropertyNode;
 	check(RootPropertyNode.IsValid());
 
-	const bool bEnableFavoriteSystem = IsEngineExitRequested() ? false : (GetDefault<UEditorExperimentalSettings>()->bEnableFavoriteSystem && DetailsViewArgs.bAllowFavoriteSystem);
+	const bool bEnableFavoriteSystem = IsEngineExitRequested() ? false : DetailsViewArgs.bAllowFavoriteSystem;
 
 	DetailLayoutHelpers::FUpdatePropertyMapArgs Args;
 
@@ -507,16 +506,6 @@ TSharedPtr<FAssetThumbnailPool> SDetailsViewBase::GetThumbnailPool() const
 	return ThumbnailPool;
 }
 
-TSharedPtr<FEditConditionParser> SDetailsViewBase::GetEditConditionParser() const
-{
-	if (!EditConditionParser.IsValid())
-	{
-		EditConditionParser = MakeShareable(new FEditConditionParser);
-	}
-
-	return EditConditionParser;
-}
-
 void SDetailsViewBase::NotifyFinishedChangingProperties(const FPropertyChangedEvent& PropertyChangedEvent)
 {
 	OnFinishedChangingPropertiesDelegate.Broadcast(PropertyChangedEvent);
diff --git a/Engine/Source/Editor/PropertyEditor/Private/SDetailsViewBase.h b/Engine/Source/Editor/PropertyEditor/Private/SDetailsViewBase.h
--- a/Engine/Source/Editor/PropertyEditor/Private/SDetailsViewBase.h
+++ b/Engine/Source/Editor/PropertyEditor/Private/SDetailsViewBase.h
@@ -70,14 +70,11 @@ struct FDetailFilter
 	TSet<FPropertyPath> WhitelistedProperties;
 };
 
-
-
 class SDetailsViewBase : public IDetailsViewPrivate
 {
 public:
-	SDetailsViewBase()
-		: ColumnWidth(.65f)
-		, bHasActiveFilter(false)
+	SDetailsViewBase() : 
+		bHasActiveFilter(false)
 		, bIsLocked(false)
 		, bHasOpenColorPicker(false)
 		, bDisableCustomDetailLayouts( false )
@@ -163,7 +160,6 @@ public:
 	virtual void NotifyFinishedChangingProperties(const FPropertyChangedEvent& PropertyChangedEvent) override;
 	void RefreshTree() override;
 	TSharedPtr<class FAssetThumbnailPool> GetThumbnailPool() const override;
-	TSharedPtr<FEditConditionParser> GetEditConditionParser() const override;
 	TSharedPtr<IPropertyUtilities> GetPropertyUtilities() override;
 	void CreateColorPickerWindow(const TSharedRef< class FPropertyEditor >& PropertyEditor, bool bUseAlpha) override;
 	virtual void UpdateSinglePropertyMap(TSharedPtr<FComplexPropertyNode> InRootPropertyNode, FDetailLayoutData& LayoutData, bool bIsExternal) override;
@@ -171,6 +167,7 @@ public:
 	virtual const FCustomPropertyTypeLayoutMap& GetCustomPropertyTypeLayoutMap() const { return InstancedTypeToLayoutMap; }
 	virtual void SaveExpandedItems( TSharedRef<FPropertyNode> StartNode ) override;
 	virtual void RestoreExpandedItems(TSharedRef<FPropertyNode> StartNode) override;
+	virtual FDetailColumnSizeData& GetColumnSizeData() override { return ColumnSizeData; }
 
 	virtual bool IsConnected() const = 0;
 	virtual FRootPropertyNodeList& GetRootNodes() = 0;
@@ -196,11 +193,6 @@ public:
 	 */
 	void SetRootExpansionStates(const bool bExpand, const bool bRecurse);
 
-	/** Column width accessibility */
-	float OnGetLeftColumnWidth() const { return 1.0f - ColumnWidth; }
-	float OnGetRightColumnWidth() const { return ColumnWidth; }
-	void OnSetColumnWidth(float InWidth) { ColumnWidth = InWidth; }
-
 	/**
 	 * Adds an action to execute next tick
 	 */
@@ -403,8 +395,6 @@ protected:
 	mutable FOnFinishedChangingProperties OnFinishedChangingPropertiesDelegate;
 	/** Container for passing around column size data to rows in the tree (each row has a splitter which can affect the column size)*/
 	FDetailColumnSizeData ColumnSizeData;
-	/** The actual width of the right column.  The left column is 1-ColumnWidth */
-	float ColumnWidth;
 	/** True if there is an active filter (text in the filter box) */
 	UE_DEPRECATED(4.26, "Use HasActiveSearch function instead.")
 	bool bHasActiveFilter;
diff --git a/Engine/Source/Editor/PropertyEditor/Private/SSingleProperty.cpp b/Engine/Source/Editor/PropertyEditor/Private/SSingleProperty.cpp
--- a/Engine/Source/Editor/PropertyEditor/Private/SSingleProperty.cpp
+++ b/Engine/Source/Editor/PropertyEditor/Private/SSingleProperty.cpp
@@ -90,11 +90,6 @@ public:
 		return false;
 	}
 
-	virtual TSharedPtr<FEditConditionParser> GetEditConditionParser() const override
-	{
-		return nullptr;
-	}
-
 private:
 	TWeakPtr< SSingleProperty > View;
 };
@@ -179,7 +174,6 @@ void SSingleProperty::SetObject( UObject* InObject )
 			.VAlign( VAlign_Center )
 			[
 				SNew( SPropertyNameWidget, PropertyEditor )
-				.DisplayResetToDefault( false )
 			];
 		}
 
diff --git a/Engine/Source/Editor/PropertyEditor/Private/SStructureDetailsView.cpp b/Engine/Source/Editor/PropertyEditor/Private/SStructureDetailsView.cpp
--- a/Engine/Source/Editor/PropertyEditor/Private/SStructureDetailsView.cpp
+++ b/Engine/Source/Editor/PropertyEditor/Private/SStructureDetailsView.cpp
@@ -46,11 +46,6 @@ void SStructureDetailsView::Construct(const FArguments& InArgs)
 	PropertyUtilities = MakeShareable( new FPropertyDetailsUtilities( *this ) );
 	PropertyGenerationUtilities = MakeShareable(new FDetailsViewPropertyGenerationUtilities(*this));
 	
-	ColumnWidth = DetailsViewArgs.ColumnWidth;
-	ColumnSizeData.LeftColumnWidth = TAttribute<float>(this, &SStructureDetailsView::OnGetLeftColumnWidth);
-	ColumnSizeData.RightColumnWidth = TAttribute<float>(this, &SStructureDetailsView::OnGetRightColumnWidth);
-	ColumnSizeData.OnWidthChanged = SSplitter::FOnSlotResized::CreateSP(this, &SStructureDetailsView::OnSetColumnWidth);
-
 	TSharedRef<SScrollBar> ExternalScrollbar = SNew(SScrollBar);
 
 	// See note in SDetailsView for why visibility is set after construction
diff --git a/Engine/Source/Editor/PropertyEditor/Private/UserInterface/PropertyDetails/PropertyDetailsUtilities.cpp b/Engine/Source/Editor/PropertyEditor/Private/UserInterface/PropertyDetails/PropertyDetailsUtilities.cpp
--- a/Engine/Source/Editor/PropertyEditor/Private/UserInterface/PropertyDetails/PropertyDetailsUtilities.cpp
+++ b/Engine/Source/Editor/PropertyEditor/Private/UserInterface/PropertyDetails/PropertyDetailsUtilities.cpp
@@ -74,8 +74,3 @@ bool FPropertyDetailsUtilities::HasClassDefaultObject() const
 {
 	return DetailsView.HasClassDefaultObject();
 }
-
-TSharedPtr<FEditConditionParser> FPropertyDetailsUtilities::GetEditConditionParser() const
-{
-	return DetailsView.GetEditConditionParser();
-}
diff --git a/Engine/Source/Editor/PropertyEditor/Private/UserInterface/PropertyDetails/PropertyDetailsUtilities.h b/Engine/Source/Editor/PropertyEditor/Private/UserInterface/PropertyDetails/PropertyDetailsUtilities.h
--- a/Engine/Source/Editor/PropertyEditor/Private/UserInterface/PropertyDetails/PropertyDetailsUtilities.h
+++ b/Engine/Source/Editor/PropertyEditor/Private/UserInterface/PropertyDetails/PropertyDetailsUtilities.h
@@ -26,7 +26,6 @@ public:
 	virtual bool DontUpdateValueWhileEditing() const override;
 	const TArray<TWeakObjectPtr<UObject>>& GetSelectedObjects() const override;
 	virtual bool HasClassDefaultObject() const override;
-	virtual TSharedPtr<FEditConditionParser> GetEditConditionParser() const override;
 
 private:
 	IDetailsViewPrivate& DetailsView;
diff --git a/Engine/Source/Editor/PropertyEditor/Private/UserInterface/PropertyEditor/SPropertyEditorAsset.cpp b/Engine/Source/Editor/PropertyEditor/Private/UserInterface/PropertyEditor/SPropertyEditorAsset.cpp
--- a/Engine/Source/Editor/PropertyEditor/Private/UserInterface/PropertyEditor/SPropertyEditorAsset.cpp
+++ b/Engine/Source/Editor/PropertyEditor/Private/UserInterface/PropertyEditor/SPropertyEditorAsset.cpp
@@ -622,21 +622,7 @@ void SPropertyEditorAsset::Construct(const FArguments& InArgs, const TSharedPtr<
 			ActorPicker
 		];
 	}
-
-	if(InArgs._ResetToDefaultSlot.Widget != SNullWidget::NullWidget )
-	{
-		TSharedRef<SWidget> ResetToDefaultWidget  = InArgs._ResetToDefaultSlot.Widget;
-		ResetToDefaultWidget->SetEnabled( IsEnabledAttribute );
-
-		ButtonBox->AddSlot()
-		.Padding( 4.0f, 0.0f )
-		.AutoWidth()
-		.VAlign(VAlign_Center)
-		[
-			ResetToDefaultWidget
-		];
-	}
-
+	
 	if (ButtonBoxWrapper.IsValid())
 	{
 		ButtonBoxWrapper->SetVisibility(ButtonBox->NumSlots() > 0 ? EVisibility::Visible : EVisibility::Collapsed);
diff --git a/Engine/Source/Editor/PropertyEditor/Private/UserInterface/PropertyEditor/SPropertyEditorAsset.h b/Engine/Source/Editor/PropertyEditor/Private/UserInterface/PropertyEditor/SPropertyEditorAsset.h
--- a/Engine/Source/Editor/PropertyEditor/Private/UserInterface/PropertyEditor/SPropertyEditorAsset.h
+++ b/Engine/Source/Editor/PropertyEditor/Private/UserInterface/PropertyEditor/SPropertyEditorAsset.h
@@ -56,7 +56,6 @@ public:
 		, _ObjectPath()
 		, _Class(nullptr)
 		, _CustomContentSlot()
-		, _ResetToDefaultSlot()
 		{}
 		SLATE_ATTRIBUTE( FSlateFontInfo, AssetFont )
 		SLATE_ATTRIBUTE( FSlateFontInfo, ClassFont )
@@ -74,7 +73,6 @@ public:
 		SLATE_EVENT( FOnSetObject, OnSetObject )
 		SLATE_EVENT(FOnShouldFilterAsset, OnShouldFilterAsset)
 		SLATE_NAMED_SLOT( FArguments, CustomContentSlot )
-		SLATE_NAMED_SLOT( FArguments, ResetToDefaultSlot )
 		SLATE_ARGUMENT( TSharedPtr<IPropertyHandle>, PropertyHandle )
 		SLATE_ARGUMENT( TArray<FAssetData>, OwnerAssetDataArray)
 
diff --git a/Engine/Source/Editor/PropertyEditor/Private/UserInterface/PropertyEditor/SPropertyEditorTableRow.cpp b/Engine/Source/Editor/PropertyEditor/Private/UserInterface/PropertyEditor/SPropertyEditorTableRow.cpp
--- a/Engine/Source/Editor/PropertyEditor/Private/UserInterface/PropertyEditor/SPropertyEditorTableRow.cpp
+++ b/Engine/Source/Editor/PropertyEditor/Private/UserInterface/PropertyEditor/SPropertyEditorTableRow.cpp
@@ -72,7 +72,10 @@ TSharedRef< SWidget > SPropertyEditorTableRow::ConstructNameColumnWidget()
 		.AutoWidth()
 		.VAlign(VAlign_Center)
 		[
-			SNew( SEditConditionWidget, PropertyEditor )
+			SNew( SEditConditionWidget )
+			.EditConditionValue(PropertyEditor.ToSharedRef(), &FPropertyEditor::IsEditConditionMet)
+			.OnEditConditionValueChanged_Lambda([this](bool bValue) { PropertyEditor->ToggleEditConditionState(); })
+
 		]
 		+SHorizontalBox::Slot()
 		.AutoWidth()
diff --git a/Engine/Source/Editor/PropertyEditor/Private/UserInterface/PropertyEditor/SResetToDefaultPropertyEditor.cpp b/Engine/Source/Editor/PropertyEditor/Private/UserInterface/PropertyEditor/SResetToDefaultPropertyEditor.cpp
--- a/Engine/Source/Editor/PropertyEditor/Private/UserInterface/PropertyEditor/SResetToDefaultPropertyEditor.cpp
+++ b/Engine/Source/Editor/PropertyEditor/Private/UserInterface/PropertyEditor/SResetToDefaultPropertyEditor.cpp
@@ -71,20 +71,13 @@ FText SResetToDefaultPropertyEditor::GetResetToolTip() const
 
 FReply SResetToDefaultPropertyEditor::OnResetClicked()
 {
-	if (PropertyHandle.IsValid())
+	if (OptionalCustomResetToDefault.IsSet())
 	{
-		if (OptionalCustomResetToDefault.IsSet())
-		{
-			PropertyHandle->ExecuteCustomResetToDefault(OptionalCustomResetToDefault.GetValue());
-		}
-		else
-		{
-			PropertyHandle->ResetToDefault();
-		}
+		OptionalCustomResetToDefault.GetValue().OnResetToDefaultClicked(PropertyHandle);
 	}
-	else if(OptionalCustomResetToDefault.IsSet())
+	else if (PropertyHandle.IsValid())
 	{
-		OptionalCustomResetToDefault.GetValue().OnResetToDefaultClicked().ExecuteIfBound(PropertyHandle);
+		PropertyHandle->ResetToDefault();
 	}
 
 	return FReply::Handled();
diff --git a/Engine/Source/Editor/PropertyEditor/Public/DetailWidgetRow.h b/Engine/Source/Editor/PropertyEditor/Public/DetailWidgetRow.h
--- a/Engine/Source/Editor/PropertyEditor/Public/DetailWidgetRow.h
+++ b/Engine/Source/Editor/PropertyEditor/Public/DetailWidgetRow.h
@@ -3,15 +3,18 @@
 #pragma once
 
 #include "CoreMinimal.h"
-#include "Misc/Attribute.h"
+#include "IDetailPropertyRow.h"
+#include "PropertyHandle.h"
+#include "Framework/Commands/UIAction.h"
 #include "Layout/Visibility.h"
+#include "Misc/Attribute.h"
 #include "Widgets/DeclarativeSyntaxSupport.h"
 #include "Widgets/SWidget.h"
-#include "Framework/Commands/UIAction.h"
 #include "Widgets/Layout/SSpacer.h"
-#include "PropertyHandle.h"
 
 class FDetailWidgetRow;
+struct FCustomEditCondition;
+class FResetToDefaultOverride;
 
 /** Widget declaration for custom widgets in a widget row */
 class FDetailWidgetDecl
@@ -104,7 +107,6 @@ public:
 		, CopyMenuAction()
 		, PasteMenuAction()
 		, RowTagName()
-		, DiffersFromDefaultAttr( false )
 	{
 	}
 	
@@ -160,7 +162,7 @@ public:
 	}
 
 	/**
-	 * Sets the visibility of the entire row
+	 * Sets the enabled state of the entire row
 	 */
 	FDetailWidgetRow& IsEnabled( const TAttribute<bool>& InIsEnabled )
 	{
@@ -239,12 +241,17 @@ public:
 	/**
 	* Sets flag to indicate if property value differs from the default
 	*/
-	FDetailWidgetRow& DiffersFromDefault(const TAttribute<bool>& InDiffersFromDefaultAttr)
+	FDetailWidgetRow& OverrideResetToDefault(const FResetToDefaultOverride& InResetToDefaultOverride)
 	{
-		DiffersFromDefaultAttr = InDiffersFromDefaultAttr;
+		CustomResetToDefault = InResetToDefaultOverride;
 		return *this;
 	}
 
+	/** 
+	 * Override the edit condition.
+	 */ 
+	FDetailWidgetRow& EditCondition(TAttribute<bool> EditConditionValue, FOnBooleanValueChanged OnEditConditionValueChanged);
+	
 	/**
 	* Used to provide all the property handles this WidgetRow represent
 	*/
@@ -296,8 +303,12 @@ public:
 
 	/* Tag to identify this row */
 	FName RowTagName;
-	/* Flag to track if property has been modified from default */
-	TAttribute<bool> DiffersFromDefaultAttr;
+	/** Custom reset to default handler */
+	TOptional<FResetToDefaultOverride> CustomResetToDefault;
+	/** Custom edit condition value. */
+	TAttribute<bool> EditConditionValue;
+	/** Custom edit condition value changed handler. */
+	FOnBooleanValueChanged OnEditConditionValueChanged;
 	/* All property handle that this custom widget represent */
 	TArray<TSharedPtr<IPropertyHandle>> PropertyHandles;
 };
diff --git a/Engine/Source/Editor/PropertyEditor/Public/IDetailPropertyRow.h b/Engine/Source/Editor/PropertyEditor/Public/IDetailPropertyRow.h
--- a/Engine/Source/Editor/PropertyEditor/Public/IDetailPropertyRow.h
+++ b/Engine/Source/Editor/PropertyEditor/Public/IDetailPropertyRow.h
@@ -11,6 +11,64 @@
 class FDetailWidgetRow;
 class FDetailWidgetDecl;
 
+class FDetailColumnSizeData
+{
+public:
+	
+	FDetailColumnSizeData()
+	{
+		ValueColumnWidthValue = 0.65f;
+		RightColumnWidthValue = 0.05f;
+		HoveredSplitterIndexValue = INDEX_NONE;
+
+		NameColumnWidth.BindRaw(this, &FDetailColumnSizeData::GetNameColumnWidth);
+		ValueColumnWidth.BindRaw(this, &FDetailColumnSizeData::GetValueColumnWidth);
+		PropertyColumnWidth.BindRaw(this, &FDetailColumnSizeData::GetPropertyColumnWidth);
+		RightColumnWidth.BindRaw(this, &FDetailColumnSizeData::GetRightColumnWidth);
+		HoveredSplitterIndex.BindRaw(this, &FDetailColumnSizeData::GetHoveredSplitterIndex);
+		OnValueColumnResized.BindRaw(this, &FDetailColumnSizeData::OnSetValueColumnWidth);
+		OnRightColumnResized.BindRaw(this, &FDetailColumnSizeData::OnSetRightColumnWidth);
+		OnSplitterHandleHovered.BindRaw(this, &FDetailColumnSizeData::OnSetHoveredSplitterIndex);
+	}
+
+	TAttribute<float> NameColumnWidth;
+	TAttribute<float> ValueColumnWidth;
+	TAttribute<float> PropertyColumnWidth;
+	TAttribute<float> RightColumnWidth;
+	TAttribute<int32> HoveredSplitterIndex;
+	SSplitter::FOnSlotResized OnNameColumnResized; // intentionally left unbound, since we don't need a handler
+	SSplitter::FOnSlotResized OnValueColumnResized;
+	SSplitter::FOnSlotResized OnRightColumnResized;
+	SSplitter::FOnHandleHovered OnSplitterHandleHovered;
+
+private:
+	float ValueColumnWidthValue;
+	float RightColumnWidthValue;
+	int HoveredSplitterIndexValue;
+
+	float GetNameColumnWidth() const { return 1.0f - ValueColumnWidthValue; }
+	float GetValueColumnWidth() const { return ValueColumnWidthValue; }
+	float GetRightColumnWidth() const { return RightColumnWidthValue; }
+	float GetPropertyColumnWidth() const { return 1.0f - RightColumnWidthValue; }
+	int32 GetHoveredSplitterIndex() const { return HoveredSplitterIndexValue; }
+
+	void OnSetValueColumnWidth(float NewWidth)
+	{ 
+		ensure(NewWidth <= 1.0f); 
+		ValueColumnWidthValue = NewWidth;
+	} 
+
+	void OnSetRightColumnWidth(float NewWidth)
+	{
+		ensure(NewWidth <= 1.0f);
+		RightColumnWidthValue = NewWidth;
+	}
+
+	void OnSetHoveredSplitterIndex(int32 HoveredIndex)
+	{
+		HoveredSplitterIndexValue = HoveredIndex;
+	}
+};
 
 DECLARE_DELEGATE_RetVal_OneParam(bool, FIsResetToDefaultVisible, TSharedPtr<IPropertyHandle> /* PropertyHandle */);
 DECLARE_DELEGATE_OneParam(FResetToDefaultHandler, TSharedPtr<IPropertyHandle> /* PropertyHandle*/);
@@ -26,7 +84,7 @@ public:
 	{
 		FResetToDefaultOverride ResetToDefault;
 		ResetToDefault.bForceShow = true;
-		ResetToDefault.OnResetToDefaultClickedDelegate = InResetToDefaultClicked;
+		ResetToDefault.OnClickedPropertyDelegate = InResetToDefaultClicked;
 		ResetToDefault.bPropagateToChildren = InPropagateToChildren;
 		ResetToDefault.bForceHide = false;
 		return ResetToDefault;
@@ -37,8 +95,31 @@ public:
 	{
 		FResetToDefaultOverride ResetToDefault;
 		ResetToDefault.bForceShow = false;
-		ResetToDefault.IsResetToDefaultVisibleDelegate = InIsResetToDefaultVisible;
-		ResetToDefault.OnResetToDefaultClickedDelegate = InResetToDefaultClicked;
+		ResetToDefault.IsVisiblePropertyDelegate = InIsResetToDefaultVisible;
+		ResetToDefault.OnClickedPropertyDelegate = InResetToDefaultClicked;
+		ResetToDefault.bPropagateToChildren = InPropagateToChildren;
+		ResetToDefault.bForceHide = false;
+		return ResetToDefault;
+	} 
+
+	/** Create a FResetToDefaultOverride from a visibility attribute. */
+	static FResetToDefaultOverride Create(TAttribute<bool> InIsResetToDefaultVisible, const bool InPropagateToChildren = false)
+	{
+		FResetToDefaultOverride ResetToDefault;
+		ResetToDefault.bForceShow = false;
+		ResetToDefault.IsVisibleAttribute = InIsResetToDefaultVisible;
+		ResetToDefault.bPropagateToChildren = InPropagateToChildren;
+		ResetToDefault.bForceHide = false;
+		return ResetToDefault;
+	}
+
+	/** Create a FResetToDefaultOverride from a visibility attribute and a simple delegate. */
+	static FResetToDefaultOverride Create(TAttribute<bool> InIsResetToDefaultVisible, FSimpleDelegate InResetToDefaultClicked, const bool InPropagateToChildren = false)
+	{
+		FResetToDefaultOverride ResetToDefault;
+		ResetToDefault.bForceShow = false;
+		ResetToDefault.IsVisibleAttribute = InIsResetToDefaultVisible;
+		ResetToDefault.OnClickedDelegate = InResetToDefaultClicked;
 		ResetToDefault.bPropagateToChildren = InPropagateToChildren;
 		ResetToDefault.bForceHide = false;
 		return ResetToDefault;
@@ -61,17 +142,42 @@ public:
 		{
 			return true;
 		}
-		if (!bForceHide && IsResetToDefaultVisibleDelegate.IsBound())
+
+		if (bForceHide)
+		{
+			return false;
+		}
+
+		if (IsVisiblePropertyDelegate.IsBound() && Property.IsValid())
+		{
+			return IsVisiblePropertyDelegate.Execute(Property);
+		}
+
+		if (IsVisibleAttribute.IsSet())
 		{
-			return IsResetToDefaultVisibleDelegate.Execute(Property);
+			return IsVisibleAttribute.Get();
 		}
+		
 		return false;
 	}
 
+	/** Does this have a custom reset to default handler? */
+	bool HasResetToDefaultHandler() const
+	{
+		return OnClickedPropertyDelegate.IsBound() || OnClickedDelegate.IsBound();
+	}
+
 	/** Called by the property editor to actually reset the property to default */
-	FResetToDefaultHandler OnResetToDefaultClicked() const
+	void OnResetToDefaultClicked(TSharedPtr<IPropertyHandle> PropertyHandle) const
 	{
-		return OnResetToDefaultClickedDelegate;
+		if (PropertyHandle.IsValid() && OnClickedPropertyDelegate.IsBound())
+		{
+			PropertyHandle->ExecuteCustomResetToDefault(*this);
+		} 
+		else
+		{
+			OnClickedDelegate.ExecuteIfBound();
+		}
 	}
 
 	/** Called by properties to determine whether this override should set on their children */
@@ -82,10 +188,16 @@ public:
 
 private:
 	/** Callback to indicate whether or not reset to default is visible */
-	FIsResetToDefaultVisible IsResetToDefaultVisibleDelegate;
+	FIsResetToDefaultVisible IsVisiblePropertyDelegate;
+
+	/** Delegate called when reset to default is clicked */
+	FResetToDefaultHandler OnClickedPropertyDelegate;
+
+	/** Attribute to determine whether or not reset to default is visible */
+	TAttribute<bool> IsVisibleAttribute; 
 
 	/** Delegate called when reset to default is clicked */
-	FResetToDefaultHandler OnResetToDefaultClickedDelegate;
+	FSimpleDelegate OnClickedDelegate;
 
 	/** Should properties pass this on to their children? */
 	bool bPropagateToChildren;
diff --git a/Engine/Source/Editor/PropertyEditor/Public/IDetailsView.h b/Engine/Source/Editor/PropertyEditor/Public/IDetailsView.h
--- a/Engine/Source/Editor/PropertyEditor/Public/IDetailsView.h
+++ b/Engine/Source/Editor/PropertyEditor/Public/IDetailsView.h
@@ -277,11 +277,9 @@ public:
 	virtual FIsPropertyEditingEnabled& GetIsPropertyEditingEnabledDelegate() = 0;
 
 	virtual void SetKeyframeHandler( TSharedPtr<class IDetailKeyframeHandler> InKeyframeHandler ) = 0;
-
 	virtual TSharedPtr<IDetailKeyframeHandler> GetKeyframeHandler() = 0;
 
-	virtual void SetExtensionHandler(TSharedPtr<class IDetailPropertyExtensionHandler> InExtensionandler) = 0;
-
+	virtual void SetExtensionHandler(TSharedPtr<class IDetailPropertyExtensionHandler> InExtensionHandler) = 0;
 	virtual TSharedPtr<class IDetailPropertyExtensionHandler> GetExtensionHandler() = 0;
 
 	/**
diff --git a/Engine/Source/Editor/PropertyEditor/Public/IPropertyUtilities.h b/Engine/Source/Editor/PropertyEditor/Public/IPropertyUtilities.h
--- a/Engine/Source/Editor/PropertyEditor/Public/IPropertyUtilities.h
+++ b/Engine/Source/Editor/PropertyEditor/Public/IPropertyUtilities.h
@@ -29,5 +29,4 @@ public:
 	virtual const TArray<TWeakObjectPtr<UObject>>& GetSelectedObjects() const = 0;
 	/** If a customization standalone widget is used, the value should be update only once, when its window is closed */
 	virtual bool DontUpdateValueWhileEditing() const = 0;
-	virtual TSharedPtr<FEditConditionParser> GetEditConditionParser() const = 0;
 };
diff --git a/Engine/Source/Editor/PropertyEditor/Public/PropertyCustomizationHelpers.h b/Engine/Source/Editor/PropertyEditor/Public/PropertyCustomizationHelpers.h
--- a/Engine/Source/Editor/PropertyEditor/Public/PropertyCustomizationHelpers.h
+++ b/Engine/Source/Editor/PropertyEditor/Public/PropertyCustomizationHelpers.h
@@ -787,12 +787,3 @@ private:
 	int32 ThumbnailSize;
 	int32 SectionsLodIndex;
 };
-
-struct FDetailColumnSizeData
-{
-	TAttribute<float> LeftColumnWidth;
-	TAttribute<float> RightColumnWidth;
-	SSplitter::FOnSlotResized OnWidthChanged;
-
-	void SetColumnWidth(float InWidth) { OnWidthChanged.ExecuteIfBound(InWidth); }
-};
\ No newline at end of file
diff --git a/Engine/Source/Editor/PropertyEditor/Public/PropertyEditorDelegates.h b/Engine/Source/Editor/PropertyEditor/Public/PropertyEditorDelegates.h
--- a/Engine/Source/Editor/PropertyEditor/Public/PropertyEditorDelegates.h
+++ b/Engine/Source/Editor/PropertyEditor/Public/PropertyEditorDelegates.h
@@ -76,12 +76,6 @@ DECLARE_MULTICAST_DELEGATE_OneParam(FOnFinishedChangingProperties, const FProper
 
 struct FOnGenerateGlobalRowExtensionArgs
 {
-	enum class EWidgetPosition : uint8
-	{
-		Left,
-		Right
-	};
-
 	/** The detail row's property handle. */
 	TSharedPtr<IPropertyHandle> PropertyHandle;
 	/** The detail row's property node. */
@@ -96,7 +90,7 @@ struct FOnGenerateGlobalRowExtensionArgs
  * When called, EWidgetPosition indicates the position for which the delegate is gathering extensions.
  * ie. The favorite system is implemented by adding the star widget when the delegate is called with the left position.
  */
-DECLARE_MULTICAST_DELEGATE_ThreeParams(FOnGenerateGlobalRowExtension, const FOnGenerateGlobalRowExtensionArgs& /*InArgs*/, FOnGenerateGlobalRowExtensionArgs::EWidgetPosition /*InWidgetPosition*/, TArray<TSharedRef<class SWidget>>& /*OutExtensions*/);
+DECLARE_MULTICAST_DELEGATE_TwoParams(FOnGenerateGlobalRowExtension, const FOnGenerateGlobalRowExtensionArgs& /*InArgs*/, TArray<TSharedRef<class SWidget>>& /*OutExtensions*/);
 
 /**
  * Callback executed to query the custom layout of details
