Added new cpp file that was missed in CL14375859.

diff --git a/Engine/Source/Editor/PropertyEditor/Private/DetailWidgetRow.cpp b/Engine/Source/Editor/PropertyEditor/Private/DetailWidgetRow.cpp
new file mode 100644
--- /dev/null
+++ b/Engine/Source/Editor/PropertyEditor/Private/DetailWidgetRow.cpp
@@ -0,0 +1,11 @@
+// Copyright Epic Games, Inc. All Rights Reserved.
+
+#include "DetailWidgetRow.h"
+#include "PropertyEditorHelpers.h"
+
+FDetailWidgetRow& FDetailWidgetRow::EditCondition(TAttribute<bool> InEditConditionValue, FOnBooleanValueChanged InOnEditConditionValueChanged)
+{
+	EditConditionValue = InEditConditionValue;
+	OnEditConditionValueChanged = InOnEditConditionValueChanged;
+	return *this;
+}
\ No newline at end of file
