Fix typo that broke 'NoResetToDefault' metadata.

diff --git a/Engine/Source/Editor/PropertyEditor/Private/SDetailSingleItemRow.cpp b/Engine/Source/Editor/PropertyEditor/Private/SDetailSingleItemRow.cpp
--- a/Engine/Source/Editor/PropertyEditor/Private/SDetailSingleItemRow.cpp
+++ b/Engine/Source/Editor/PropertyEditor/Private/SDetailSingleItemRow.cpp
@@ -521,7 +521,7 @@ void SDetailSingleItemRow::Construct( const FArguments& InArgs, FDetailLayoutCus
 
 			// reset to default widget
 			TSharedPtr<IPropertyHandle> PropertyHandle = GetPropertyHandle();
-			if (PropertyHandle.IsValid() && !PropertyHandle->HasMetaData("NoResetToDefualt") && !PropertyHandle->GetInstanceMetaData("NoResetToDefault"))
+			if (PropertyHandle.IsValid() && !PropertyHandle->HasMetaData("NoResetToDefault") && !PropertyHandle->GetInstanceMetaData("NoResetToDefault"))
 			{
 				TSharedPtr<SWidget> ResetWidget = nullptr;
 
