Details panel extension widgets are now right-aligned in the 'Value' slot, rather than the 'Right' slot for a single details item

This allows them to be unclipped by default. May change in the future but unblocks functionality for now.

diff --git a/Engine/Source/Editor/PropertyEditor/Private/SDetailSingleItemRow.cpp b/Engine/Source/Editor/PropertyEditor/Private/SDetailSingleItemRow.cpp
--- a/Engine/Source/Editor/PropertyEditor/Private/SDetailSingleItemRow.cpp
+++ b/Engine/Source/Editor/PropertyEditor/Private/SDetailSingleItemRow.cpp
@@ -497,6 +497,15 @@ void SDetailSingleItemRow::Construct( const FArguments& InArgs, FDetailLayoutCus
 						[
 							ValueWidget.ToSharedRef()
 						]
+						// extension widget
+						+ SHorizontalBox::Slot()
+						.HAlign(HAlign_Right)
+						.VAlign(VAlign_Center)
+						.Padding(3,0,3,0)
+						.AutoWidth()
+						[
+							ExtensionWidget.ToSharedRef()
+						]
 					];
 			}
 			else
@@ -566,16 +575,6 @@ void SDetailSingleItemRow::Construct( const FArguments& InArgs, FDetailLayoutCus
 					CreateKeyframeButton(*Customization, InOwnerTreeNode)
 				];
 
-			// extension widget
-			RightColumnBox->AddSlot()
-				.HAlign(HAlign_Left)
-				.VAlign(VAlign_Center)
-				.Padding(5,0,0,0)
-				.AutoWidth()
-				[
-					ExtensionWidget.ToSharedRef()
-				];
-
 			// fetch global extension widgets 
 			FPropertyEditorModule& PropertyEditorModule = FModuleManager::Get().GetModuleChecked<FPropertyEditorModule>("PropertyEditor");
 
