Exposed SResetToDefaultPropertyEditor publicly to allow it to be used for material slots.

diff --git a/Engine/Source/Editor/PropertyEditor/Private/DetailPropertyRow.cpp b/Engine/Source/Editor/PropertyEditor/Private/DetailPropertyRow.cpp
--- a/Engine/Source/Editor/PropertyEditor/Private/DetailPropertyRow.cpp
+++ b/Engine/Source/Editor/PropertyEditor/Private/DetailPropertyRow.cpp
@@ -3,7 +3,7 @@
 #include "DetailPropertyRow.h"
 #include "Modules/ModuleManager.h"
 #include "PropertyCustomizationHelpers.h"
-#include "UserInterface/PropertyEditor/SResetToDefaultPropertyEditor.h"
+#include "SResetToDefaultPropertyEditor.h"
 #include "DetailItemNode.h"
 #include "DetailCategoryGroupNode.h"
 #include "ObjectPropertyNode.h"
diff --git a/Engine/Source/Editor/PropertyEditor/Private/PropertyEditorHelpers.cpp b/Engine/Source/Editor/PropertyEditor/Private/PropertyEditorHelpers.cpp
--- a/Engine/Source/Editor/PropertyEditor/Private/PropertyEditorHelpers.cpp
+++ b/Engine/Source/Editor/PropertyEditor/Private/PropertyEditorHelpers.cpp
@@ -7,6 +7,7 @@
 #include "IDocumentation.h"
 
 #include "PropertyHandleImpl.h"
+#include "SResetToDefaultPropertyEditor.h"
 
 #include "UserInterface/PropertyEditor/SPropertyEditor.h"
 #include "UserInterface/PropertyEditor/SPropertyEditorNumeric.h"
@@ -18,7 +19,6 @@
 #include "UserInterface/PropertyEditor/SPropertyEditorArrayItem.h"
 #include "UserInterface/PropertyEditor/SPropertyEditorTitle.h"
 #include "UserInterface/PropertyEditor/SPropertyEditorDateTime.h"
-#include "UserInterface/PropertyEditor/SResetToDefaultPropertyEditor.h"
 #include "UserInterface/PropertyEditor/SPropertyEditorAsset.h"
 #include "UserInterface/PropertyEditor/SPropertyEditorClass.h"
 #include "UserInterface/PropertyEditor/SPropertyEditorStruct.h"
diff --git a/Engine/Source/Editor/PropertyEditor/Private/SDetailSingleItemRow.cpp b/Engine/Source/Editor/PropertyEditor/Private/SDetailSingleItemRow.cpp
--- a/Engine/Source/Editor/PropertyEditor/Private/SDetailSingleItemRow.cpp
+++ b/Engine/Source/Editor/PropertyEditor/Private/SDetailSingleItemRow.cpp
@@ -529,41 +529,32 @@ void SDetailSingleItemRow::Construct( const FArguments& InArgs, FDetailLayoutCus
 				.Clipping(EWidgetClipping::OnDemand);
 
 			// reset to default widget
+			TSharedPtr<SWidget> ResetWidget; 
+
+			bool bCreateResetToDefault = Row.CustomResetToDefault.IsSet();
+
 			TSharedPtr<IPropertyHandle> PropertyHandle = GetPropertyHandle();
-			if (PropertyHandle.IsValid() && !PropertyHandle->HasMetaData("NoResetToDefault") && !PropertyHandle->GetInstanceMetaData("NoResetToDefault"))
+			if (PropertyHandle.IsValid() )
 			{
-				TSharedPtr<SWidget> ResetWidget = nullptr;
-
-				if (PropertyHandle->IsResetToDefaultCustomized())
-				{
-					// FIXME: Workaround for JIRA UE-73210.
-					// We had an oscillating SPropertyValueWidget width while dragging a UMG widget in the designer.
-					// The way drag&drop is implemented (SDesignerView::ProcessDropAndAddWidget), a new UCanvasPanelSlot gets 
-					// recreated every frame, so the details panel gets refreshed every frame. Since new property rows are created 
-					// before old ones are destroyed in the details panel, the HasCustomResetToDefault flag on the property node 
-					// toggles from frame to frame, so we alternate between having a ResetToDefaultPropertyEditor and not having one.
-					// By having a spacer fill the blank, the property row layout doesn't change while dragging, but we still see 
-					// a flashing yellow reset arrow (when visible).
-					const FSlateBrush* DiffersFromDefaultBrush = FEditorStyle::GetBrush("PropertyWindow.DiffersFromDefault");
-					ResetWidget = SNew(SSpacer)
-						.Size(DiffersFromDefaultBrush != nullptr ? DiffersFromDefaultBrush->ImageSize : FVector2D(8.0f, 8.0f));
-				}
-				else
+				if (PropertyHandle->HasMetaData("NoResetToDefault") || PropertyHandle->GetInstanceMetaData("NoResetToDefault"))
 				{
-					SAssignNew(ResetWidget, SResetToDefaultPropertyEditor, PropertyHandle)
-						.IsEnabled(IsPropertyEditingEnabled)
-						.CustomResetToDefault(Row.CustomResetToDefault);
+					bCreateResetToDefault = false;
 				}
+			} 
 
+			if (bCreateResetToDefault)
+			{
 				RightColumnBox->AddSlot()
-					.HAlign(HAlign_Left)
-					.VAlign(VAlign_Center)
-					.Padding(10,0,0,0)
-					.AutoWidth()
-					[
-						ResetWidget.ToSharedRef()
-					];
-			} 
+				.HAlign(HAlign_Left)
+				.VAlign(VAlign_Center)
+				.Padding(5,0,0,0)
+				.AutoWidth()
+				[
+					SNew(SResetToDefaultPropertyEditor, PropertyHandle)
+					.IsEnabled(IsPropertyEditingEnabled)
+					.CustomResetToDefault(Row.CustomResetToDefault)
+				];
+			}
 
 			// keyframe button
 			RightColumnBox->AddSlot()
diff --git a/Engine/Source/Editor/PropertyEditor/Private/SSingleProperty.cpp b/Engine/Source/Editor/PropertyEditor/Private/SSingleProperty.cpp
--- a/Engine/Source/Editor/PropertyEditor/Private/SSingleProperty.cpp
+++ b/Engine/Source/Editor/PropertyEditor/Private/SSingleProperty.cpp
@@ -11,7 +11,7 @@
 #include "Presentation/PropertyEditor/PropertyEditor.h"
 #include "ObjectPropertyNode.h"
 #include "PropertyEditorHelpers.h"
-#include "UserInterface/PropertyEditor/SResetToDefaultPropertyEditor.h"
+#include "SResetToDefaultPropertyEditor.h"
 #include "Widgets/Colors/SColorPicker.h"
 
 
diff --git a/Engine/Source/Editor/PropertyEditor/Private/UserInterface/PropertyEditor/SResetToDefaultPropertyEditor.cpp b/Engine/Source/Editor/PropertyEditor/Private/UserInterface/PropertyEditor/SResetToDefaultPropertyEditor.cpp
--- a/Engine/Source/Editor/PropertyEditor/Private/UserInterface/PropertyEditor/SResetToDefaultPropertyEditor.cpp
+++ b/Engine/Source/Editor/PropertyEditor/Private/UserInterface/PropertyEditor/SResetToDefaultPropertyEditor.cpp
@@ -1,5 +1,5 @@
 // Copyright Epic Games, Inc. All Rights Reserved.
-#include "UserInterface/PropertyEditor/SResetToDefaultPropertyEditor.h"
+#include "SResetToDefaultPropertyEditor.h"
 #include "Widgets/Images/SImage.h"
 #include "Widgets/Input/SButton.h"
 
diff --git a/Engine/Source/Editor/PropertyEditor/Private/UserInterface/PropertyEditor/SResetToDefaultPropertyEditor.h b/Engine/Source/Editor/PropertyEditor/Private/UserInterface/PropertyEditor/SResetToDefaultPropertyEditor.h
deleted file mode 100644
diff --git a/Engine/Source/Editor/PropertyEditor/Private/UserInterface/PropertyTable/BooleanPropertyTableCellPresenter.cpp b/Engine/Source/Editor/PropertyEditor/Private/UserInterface/PropertyTable/BooleanPropertyTableCellPresenter.cpp
--- a/Engine/Source/Editor/PropertyEditor/Private/UserInterface/PropertyTable/BooleanPropertyTableCellPresenter.cpp
+++ b/Engine/Source/Editor/PropertyEditor/Private/UserInterface/PropertyTable/BooleanPropertyTableCellPresenter.cpp
@@ -8,7 +8,7 @@
 #include "EditorStyleSet.h"
 
 #include "UserInterface/PropertyEditor/SPropertyEditorBool.h"
-#include "UserInterface/PropertyEditor/SResetToDefaultPropertyEditor.h"
+#include "SResetToDefaultPropertyEditor.h"
 
 FBooleanPropertyTableCellPresenter::FBooleanPropertyTableCellPresenter( const TSharedRef< class FPropertyEditor >& InPropertyEditor )
 	: PropertyEditor( InPropertyEditor )
diff --git a/Engine/Source/Editor/PropertyEditor/Private/UserInterface/PropertyTable/ColorPropertyTableCellPresenter.cpp b/Engine/Source/Editor/PropertyEditor/Private/UserInterface/PropertyTable/ColorPropertyTableCellPresenter.cpp
--- a/Engine/Source/Editor/PropertyEditor/Private/UserInterface/PropertyTable/ColorPropertyTableCellPresenter.cpp
+++ b/Engine/Source/Editor/PropertyEditor/Private/UserInterface/PropertyTable/ColorPropertyTableCellPresenter.cpp
@@ -10,7 +10,7 @@
 #include "IPropertyTableUtilities.h"
 
 #include "UserInterface/PropertyEditor/SPropertyEditorColor.h"
-#include "UserInterface/PropertyEditor/SResetToDefaultPropertyEditor.h"
+#include "SResetToDefaultPropertyEditor.h"
 
 FColorPropertyTableCellPresenter::FColorPropertyTableCellPresenter( const TSharedRef< class FPropertyEditor >& InPropertyEditor, const TSharedRef< class IPropertyTableUtilities >& InPropertyUtilities )
 	: PropertyEditor( InPropertyEditor )
diff --git a/Engine/Source/Editor/PropertyEditor/Private/UserInterface/PropertyTable/TextPropertyTableCellPresenter.cpp b/Engine/Source/Editor/PropertyEditor/Private/UserInterface/PropertyTable/TextPropertyTableCellPresenter.cpp
--- a/Engine/Source/Editor/PropertyEditor/Private/UserInterface/PropertyTable/TextPropertyTableCellPresenter.cpp
+++ b/Engine/Source/Editor/PropertyEditor/Private/UserInterface/PropertyTable/TextPropertyTableCellPresenter.cpp
@@ -8,9 +8,9 @@
 #include "Widgets/Layout/SBorder.h"
 #include "IPropertyTableUtilities.h"
 #include "PropertyEditorHelpers.h"
+#include "SResetToDefaultPropertyEditor.h"
 
 #include "UserInterface/PropertyEditor/SPropertyEditor.h"
-#include "UserInterface/PropertyEditor/SResetToDefaultPropertyEditor.h"
 #include "UserInterface/PropertyEditor/SPropertyEditorNumeric.h"
 #include "UserInterface/PropertyEditor/SPropertyEditorCombo.h"
 #include "UserInterface/PropertyEditor/SPropertyEditorEditInline.h"
diff --git a/Engine/Source/Editor/PropertyEditor/Public/SResetToDefaultPropertyEditor.h b/Engine/Source/Editor/PropertyEditor/Public/SResetToDefaultPropertyEditor.h
new file mode 100644
--- /dev/null
+++ b/Engine/Source/Editor/PropertyEditor/Public/SResetToDefaultPropertyEditor.h
@@ -0,0 +1,49 @@
+// Copyright Epic Games, Inc. All Rights Reserved.
+#pragma once
+
+#include "CoreMinimal.h"
+#include "Layout/Visibility.h"
+#include "Input/Reply.h"
+#include "Widgets/DeclarativeSyntaxSupport.h"
+#include "Widgets/SCompoundWidget.h"
+#include "IDetailPropertyRow.h"
+
+// Forward decl
+class FResetToDefaultOverride;
+
+/** Widget showing the reset to default value button */
+class PROPERTYEDITOR_API SResetToDefaultPropertyEditor : public SCompoundWidget
+{
+public:
+
+	SLATE_BEGIN_ARGS( SResetToDefaultPropertyEditor )
+		: _NonVisibleState( EVisibility::Hidden )
+		{}
+		SLATE_ARGUMENT( EVisibility, NonVisibleState )
+		SLATE_ARGUMENT( TOptional<FResetToDefaultOverride>, CustomResetToDefault )
+	SLATE_END_ARGS()
+
+	~SResetToDefaultPropertyEditor();
+
+	void Construct( const FArguments& InArgs, const TSharedPtr< class IPropertyHandle>& InPropertyHandle );
+
+private:
+	FText GetResetToolTip() const;
+
+	EVisibility GetDiffersFromDefaultAsVisibility() const;
+
+	void Tick( const FGeometry& AllottedGeometry, const double InCurrentTime, const float InDeltaTime );
+
+	FReply OnResetClicked();
+
+	void UpdateDiffersFromDefaultState();
+
+private:
+	TOptional<FResetToDefaultOverride> OptionalCustomResetToDefault;
+
+	TSharedPtr< class IPropertyHandle > PropertyHandle;
+
+	EVisibility NonVisibleState;
+
+	bool bValueDiffersFromDefault;
+};
