Fixed bug causing the details view splitter to become unlinked for a single row after having been resized from that row, because the delegate it was being passed was (intentionally) unbound. Now it uses an empty lambda instead.

diff --git a/Engine/Source/Editor/PropertyEditor/Public/IDetailPropertyRow.h b/Engine/Source/Editor/PropertyEditor/Public/IDetailPropertyRow.h
--- a/Engine/Source/Editor/PropertyEditor/Public/IDetailPropertyRow.h
+++ b/Engine/Source/Editor/PropertyEditor/Public/IDetailPropertyRow.h
@@ -32,6 +32,10 @@ public:
 		OnValueColumnResized.BindRaw(this, &FDetailColumnSizeData::OnSetValueColumnWidth);
 		OnRightColumnResized.BindRaw(this, &FDetailColumnSizeData::OnSetRightColumnWidth);
 		OnSplitterHandleHovered.BindRaw(this, &FDetailColumnSizeData::OnSetHoveredSplitterIndex);
+
+		// these are intentionally left no-op, since the widths are derived from the other column width
+		OnPropertyColumnResized.BindLambda([](float) {});
+		OnNameColumnResized.BindLambda([](float) {}); 
 	}
 
 	TAttribute<float> NameColumnWidth;
@@ -39,9 +43,9 @@ public:
 	TAttribute<float> PropertyColumnWidth;
 	TAttribute<float> RightColumnWidth;
 	TAttribute<int32> HoveredSplitterIndex;
-	SSplitter::FOnSlotResized OnNameColumnResized; // intentionally left unbound, since we don't need a handler as the width is derived from the value width
+	SSplitter::FOnSlotResized OnNameColumnResized; 
 	SSplitter::FOnSlotResized OnValueColumnResized;
-	SSplitter::FOnSlotResized OnPropertyColumnResized; // intentionally left unbound, since we don't need a handler as the width is derived from the right width
+	SSplitter::FOnSlotResized OnPropertyColumnResized; 
 	SSplitter::FOnSlotResized OnRightColumnResized;
 	SSplitter::FOnHandleHovered OnSplitterHandleHovered;
 
