Reset to Default now always reserves the leftmost space in the right column of the details view.

diff --git a/Engine/Source/Editor/PropertyEditor/Private/SDetailSingleItemRow.cpp b/Engine/Source/Editor/PropertyEditor/Private/SDetailSingleItemRow.cpp
--- a/Engine/Source/Editor/PropertyEditor/Private/SDetailSingleItemRow.cpp
+++ b/Engine/Source/Editor/PropertyEditor/Private/SDetailSingleItemRow.cpp
@@ -529,8 +529,6 @@ void SDetailSingleItemRow::Construct( const FArguments& InArgs, FDetailLayoutCus
 				.Clipping(EWidgetClipping::OnDemand);
 
 			// reset to default widget
-			TSharedPtr<SWidget> ResetWidget; 
-
 			bool bCreateResetToDefault = Row.CustomResetToDefault.IsSet();
 
 			TSharedPtr<IPropertyHandle> PropertyHandle = GetPropertyHandle();
@@ -542,19 +540,29 @@ void SDetailSingleItemRow::Construct( const FArguments& InArgs, FDetailLayoutCus
 				}
 			} 
 
+			TSharedRef<SWidget> ResetWidget = SNullWidget::NullWidget;
 			if (bCreateResetToDefault)
 			{
-				RightColumnBox->AddSlot()
+				ResetWidget = SNew(SResetToDefaultPropertyEditor, PropertyHandle)
+					.IsEnabled(IsPropertyEditingEnabled)
+					.CustomResetToDefault(Row.CustomResetToDefault);
+			}	
+			else
+			{
+				// leave space in the reset to default column regardless
+				const FSlateBrush* DiffersFromDefaultBrush = FEditorStyle::GetBrush("PropertyWindow.DiffersFromDefault");
+				ResetWidget = SNew(SSpacer)
+					.Size(DiffersFromDefaultBrush != nullptr ? DiffersFromDefaultBrush->ImageSize : FVector2D(8.0f, 8.0f));
+			}
+
+			RightColumnBox->AddSlot()
 				.HAlign(HAlign_Left)
 				.VAlign(VAlign_Center)
 				.Padding(5,0,0,0)
 				.AutoWidth()
 				[
-					SNew(SResetToDefaultPropertyEditor, PropertyHandle)
-					.IsEnabled(IsPropertyEditingEnabled)
-					.CustomResetToDefault(Row.CustomResetToDefault)
+					ResetWidget
 				];
-			}
 
 			// keyframe button
 			RightColumnBox->AddSlot()
