Fixed bug causing categories that had been hidden by customizations to be visible.

I'd accidentally submitted some changes that enabled the favourites system prematurely, and apparently many rows were being counted as favorited despite not being so. This is a temporary fix that will go away when further editor config and favourites system work is done.

diff --git a/Engine/Source/Editor/PropertyEditor/Private/SDetailsViewBase.cpp b/Engine/Source/Editor/PropertyEditor/Private/SDetailsViewBase.cpp
--- a/Engine/Source/Editor/PropertyEditor/Private/SDetailsViewBase.cpp
+++ b/Engine/Source/Editor/PropertyEditor/Private/SDetailsViewBase.cpp
@@ -387,7 +387,7 @@ void SDetailsViewBase::UpdateSinglePropertyMap(TSharedPtr<FComplexPropertyNode>
 	TSharedPtr<FComplexPropertyNode> RootPropertyNode = InRootPropertyNode;
 	check(RootPropertyNode.IsValid());
 
-	const bool bEnableFavoriteSystem = IsEngineExitRequested() ? false : DetailsViewArgs.bAllowFavoriteSystem;
+	const bool bEnableFavoriteSystem = false; //TODO: IsEngineExitRequested() ? false : DetailsViewArgs.bAllowFavoriteSystem;
 
 	DetailLayoutHelpers::FUpdatePropertyMapArgs Args;
 
