Details view now shades nested property rows based on their nesting level, and displays a drop shadow to indicate hierarchy.

diff --git a/Engine/Source/Editor/EditorStyle/Private/StarshipStyle.cpp b/Engine/Source/Editor/EditorStyle/Private/StarshipStyle.cpp
--- a/Engine/Source/Editor/EditorStyle/Private/StarshipStyle.cpp
+++ b/Engine/Source/Editor/EditorStyle/Private/StarshipStyle.cpp
@@ -3176,11 +3176,7 @@ void FStarshipEditorStyle::FStyle::SetupPropertyEditorStyles()
 		Set( "DetailsView.CategoryTop",               new FSlateColorBrush(FStyleColors::Header));
 		Set( "DetailsView.CategoryTop_Hovered",       new FSlateColorBrush(FStyleColors::Header));
 		Set( "DetailsView.CategoryBottom",            new FSlateColorBrush(FStyleColors::Recessed));
-
-		Set( "DetailsView.CategoryMiddle",            new FSlateColorBrush(FStyleColors::Background));
-		Set( "DetailsView.CategoryMiddle_Hovered",    new FSlateColorBrush(FStyleColors::Header));
-		Set( "DetailsView.CategoryMiddle_Highlighted",new FSlateColorBrush(FStyleColors::Background));
-		Set( "DetailsView.CategoryMiddle_Active",     new FSlateColorBrush(FStyleColors::Background));
+		Set( "DetailsView.CategoryMiddle",            new FSlateColorBrush(FStyleColors::White));
 
 		Set( "DetailsView.PropertyIsFavorite", new IMAGE_BRUSH("PropertyView/Favorites_Enabled", Icon12x12));
 		Set( "DetailsView.PropertyIsNotFavorite", new IMAGE_BRUSH("PropertyView/Favorites_Disabled", Icon12x12));
@@ -3206,6 +3202,8 @@ void FStarshipEditorStyle::FStyle::SetupPropertyEditorStyles()
 			.SetFont(GetFontStyle("DetailsView.CategoryFontStyle"))
 			.SetColorAndOpacity(FStyleColors::ForegroundHeader)
 		);
+		
+		Set("DetailsView.ArrayDropShadow", new IMAGE_BRUSH("Common/ArrayDropShadow", FVector2D(32,2)));
 
 		Set( "DetailsView.TreeView.TableRow", FTableRowStyle()
 			.SetEvenRowBackgroundBrush( FSlateNoResource() )
diff --git a/Engine/Source/Editor/PropertyEditor/Private/PropertyEditorHelpers.cpp b/Engine/Source/Editor/PropertyEditor/Private/PropertyEditorHelpers.cpp
--- a/Engine/Source/Editor/PropertyEditor/Private/PropertyEditorHelpers.cpp
+++ b/Engine/Source/Editor/PropertyEditor/Private/PropertyEditorHelpers.cpp
@@ -821,9 +821,11 @@ namespace PropertyEditorHelpers
 		TSharedRef<SArrayRowHandle> Handle = SNew(SArrayRowHandle)
 			.Content()
 			[
-				SNew(SHorizontalBox)
-				+ SHorizontalBox::Slot()
-				.Padding(5.0f, 0.0f)
+				SNew(SBox)
+				.Padding(0.0f)
+				.HAlign(HAlign_Center)
+				.VAlign(VAlign_Center)
+				.WidthOverride(16.0f)
 				[
 					SNew(SImage)
 					.Image(FCoreStyle::Get().GetBrush("VerticalBoxDragIndicatorShort"))
diff --git a/Engine/Source/Editor/PropertyEditor/Private/SDetailSingleItemRow.cpp b/Engine/Source/Editor/PropertyEditor/Private/SDetailSingleItemRow.cpp
--- a/Engine/Source/Editor/PropertyEditor/Private/SDetailSingleItemRow.cpp
+++ b/Engine/Source/Editor/PropertyEditor/Private/SDetailSingleItemRow.cpp
@@ -18,6 +18,217 @@
 #include "Modules/ModuleManager.h"
 #include "SResetToDefaultPropertyEditor.h"
 
+
+namespace DetailWidgetConstants
+{
+	const FMargin LeftRowPadding( 12.0f, 2.5f, 12.0f, 2.5f );
+	const FMargin RightRowPadding( 12.0f, 2.5f, 2.0f, 2.5f );
+}
+
+namespace SDetailSingleItemRow_Helper
+{
+	//Get the node item number in case it is expand we have to recursively count all expanded children
+	void RecursivelyGetItemShow(TSharedRef<FDetailTreeNode> ParentItem, int32 &ItemShowNum)
+	{
+		if (ParentItem->GetVisibility() == ENodeVisibility::Visible)
+		{
+			ItemShowNum++;
+		}
+
+		if (ParentItem->ShouldBeExpanded())
+		{
+			TArray< TSharedRef<FDetailTreeNode> > Childrens;
+			ParentItem->GetChildren(Childrens);
+			for (TSharedRef<FDetailTreeNode> ItemChild : Childrens)
+			{
+				RecursivelyGetItemShow(ItemChild, ItemShowNum);
+			}
+		}
+	}
+
+	FSlateColor GetBackgroundColor(int32 IndentLevel)
+	{
+		FLinearColor BackgroundColor = FAppStyle::Get().GetSlateColor("Colors.Background").GetSpecifiedColor();
+
+		return FSlateColor(BackgroundColor + FLinearColor(2.0f/255, 2.0f/255, 2.0f/255) * (IndentLevel % 4));
+	}
+}
+
+class SDetailsIndent : public SCompoundWidget
+{
+public:
+
+	SLATE_BEGIN_ARGS(SDetailsIndent) {}
+	SLATE_END_ARGS()
+
+	TSharedPtr<SHorizontalBox> ChildBox;
+	TSharedPtr<SButton> ExpanderArrow;
+	TSharedPtr<SDetailSingleItemRow> Row;
+
+	void Construct(const FArguments& InArgs, TSharedRef<SDetailSingleItemRow> DetailsRow)
+	{
+		Row = DetailsRow;
+
+		ChildSlot
+		[
+			SAssignNew(ChildBox, SHorizontalBox)
+		];
+	}
+
+	void Tick(const FGeometry& AllottedGeometry, const double InCurrentTime, const float InDeltaTime) override
+	{
+		int32 IndentLevel = Row->GetIndentLevel();
+		if (ChildBox->NumSlots() == IndentLevel)
+		{
+			// already initialized, don't need to do anything
+			return;
+		}
+
+		ChildBox->ClearChildren();
+
+		for (int32 i = 0; i < IndentLevel; ++i)
+		{
+			SHorizontalBox::FSlot& NewSlot = ChildBox->AddSlot()
+				.VAlign(VAlign_Fill)
+				.HAlign(HAlign_Left)
+				.Padding(0);
+
+			if (i == IndentLevel - 1)
+			{
+				NewSlot
+				.AutoWidth()
+				[
+					SNew(SBox)
+					.WidthOverride(32)
+					[
+						SNew(SBorder)
+						.BorderImage(FAppStyle::Get().GetBrush("DetailsView.GridLine"))
+						.Visibility(this, &SDetailsIndent::GetExpanderVisibility)
+						.Padding(FMargin(0, 0, 1, 0))
+						[
+							SNew(SBorder)
+							.BorderImage(FAppStyle::Get().GetBrush("DetailsView.CategoryMiddle"))
+							.BorderBackgroundColor(Row.ToSharedRef(), &SDetailSingleItemRow::GetInnerBackgroundColor)
+							.Padding(0)
+							[
+								SAssignNew(ExpanderArrow, SButton)
+								.ButtonStyle(FCoreStyle::Get(), "NoBorder")
+								.VAlign(VAlign_Center)
+								.HAlign(HAlign_Center)
+								.Visibility(this, &SDetailsIndent::GetExpanderVisibility)
+								.ClickMethod(EButtonClickMethod::MouseDown)
+								.OnClicked(this, &SDetailsIndent::OnExpanderClicked)
+								.ContentPadding(0)
+								.IsFocusable(false)
+								[
+									SNew(SImage)
+									.Image(this, &SDetailsIndent::GetExpanderImage)
+									.ColorAndOpacity(FSlateColor::UseForeground())
+								]
+							]
+						]
+					]
+				];
+			}
+			else
+			{
+				NewSlot
+				[
+					SNew(SBox)
+					.WidthOverride(16)
+					[
+						SNew(SBorder)
+						.BorderImage(FAppStyle::Get().GetBrush("DetailsView.CategoryMiddle"))
+						.BorderBackgroundColor_Static(&SDetailSingleItemRow_Helper::GetBackgroundColor, i)
+						.Padding(0)
+						[
+							SNew(SImage)
+							.Image(FAppStyle::Get().GetBrush("DetailsView.ArrayDropShadow"))
+						]
+					]
+				];
+			}
+		}
+	}
+
+	EVisibility GetExpanderVisibility() const
+	{
+		return Row->DoesItemHaveChildren() ? EVisibility::Visible : EVisibility::Collapsed;
+	}
+
+	const FSlateBrush* GetExpanderImage() const
+	{
+		const bool bIsItemExpanded = Row->IsItemExpanded();
+
+		FName ResourceName;
+		if (bIsItemExpanded)
+		{
+			if (ExpanderArrow->IsHovered())
+			{
+				static FName ExpandedHoveredName = "TreeArrow_Expanded_Hovered";
+				ResourceName = ExpandedHoveredName;
+			}
+			else
+			{
+				static FName ExpandedName = "TreeArrow_Expanded";
+				ResourceName = ExpandedName;
+			}
+		}
+		else
+		{
+			if (ExpanderArrow->IsHovered())
+			{
+				static FName CollapsedHoveredName = "TreeArrow_Collapsed_Hovered";
+				ResourceName = CollapsedHoveredName;
+			}
+			else
+			{
+				static FName CollapsedName = "TreeArrow_Collapsed";
+				ResourceName = CollapsedName;
+			}
+		}
+
+		return FAppStyle::Get().GetBrush(ResourceName);
+	}
+
+	FReply OnExpanderClicked()
+	{
+		// Recurse the expansion if "shift" is being pressed
+		const FModifierKeysState ModKeyState = FSlateApplication::Get().GetModifierKeys();
+		if(ModKeyState.IsShiftDown())
+		{
+			Row->Private_OnExpanderArrowShiftClicked();
+		}
+		else
+		{
+			Row->ToggleExpansion();
+		}
+
+		return FReply::Handled();
+	}
+
+	/*int32 OnPaint(const FPaintArgs& Args, const FGeometry& AllottedGeometry, const FSlateRect& MyCullingRect, FSlateWindowElementList& OutDrawElements, int32 LayerId, const FWidgetStyle& InWidgetStyle, bool bParentEnabled) const override
+	{
+		int32 IndentLevel = Row->GetIndentLevel();
+		for (int32 i = 0; i < IndentLevel; ++i)
+		{
+			FVector2D Offset(i * 16, 0);
+
+			FSlateDrawElement::MakeBox( 
+				OutDrawElements,
+				LayerId,
+				AllottedGeometry.ToPaintGeometry(FSlateLayoutTransform(Offset)),
+				GetIndentBrush(i),
+				ESlateDrawEffect::None,
+				FLinearColor(1,1,1,0.5f)
+			);
+		}
+
+		return LayerId;
+	}
+	*/
+};
+
 void SConstrainedBox::Construct(const FArguments& InArgs)
 {
 	MinWidth = InArgs._MinWidth;
@@ -52,34 +263,6 @@ FVector2D SConstrainedBox::ComputeDesiredSize(float LayoutScaleMultiplier) const
 	}
 }
 
-namespace DetailWidgetConstants
-{
-	const FMargin LeftRowPadding( 0.0f, 2.5f, 12.0f, 2.5f );
-	const FMargin RightRowPadding( 12.0f, 2.5f, 2.0f, 2.5f );
-}
-
-namespace SDetailSingleItemRow_Helper
-{
-	//Get the node item number in case it is expand we have to recursively count all expanded children
-	void RecursivelyGetItemShow(TSharedRef<FDetailTreeNode> ParentItem, int32 &ItemShowNum)
-	{
-		if (ParentItem->GetVisibility() == ENodeVisibility::Visible)
-		{
-			ItemShowNum++;
-		}
-
-		if (ParentItem->ShouldBeExpanded())
-		{
-			TArray< TSharedRef<FDetailTreeNode> > Childrens;
-			ParentItem->GetChildren(Childrens);
-			for (TSharedRef<FDetailTreeNode> ItemChild : Childrens)
-			{
-				RecursivelyGetItemShow(ItemChild, ItemShowNum);
-			}
-		}
-	}
-}
-
 void SDetailSingleItemRow::OnFavoriteMenuToggle()
 {
 	if (Customization->GetPropertyNode().IsValid() && Customization->GetPropertyNode()->CanDisplayFavorite())
@@ -398,18 +581,34 @@ void SDetailSingleItemRow::Construct( const FArguments& InArgs, FDetailLayoutCus
 				.Value(ColumnSizeData.PropertyColumnWidth)
 				.OnSlotResized(ColumnSizeData.OnPropertyColumnResized)
 			[
-				SAssignNew(InnerSplitter, SSplitter)
-				.Style(FEditorStyle::Get(), "DetailsView.Splitter")
-				.PhysicalSplitterHandleSize(1.0f)
-				.HitDetectionSplitterHandleSize(5.0f)
-				.HighlightedHandleIndex(ColumnSizeData.HoveredSplitterIndex)
-				.OnHandleHovered(ColumnSizeData.OnSplitterHandleHovered)
+				SNew(SBorder)
+				.BorderImage(FAppStyle::Get().GetBrush("DetailsView.CategoryMiddle"))
+				.BorderBackgroundColor(this, &SDetailSingleItemRow::GetInnerBackgroundColor)
+				.Padding(0)
+				[
+					SAssignNew(InnerSplitter, SSplitter)
+					.Style(FEditorStyle::Get(), "DetailsView.Splitter")
+					.PhysicalSplitterHandleSize(1.0f)
+					.HitDetectionSplitterHandleSize(5.0f)
+					.HighlightedHandleIndex(ColumnSizeData.HoveredSplitterIndex)
+					.OnHandleHovered(ColumnSizeData.OnSplitterHandleHovered)
+				]
 			];
 
 			// create Name column:
 			// | Left  | Name | Value | Right |
 			TSharedRef<SHorizontalBox> NameColumnBox = SNew(SHorizontalBox)
 				.Clipping(EWidgetClipping::OnDemand);
+
+			// indentation and expander arrow
+			NameColumnBox->AddSlot()
+				.HAlign(HAlign_Left)
+				.VAlign(VAlign_Fill)
+				.Padding(0,0,0,0)
+				.AutoWidth()
+				[
+					SNew(SDetailsIndent, SharedThis(this))
+				];
 			
 			TSharedPtr<FPropertyNode> PropertyNode = Customization->GetPropertyNode();
 			if (PropertyNode.IsValid())
@@ -417,15 +616,16 @@ void SDetailSingleItemRow::Construct( const FArguments& InArgs, FDetailLayoutCus
 				if (PropertyNode->IsReorderable())
 				{
 					TSharedPtr<SDetailSingleItemRow> InRow = SharedThis(this);
-					TSharedRef<SWidget> Handle = PropertyEditorHelpers::MakePropertyReorderHandle(PropertyNode.ToSharedRef(), InRow);
-					Handle->SetEnabled(IsPropertyEditingEnabled);
+					TSharedRef<SWidget> ArrayHandle = PropertyEditorHelpers::MakePropertyReorderHandle(PropertyNode.ToSharedRef(), InRow);
+					ArrayHandle->SetEnabled(IsPropertyEditingEnabled);
 
 					NameColumnBox->AddSlot()
-						.Padding(0.0f, 0.0f, 10.0f, 0.0f)
-						.HAlign(HAlign_Right)
+						.Padding(5,0,0,0)
+						.HAlign(HAlign_Left)
 						.VAlign(VAlign_Center)
+						.AutoWidth()
 						[
-							Handle
+							ArrayHandle
 						];
 
 					ArrayDragDelegate = FOnTableRowDragEnter::CreateSP(this, &SDetailSingleItemRow::OnArrayDragEnter);
@@ -444,17 +644,6 @@ void SDetailSingleItemRow::Construct( const FArguments& InArgs, FDetailLayoutCus
 				}
 			}
 
-			// expander arrow
-			NameColumnBox->AddSlot()
-				.HAlign(HAlign_Right)
-				.VAlign(VAlign_Center)
-				.Padding(0,0,3,0)
-				.AutoWidth()
-				[
-					SNew(SExpanderArrow, SharedThis(this))
-					.BaseIndentLevel(1)
-				];
-
 			if (bHasMultipleColumns)
 			{
 				// If the NameWidget has already been disabled, don't re-enable it if IsPropertyEditingEnabled is true.
@@ -501,7 +690,7 @@ void SDetailSingleItemRow::Construct( const FArguments& InArgs, FDetailLayoutCus
 						+ SHorizontalBox::Slot()
 						.HAlign(HAlign_Right)
 						.VAlign(VAlign_Center)
-						.Padding(3,0,3,0)
+						.Padding(5,0,0,0)
 						.AutoWidth()
 						[
 							ExtensionWidget.ToSharedRef()
@@ -613,12 +802,13 @@ void SDetailSingleItemRow::Construct( const FArguments& InArgs, FDetailLayoutCus
 	this->ChildSlot
 	[
 		SNew( SBorder )
-		.BorderImage( FEditorStyle::GetBrush( "DetailsView.GridLine" ) )
-		.Padding( FMargin( 0, 0, 0, 1 ) )
+		.BorderImage(FAppStyle::Get().GetBrush( "DetailsView.GridLine"))
+		.Padding( FMargin(0, 0, 0, 1) )
 		[
 			SNew( SBorder )
-			.BorderImage( this, &SDetailSingleItemRow::GetBorderImage )
-			.Padding( FMargin( 0.0f, 0.0f, SDetailTableRowBase::ScrollbarPaddingSize, 0.0f ) )
+			.BorderImage(FAppStyle::Get().GetBrush("DetailsView.CategoryMiddle"))
+			.BorderBackgroundColor(this, &SDetailSingleItemRow::GetOuterBackgroundColor)
+			.Padding(FMargin( 0.0f, 0.0f, SDetailTableRowBase::ScrollbarPaddingSize, 0.0f ))
 			[
 				Widget
 			]
@@ -637,6 +827,47 @@ void SDetailSingleItemRow::Construct( const FArguments& InArgs, FDetailLayoutCus
 	);
 }
 
+/** Get the background color of the outer part of the row, which contains the edit condition and extension widgets. */
+FSlateColor SDetailSingleItemRow::GetOuterBackgroundColor() const
+{
+	if (IsHighlighted())
+	{
+		return FAppStyle::Get().GetSlateColor("Colors.Background");
+	}
+	else if (bIsDragDropObject)
+	{
+		return FAppStyle::Get().GetSlateColor("Colors.Background");
+	}
+	else if (IsHovered() && !bIsHoveredDragTarget)
+	{
+		return FAppStyle::Get().GetSlateColor("Colors.Header");
+	}
+	else if (bIsHoveredDragTarget)
+	{
+		return FAppStyle::Get().GetSlateColor("Colors.Background");
+	}
+	
+	return FAppStyle::Get().GetSlateColor("Colors.Background");
+}
+
+/** Get the background color of the inner part of the row, which contains the name and value widgets. */
+FSlateColor SDetailSingleItemRow::GetInnerBackgroundColor() const
+{
+	if (IsHovered() && !bIsHoveredDragTarget)
+	{
+		return FAppStyle::Get().GetSlateColor("Colors.Header");
+	}
+
+	int32 IndentLevel = 0; 
+	if (OwnerTablePtr.IsValid())
+	{
+		IndentLevel = GetIndentLevel();
+	}
+	IndentLevel = FMath::Max(0, IndentLevel - 1);
+
+	return SDetailSingleItemRow_Helper::GetBackgroundColor(IndentLevel);
+}
+
 bool SDetailSingleItemRow::OnContextMenuOpening(FMenuBuilder& MenuBuilder)
 {
 	const bool bIsCopyPasteBound = Customization->GetWidgetRow().IsCopyPasteBound();
@@ -862,30 +1093,6 @@ bool SDetailSingleItemRow::CanPasteProperty() const
 	return !ClipboardContent.IsEmpty();
 }
 
-const FSlateBrush* SDetailSingleItemRow::GetBorderImage() const
-{
-	if( IsHighlighted() )
-	{
-		return FEditorStyle::GetBrush("DetailsView.CategoryMiddle_Highlighted");
-	}
-	else if (bIsDragDropObject)
-	{
-		return FEditorStyle::GetBrush("DetailsView.CategoryMiddle_Active");
-	}
-	else if (IsHovered() && !bIsHoveredDragTarget)
-	{
-		return FEditorStyle::GetBrush("DetailsView.CategoryMiddle_Hovered");
-	}
-	else if (bIsHoveredDragTarget)
-	{
-		return FEditorStyle::GetBrush("DetailsView.CategoryMiddle_Highlighted");
-	}
-	else
-	{
-		return FEditorStyle::GetBrush("DetailsView.CategoryMiddle");
-	}
-}
-
 TSharedRef<SWidget> SDetailSingleItemRow::CreateExtensionWidget(TSharedRef<SWidget> ValueWidget, FDetailLayoutCustomization& InCustomization, TSharedRef<FDetailTreeNode> InTreeNode)
 {
 	TSharedPtr<SWidget> ExtensionWidget = SNullWidget::NullWidget;
diff --git a/Engine/Source/Editor/PropertyEditor/Private/SDetailSingleItemRow.h b/Engine/Source/Editor/PropertyEditor/Private/SDetailSingleItemRow.h
--- a/Engine/Source/Editor/PropertyEditor/Private/SDetailSingleItemRow.h
+++ b/Engine/Source/Editor/PropertyEditor/Private/SDetailSingleItemRow.h
@@ -88,7 +88,8 @@ private:
 	void OnCopyProperty();
 	void OnPasteProperty();
 	bool CanPasteProperty() const;
-	const FSlateBrush* GetBorderImage() const;
+	FSlateColor GetOuterBackgroundColor() const;
+	FSlateColor GetInnerBackgroundColor() const;
 	TSharedRef<SWidget> CreateExtensionWidget( TSharedRef<SWidget> ValueWidget, FDetailLayoutCustomization& InCustomization, TSharedRef<FDetailTreeNode> InTreeNode );
 	TSharedRef<SWidget> CreateKeyframeButton( FDetailLayoutCustomization& InCustomization, TSharedRef<FDetailTreeNode> InTreeNode );
 	bool IsKeyframeButtonEnabled(TSharedRef<FDetailTreeNode> InTreeNode) const;
@@ -121,6 +122,8 @@ private:
 	bool bIsHoveredDragTarget;
 	bool bIsDragDropObject;
 	TSharedPtr<FPropertyNode> SwappablePropertyNode;
+
+	friend class SDetailsIndent;
 };
 
 class FArrayRowDragDropOp : public FDecoratedDragDropOp
diff --git a/Engine/Source/Editor/PropertyEditor/Public/DetailWidgetRow.h b/Engine/Source/Editor/PropertyEditor/Public/DetailWidgetRow.h
--- a/Engine/Source/Editor/PropertyEditor/Public/DetailWidgetRow.h
+++ b/Engine/Source/Editor/PropertyEditor/Public/DetailWidgetRow.h
@@ -13,7 +13,6 @@
 #include "Widgets/Layout/SSpacer.h"
 
 class FDetailWidgetRow;
-struct FCustomEditCondition;
 class FResetToDefaultOverride;
 
 /** Widget declaration for custom widgets in a widget row */
