Fixed subcategory headers not rendering with the correct background color.

Moved GetBackgroundColor from SDetailSingleItemRow_Helpers to PropertyEditorConstants.

Fixed bug with valid property handles not displaying a reset to default by default.

diff --git a/Engine/Source/Editor/PropertyEditor/Private/DetailCategoryGroupNode.cpp b/Engine/Source/Editor/PropertyEditor/Private/DetailCategoryGroupNode.cpp
--- a/Engine/Source/Editor/PropertyEditor/Private/DetailCategoryGroupNode.cpp
+++ b/Engine/Source/Editor/PropertyEditor/Private/DetailCategoryGroupNode.cpp
@@ -57,6 +57,7 @@ void SDetailCategoryTableRow::Construct( const FArguments& InArgs, TSharedRef<FD
 	[	
 		SNew( SBorder )
 		.BorderImage( this, &SDetailCategoryTableRow::GetBackgroundImage )
+		.BorderBackgroundColor( this, &SDetailCategoryTableRow::GetBackgroundColor )
 		.Padding( FMargin( 0.0f, BorderVerticalPadding, SDetailTableRowBase::ScrollbarPaddingSize, BorderVerticalPadding ) )
 		[
 			Widget.ToSharedRef()
@@ -80,22 +81,17 @@ const FSlateBrush* SDetailCategoryTableRow::GetBackgroundImage() const
 {
 	if (bShowBorder)
 	{
-		if (IsHovered())
+		if (bIsInnerCategory)
 		{
-			if (bIsInnerCategory)
-			{
-				return FEditorStyle::GetBrush("DetailsView.CategoryMiddle_Hovered");
-			}
+			return FEditorStyle::GetBrush("DetailsView.CategoryMiddle");
+		}
 
+		if (IsHovered())
+		{
 			return IsItemExpanded() ? FEditorStyle::GetBrush("DetailsView.CategoryTop_Hovered") : FEditorStyle::GetBrush("DetailsView.CollapsedCategory_Hovered");
 		}
 		else
 		{
-			if (bIsInnerCategory)
-			{
-				return FEditorStyle::GetBrush("DetailsView.CategoryMiddle");
-			}
-			
 			return IsItemExpanded() ? FEditorStyle::GetBrush("DetailsView.CategoryTop") : FEditorStyle::GetBrush("DetailsView.CollapsedCategory");
 		}
 	}
@@ -103,6 +99,24 @@ const FSlateBrush* SDetailCategoryTableRow::GetBackgroundImage() const
 	return nullptr;
 }
 
+FSlateColor SDetailCategoryTableRow::GetBackgroundColor() const
+{
+	if (bShowBorder && bIsInnerCategory)
+	{
+		int32 IndentLevel = -1;
+		if (OwnerTablePtr.IsValid())
+		{
+			IndentLevel = GetIndentLevel();
+		}
+
+		IndentLevel = FMath::Max(IndentLevel - 1, 0);
+
+		return PropertyEditorConstants::GetRowBackgroundColor(IndentLevel);
+	}
+
+	return FSlateColor(FLinearColor::White);
+}
+
 FReply SDetailCategoryTableRow::OnMouseButtonDown( const FGeometry& MyGeometry, const FPointerEvent& MouseEvent )
 {
 	if (MouseEvent.GetEffectingButton() == EKeys::LeftMouseButton)
diff --git a/Engine/Source/Editor/PropertyEditor/Private/DetailCategoryGroupNode.h b/Engine/Source/Editor/PropertyEditor/Private/DetailCategoryGroupNode.h
--- a/Engine/Source/Editor/PropertyEditor/Private/DetailCategoryGroupNode.h
+++ b/Engine/Source/Editor/PropertyEditor/Private/DetailCategoryGroupNode.h
@@ -32,6 +32,7 @@ public:
 private:
 	EVisibility IsSeparatorVisible() const;
 	const FSlateBrush* GetBackgroundImage() const;
+	FSlateColor GetBackgroundColor() const;
 
 	virtual FReply OnMouseButtonDown( const FGeometry& MyGeometry, const FPointerEvent& MouseEvent ) override;
 	virtual FReply OnMouseButtonDoubleClick( const FGeometry& InMyGeometry, const FPointerEvent& InMouseEvent ) override;
diff --git a/Engine/Source/Editor/PropertyEditor/Private/SDetailSingleItemRow.cpp b/Engine/Source/Editor/PropertyEditor/Private/SDetailSingleItemRow.cpp
--- a/Engine/Source/Editor/PropertyEditor/Private/SDetailSingleItemRow.cpp
+++ b/Engine/Source/Editor/PropertyEditor/Private/SDetailSingleItemRow.cpp
@@ -17,7 +17,7 @@
 #include "Modules/ModuleInterface.h"
 #include "Modules/ModuleManager.h"
 #include "SResetToDefaultPropertyEditor.h"
-
+#include "PropertyEditorConstants.h"
 
 namespace DetailWidgetConstants
 {
@@ -45,13 +45,6 @@ namespace SDetailSingleItemRow_Helper
 			}
 		}
 	}
-
-	FSlateColor GetBackgroundColor(int32 IndentLevel)
-	{
-		FLinearColor BackgroundColor = FAppStyle::Get().GetSlateColor("Colors.Background").GetSpecifiedColor();
-
-		return FSlateColor(BackgroundColor + FLinearColor(2.0f/255, 2.0f/255, 2.0f/255) * (IndentLevel % 4));
-	}
 }
 
 class SDetailsIndent : public SCompoundWidget
@@ -139,7 +132,7 @@ public:
 					[
 						SNew(SBorder)
 						.BorderImage(FAppStyle::Get().GetBrush("DetailsView.CategoryMiddle"))
-						.BorderBackgroundColor_Static(&SDetailSingleItemRow_Helper::GetBackgroundColor, i)
+						.BorderBackgroundColor_Static(&PropertyEditorConstants::GetRowBackgroundColor, i)
 						.Padding(0)
 						[
 							SNew(SImage)
@@ -206,27 +199,6 @@ public:
 
 		return FReply::Handled();
 	}
-
-	/*int32 OnPaint(const FPaintArgs& Args, const FGeometry& AllottedGeometry, const FSlateRect& MyCullingRect, FSlateWindowElementList& OutDrawElements, int32 LayerId, const FWidgetStyle& InWidgetStyle, bool bParentEnabled) const override
-	{
-		int32 IndentLevel = Row->GetIndentLevel();
-		for (int32 i = 0; i < IndentLevel; ++i)
-		{
-			FVector2D Offset(i * 16, 0);
-
-			FSlateDrawElement::MakeBox( 
-				OutDrawElements,
-				LayerId,
-				AllottedGeometry.ToPaintGeometry(FSlateLayoutTransform(Offset)),
-				GetIndentBrush(i),
-				ESlateDrawEffect::None,
-				FLinearColor(1,1,1,0.5f)
-			);
-		}
-
-		return LayerId;
-	}
-	*/
 };
 
 void SConstrainedBox::Construct(const FArguments& InArgs)
@@ -721,12 +693,17 @@ void SDetailSingleItemRow::Construct( const FArguments& InArgs, FDetailLayoutCus
 			bool bCreateResetToDefault = Row.CustomResetToDefault.IsSet();
 
 			TSharedPtr<IPropertyHandle> PropertyHandle = GetPropertyHandle();
-			if (PropertyHandle.IsValid() )
+
+			if (PropertyHandle.IsValid())
 			{
 				if (PropertyHandle->HasMetaData("NoResetToDefault") || PropertyHandle->GetInstanceMetaData("NoResetToDefault"))
 				{
 					bCreateResetToDefault = false;
 				}
+				else
+				{
+					bCreateResetToDefault = true;
+				}
 			} 
 
 			TSharedRef<SWidget> ResetWidget = SNullWidget::NullWidget;
@@ -865,7 +842,7 @@ FSlateColor SDetailSingleItemRow::GetInnerBackgroundColor() const
 	}
 	IndentLevel = FMath::Max(0, IndentLevel - 1);
 
-	return SDetailSingleItemRow_Helper::GetBackgroundColor(IndentLevel);
+	return PropertyEditorConstants::GetRowBackgroundColor(IndentLevel);
 }
 
 bool SDetailSingleItemRow::OnContextMenuOpening(FMenuBuilder& MenuBuilder)
diff --git a/Engine/Source/Editor/PropertyEditor/Private/UserInterface/PropertyEditor/PropertyEditorConstants.cpp b/Engine/Source/Editor/PropertyEditor/Private/UserInterface/PropertyEditor/PropertyEditorConstants.cpp
--- a/Engine/Source/Editor/PropertyEditor/Private/UserInterface/PropertyEditor/PropertyEditorConstants.cpp
+++ b/Engine/Source/Editor/PropertyEditor/Private/UserInterface/PropertyEditor/PropertyEditorConstants.cpp
@@ -14,3 +14,9 @@ const FSlateBrush* PropertyEditorConstants::GetOverlayBrush( const TSharedRef< c
 {
 	return FEditorStyle::GetBrush( TEXT("PropertyWindow.NoOverlayColor") );
 }
+
+FSlateColor PropertyEditorConstants::GetRowBackgroundColor(int32 IndentLevel) 
+{
+	FLinearColor BackgroundColor = FAppStyle::Get().GetSlateColor("Colors.Background").GetSpecifiedColor();
+	return FSlateColor(BackgroundColor + FLinearColor(2.0f/255, 2.0f/255, 2.0f/255) * (IndentLevel % 4));
+}
diff --git a/Engine/Source/Editor/PropertyEditor/Private/UserInterface/PropertyEditor/PropertyEditorConstants.h b/Engine/Source/Editor/PropertyEditor/Private/UserInterface/PropertyEditor/PropertyEditorConstants.h
--- a/Engine/Source/Editor/PropertyEditor/Private/UserInterface/PropertyEditor/PropertyEditorConstants.h
+++ b/Engine/Source/Editor/PropertyEditor/Private/UserInterface/PropertyEditor/PropertyEditorConstants.h
@@ -17,4 +17,6 @@ public:
 	static const FName MD_UseEnumValuesAsMaskValuesInEditor;
 
 	static const FSlateBrush* GetOverlayBrush( const TSharedRef< class FPropertyEditor > PropertyEditor );
+
+	static FSlateColor GetRowBackgroundColor(int32 IndentLevel);
 };
