Changed colours of nested array entries in the details view to be perceptually linear.

Minor tweaking of the expander arrow positioning, and moved the widget out of SDetailsIndent and into SDetailSingleItemRow.

diff --git a/Engine/Source/Editor/PropertyEditor/Private/DetailCategoryGroupNode.cpp b/Engine/Source/Editor/PropertyEditor/Private/DetailCategoryGroupNode.cpp
--- a/Engine/Source/Editor/PropertyEditor/Private/DetailCategoryGroupNode.cpp
+++ b/Engine/Source/Editor/PropertyEditor/Private/DetailCategoryGroupNode.cpp
@@ -12,25 +12,13 @@ void SDetailCategoryTableRow::Construct( const FArguments& InArgs, TSharedRef<FD
 	bShowBorder = InArgs._ShowBorder;
 
 	TSharedPtr<SWidget> Widget = SNullWidget::NullWidget;
-
-	float MyContentTopPadding = 2.0f;
-	float MyContentBottomPadding = 2.0f;
-
-	float ChildSlotPadding = 2.0f;
-	float BorderVerticalPadding = 3.0f;
-
 	FDetailColumnSizeData& ColumnSizeData = InOwnerTreeNode->GetDetailsView()->GetColumnSizeData();
 
-	MyContentTopPadding += ChildSlotPadding + 2 * BorderVerticalPadding;
-	MyContentBottomPadding += 2 * BorderVerticalPadding;
-
-	ChildSlotPadding = 0.0f;
-	BorderVerticalPadding = 0.0f;
 
 	Widget = SNew(SHorizontalBox)
 		+ SHorizontalBox::Slot()
 		.VAlign(VAlign_Center)
-		.Padding(12.0f, MyContentTopPadding, 12.0f, MyContentBottomPadding)
+		.Padding(8, 0, 8, 0)
 		.AutoWidth()
 		[
 			SNew(SExpanderArrow, SharedThis(this))
@@ -51,14 +39,16 @@ void SDetailCategoryTableRow::Construct( const FArguments& InArgs, TSharedRef<FD
 		[
 			InArgs._HeaderContent.IsValid() ? InArgs._HeaderContent.ToSharedRef() : SNullWidget::NullWidget
 		];
+
+	const float VerticalPadding = bIsInnerCategory ? 8 : 10;
 	
 	ChildSlot
-	.Padding( 0.0f, bIsInnerCategory ? 0.0f : ChildSlotPadding, 0.0f, 0.0f )
+	.Padding(0)
 	[	
 		SNew( SBorder )
 		.BorderImage( this, &SDetailCategoryTableRow::GetBackgroundImage )
 		.BorderBackgroundColor( this, &SDetailCategoryTableRow::GetBackgroundColor )
-		.Padding( FMargin( 0.0f, BorderVerticalPadding, SDetailTableRowBase::ScrollbarPaddingSize, BorderVerticalPadding ) )
+		.Padding( FMargin(0, VerticalPadding, SDetailTableRowBase::ScrollbarPaddingSize, VerticalPadding) )
 		[
 			Widget.ToSharedRef()
 		]
diff --git a/Engine/Source/Editor/PropertyEditor/Private/SDetailSingleItemRow.cpp b/Engine/Source/Editor/PropertyEditor/Private/SDetailSingleItemRow.cpp
--- a/Engine/Source/Editor/PropertyEditor/Private/SDetailSingleItemRow.cpp
+++ b/Engine/Source/Editor/PropertyEditor/Private/SDetailSingleItemRow.cpp
@@ -48,16 +48,16 @@ namespace SDetailSingleItemRow_Helper
 }
 
 class SDetailsIndent : public SCompoundWidget
-{
+{ 
+private:
+	TSharedPtr<SHorizontalBox> ChildBox;
+	TSharedPtr<SDetailSingleItemRow> Row;
+
 public:
 
 	SLATE_BEGIN_ARGS(SDetailsIndent) {}
 	SLATE_END_ARGS()
 
-	TSharedPtr<SHorizontalBox> ChildBox;
-	TSharedPtr<SButton> ExpanderArrow;
-	TSharedPtr<SDetailSingleItemRow> Row;
-
 	void Construct(const FArguments& InArgs, TSharedRef<SDetailSingleItemRow> DetailsRow)
 	{
 		Row = DetailsRow;
@@ -70,8 +70,8 @@ public:
 
 	void Tick(const FGeometry& AllottedGeometry, const double InCurrentTime, const float InDeltaTime) override
 	{
-		int32 IndentLevel = Row->GetIndentLevel();
-		if (ChildBox->NumSlots() == IndentLevel)
+		int32 IndentLevel = Row->GetIndentLevel() - 1;
+		if (ChildBox->NumSlots() >= IndentLevel)
 		{
 			// already initialized, don't need to do anything
 			return;
@@ -81,58 +81,17 @@ public:
 
 		for (int32 i = 0; i < IndentLevel; ++i)
 		{
-			SHorizontalBox::FSlot& NewSlot = ChildBox->AddSlot()
+			ChildBox->AddSlot()
 				.VAlign(VAlign_Fill)
 				.HAlign(HAlign_Left)
-				.Padding(0);
-
-			if (i == IndentLevel - 1)
-			{
-				NewSlot
-				.AutoWidth()
-				[
-					SNew(SBox)
-					.WidthOverride(32)
-					[
-						SNew(SBorder)
-						.BorderImage(FAppStyle::Get().GetBrush("DetailsView.GridLine"))
-						.Visibility(this, &SDetailsIndent::GetExpanderVisibility)
-						.Padding(FMargin(0, 0, 1, 0))
-						[
-							SNew(SBorder)
-							.BorderImage(FAppStyle::Get().GetBrush("DetailsView.CategoryMiddle"))
-							.BorderBackgroundColor(Row.ToSharedRef(), &SDetailSingleItemRow::GetInnerBackgroundColor)
-							.Padding(0)
-							[
-								SAssignNew(ExpanderArrow, SButton)
-								.ButtonStyle(FCoreStyle::Get(), "NoBorder")
-								.VAlign(VAlign_Center)
-								.HAlign(HAlign_Center)
-								.Visibility(this, &SDetailsIndent::GetExpanderVisibility)
-								.ClickMethod(EButtonClickMethod::MouseDown)
-								.OnClicked(this, &SDetailsIndent::OnExpanderClicked)
-								.ContentPadding(0)
-								.IsFocusable(false)
-								[
-									SNew(SImage)
-									.Image(this, &SDetailsIndent::GetExpanderImage)
-									.ColorAndOpacity(FSlateColor::UseForeground())
-								]
-							]
-						]
-					]
-				];
-			}
-			else
-			{
-				NewSlot
+				.Padding(0)
 				[
 					SNew(SBox)
 					.WidthOverride(16)
 					[
 						SNew(SBorder)
 						.BorderImage(FAppStyle::Get().GetBrush("DetailsView.CategoryMiddle"))
-						.BorderBackgroundColor_Static(&PropertyEditorConstants::GetRowBackgroundColor, i)
+						.BorderBackgroundColor(this, &SDetailsIndent::GetRowBackgroundColor, i)
 						.Padding(0)
 						[
 							SNew(SImage)
@@ -140,64 +99,17 @@ public:
 						]
 					]
 				];
-			}
-		}
-	}
-
-	EVisibility GetExpanderVisibility() const
-	{
-		return Row->DoesItemHaveChildren() ? EVisibility::Visible : EVisibility::Collapsed;
-	}
-
-	const FSlateBrush* GetExpanderImage() const
-	{
-		const bool bIsItemExpanded = Row->IsItemExpanded();
-
-		FName ResourceName;
-		if (bIsItemExpanded)
-		{
-			if (ExpanderArrow->IsHovered())
-			{
-				static FName ExpandedHoveredName = "TreeArrow_Expanded_Hovered";
-				ResourceName = ExpandedHoveredName;
-			}
-			else
-			{
-				static FName ExpandedName = "TreeArrow_Expanded";
-				ResourceName = ExpandedName;
-			}
 		}
-		else
-		{
-			if (ExpanderArrow->IsHovered())
-			{
-				static FName CollapsedHoveredName = "TreeArrow_Collapsed_Hovered";
-				ResourceName = CollapsedHoveredName;
-			}
-			else
-			{
-				static FName CollapsedName = "TreeArrow_Collapsed";
-				ResourceName = CollapsedName;
-			}
-		}
-
-		return FAppStyle::Get().GetBrush(ResourceName);
 	}
 
-	FReply OnExpanderClicked()
+	FSlateColor GetRowBackgroundColor(int32 i) const
 	{
-		// Recurse the expansion if "shift" is being pressed
-		const FModifierKeysState ModKeyState = FSlateApplication::Get().GetModifierKeys();
-		if(ModKeyState.IsShiftDown())
-		{
-			Row->Private_OnExpanderArrowShiftClicked();
-		}
-		else
+		if (Row->IsHovered() && !Row->bIsHoveredDragTarget)
 		{
-			Row->ToggleExpansion();
+			return FAppStyle::Get().GetSlateColor("Colors.Header");
 		}
 
-		return FReply::Handled();
+		return PropertyEditorConstants::GetRowBackgroundColor(i);
 	}
 };
 
@@ -526,8 +438,8 @@ void SDetailSingleItemRow::Construct( const FArguments& InArgs, FDetailLayoutCus
 
 			// edit condition widget
 			LeftColumnBox->AddSlot()
-				.Padding(0.0f, 0.0f, 3.0f, 0.0f)
-				.HAlign(HAlign_Left)
+				.Padding(6.0f, 0.0f, 6.0f, 0.0f)
+				.HAlign(HAlign_Center)
 				.VAlign(VAlign_Center)
 				.AutoWidth()
 				[
@@ -581,6 +493,32 @@ void SDetailSingleItemRow::Construct( const FArguments& InArgs, FDetailLayoutCus
 				[
 					SNew(SDetailsIndent, SharedThis(this))
 				];
+
+			NameColumnBox->AddSlot()
+				.HAlign(HAlign_Left)
+				.VAlign(VAlign_Center)
+				.Padding(0)
+				.AutoWidth()
+				[
+					SNew(SConstrainedBox)
+					.MinWidth(20)
+					.Visibility(this, &SDetailSingleItemRow::GetExpanderVisibility)
+					[
+						SAssignNew(ExpanderArrow, SButton)
+						.ButtonStyle(FCoreStyle::Get(), "NoBorder")
+						.VAlign(VAlign_Center)
+						.HAlign(HAlign_Center)
+						.ClickMethod(EButtonClickMethod::MouseDown)
+						.OnClicked(this, &SDetailSingleItemRow::OnExpanderClicked)
+						.ContentPadding(FMargin(5,0,0,0))
+						.IsFocusable(false)
+						[
+							SNew(SImage)
+							.Image(this, &SDetailSingleItemRow::GetExpanderImage)
+							.ColorAndOpacity(FSlateColor::UseForeground())
+						]
+					]
+				];
 			
 			TSharedPtr<FPropertyNode> PropertyNode = Customization->GetPropertyNode();
 			if (PropertyNode.IsValid())
@@ -947,6 +885,62 @@ bool SDetailSingleItemRow::OnContextMenuOpening(FMenuBuilder& MenuBuilder)
 	return true;
 }
 
+EVisibility SDetailSingleItemRow::GetExpanderVisibility() const
+{
+	return DoesItemHaveChildren() ? EVisibility::Visible : EVisibility::Collapsed;
+}
+
+const FSlateBrush* SDetailSingleItemRow::GetExpanderImage() const
+{
+	const bool bIsItemExpanded = IsItemExpanded();
+
+	FName ResourceName;
+	if (bIsItemExpanded)
+	{
+		if (ExpanderArrow->IsHovered())
+		{
+			static FName ExpandedHoveredName = "TreeArrow_Expanded_Hovered";
+			ResourceName = ExpandedHoveredName;
+		}
+		else
+		{
+			static FName ExpandedName = "TreeArrow_Expanded";
+			ResourceName = ExpandedName;
+		}
+	}
+	else
+	{
+		if (ExpanderArrow->IsHovered())
+		{
+			static FName CollapsedHoveredName = "TreeArrow_Collapsed_Hovered";
+			ResourceName = CollapsedHoveredName;
+		}
+		else
+		{
+			static FName CollapsedName = "TreeArrow_Collapsed";
+			ResourceName = CollapsedName;
+		}
+	}
+
+	return FAppStyle::Get().GetBrush(ResourceName);
+}
+
+FReply SDetailSingleItemRow::OnExpanderClicked()
+{
+	// Recurse the expansion if "shift" is being pressed
+	const FModifierKeysState ModKeyState = FSlateApplication::Get().GetModifierKeys();
+	if(ModKeyState.IsShiftDown())
+	{
+		Private_OnExpanderArrowShiftClicked();
+	}
+	else
+	{
+		ToggleExpansion();
+	}
+
+	return FReply::Handled();
+}
+
 void SDetailSingleItemRow::OnCopyProperty()
 {
 	if (OwnerTreeNode.IsValid())
diff --git a/Engine/Source/Editor/PropertyEditor/Private/SDetailSingleItemRow.h b/Engine/Source/Editor/PropertyEditor/Private/SDetailSingleItemRow.h
--- a/Engine/Source/Editor/PropertyEditor/Private/SDetailSingleItemRow.h
+++ b/Engine/Source/Editor/PropertyEditor/Private/SDetailSingleItemRow.h
@@ -114,6 +114,10 @@ private:
 	TSharedPtr<FPropertyNode> GetPropertyNode() const;
 	TSharedPtr<IPropertyHandle> GetPropertyHandle() const;
 
+	EVisibility GetExpanderVisibility() const;
+	const FSlateBrush* GetExpanderImage() const;
+	FReply OnExpanderClicked();
+
 private:
 	TWeakPtr<IDetailKeyframeHandler> KeyframeHandler;
 	/** Customization for this widget */
@@ -122,6 +126,7 @@ private:
 	bool bIsHoveredDragTarget;
 	bool bIsDragDropObject;
 	TSharedPtr<FPropertyNode> SwappablePropertyNode;
+	TSharedPtr<SButton> ExpanderArrow;
 
 	friend class SDetailsIndent;
 };
