SDetailsIndent now uses OnPaint to render, which fixes layout-related flickering due to the indent level not being known at Construct() time.

Project Settings (and other MultipleTopLevelObject details views if any exist) no longer display an extra indent level for items.

diff --git a/Engine/Source/Editor/PropertyEditor/Private/SDetailSingleItemRow.cpp b/Engine/Source/Editor/PropertyEditor/Private/SDetailSingleItemRow.cpp
--- a/Engine/Source/Editor/PropertyEditor/Private/SDetailSingleItemRow.cpp
+++ b/Engine/Source/Editor/PropertyEditor/Private/SDetailSingleItemRow.cpp
@@ -28,7 +28,7 @@ namespace DetailWidgetConstants
 namespace SDetailSingleItemRow_Helper
 {
 	//Get the node item number in case it is expand we have to recursively count all expanded children
-	void RecursivelyGetItemShow(TSharedRef<FDetailTreeNode> ParentItem, int32 &ItemShowNum)
+	void RecursivelyGetItemShow(TSharedRef<FDetailTreeNode> ParentItem, int32& ItemShowNum)
 	{
 		if (ParentItem->GetVisibility() == ENodeVisibility::Visible)
 		{
@@ -50,8 +50,8 @@ namespace SDetailSingleItemRow_Helper
 class SDetailsIndent : public SCompoundWidget
 { 
 private:
-	TSharedPtr<SHorizontalBox> ChildBox;
 	TSharedPtr<SDetailSingleItemRow> Row;
+	TSharedPtr<SHorizontalBox> ChildBox;
 
 public:
 
@@ -64,42 +64,45 @@ public:
 
 		ChildSlot
 		[
-			SAssignNew(ChildBox, SHorizontalBox)
+			SNew(SBox)
+			.WidthOverride(this, &SDetailsIndent::GetIndentWidth)
 		];
 	}
 
-	void Tick(const FGeometry& AllottedGeometry, const double InCurrentTime, const float InDeltaTime) override
+	int32 OnPaint(const FPaintArgs& Args, const FGeometry& AllottedGeometry, const FSlateRect& MyCullingRect, FSlateWindowElementList& OutDrawElements, int32 LayerId, const FWidgetStyle& InWidgetStyle, bool bParentEnabled) const
 	{
-		int32 IndentLevel = Row->GetIndentLevel() - 1;
-		if (ChildBox->NumSlots() >= IndentLevel)
-		{
-			// already initialized, don't need to do anything
-			return;
-		}
-
-		ChildBox->ClearChildren();
+		const FSlateBrush* BackgroundBrush = FAppStyle::Get().GetBrush("DetailsView.CategoryMiddle");
+		const FSlateBrush* DropShadowBrush = FAppStyle::Get().GetBrush("DetailsView.ArrayDropShadow");
 
+		int32 IndentLevel = Row->GetIndentLevelForBackgroundColor();
 		for (int32 i = 0; i < IndentLevel; ++i)
 		{
-			ChildBox->AddSlot()
-				.VAlign(VAlign_Fill)
-				.HAlign(HAlign_Left)
-				.Padding(0)
-				[
-					SNew(SBox)
-					.WidthOverride(16)
-					[
-						SNew(SBorder)
-						.BorderImage(FAppStyle::Get().GetBrush("DetailsView.CategoryMiddle"))
-						.BorderBackgroundColor(this, &SDetailsIndent::GetRowBackgroundColor, i)
-						.Padding(0)
-						[
-							SNew(SImage)
-							.Image(FAppStyle::Get().GetBrush("DetailsView.ArrayDropShadow"))
-						]
-					]
-				];
+			FSlateColor BackgroundColor = GetRowBackgroundColor(i);
+
+			FSlateDrawElement::MakeBox(
+				OutDrawElements,
+				LayerId,
+				AllottedGeometry.ToPaintGeometry(FVector2D(16 * i, 0), FVector2D(16, AllottedGeometry.GetLocalSize().Y)),
+				BackgroundBrush,
+				ESlateDrawEffect::None,
+				BackgroundColor.GetColor(InWidgetStyle)
+			);
+
+			FSlateDrawElement::MakeBox(
+				OutDrawElements,
+				LayerId + 1,
+				AllottedGeometry.ToPaintGeometry(FVector2D(16 * i, 0), FVector2D(16, AllottedGeometry.GetLocalSize().Y)),
+				DropShadowBrush
+			);
 		}
+		
+		return LayerId + 1;
+	}
+
+	FOptionalSize GetIndentWidth() const
+	{
+		int32 IndentLevel = Row->GetIndentLevelForBackgroundColor();
+		return IndentLevel * 16.0f;
 	}
 
 	FSlateColor GetRowBackgroundColor(int32 i) const
@@ -753,13 +756,16 @@ FSlateColor SDetailSingleItemRow::GetOuterBackgroundColor() const
 	{
 		return FAppStyle::Get().GetSlateColor("Colors.Background");
 	}
-	else if (IsHovered() && !bIsHoveredDragTarget)
+	else if (IsHovered())
 	{
-		return FAppStyle::Get().GetSlateColor("Colors.Header");
-	}
-	else if (bIsHoveredDragTarget)
-	{
-		return FAppStyle::Get().GetSlateColor("Colors.Background");
+		if (bIsHoveredDragTarget)
+		{
+			return FAppStyle::Get().GetSlateColor("Colors.Background");
+		}
+		else
+		{
+			return FAppStyle::Get().GetSlateColor("Colors.Header");
+		}
 	}
 	
 	return FAppStyle::Get().GetSlateColor("Colors.Background");
@@ -773,14 +779,28 @@ FSlateColor SDetailSingleItemRow::GetInnerBackgroundColor() const
 		return FAppStyle::Get().GetSlateColor("Colors.Header");
 	}
 
+	int32 IndentLevel = GetIndentLevelForBackgroundColor();
+	return PropertyEditorConstants::GetRowBackgroundColor(IndentLevel);
+}
+
+int32 SDetailSingleItemRow::GetIndentLevelForBackgroundColor() const
+{
 	int32 IndentLevel = 0; 
 	if (OwnerTablePtr.IsValid())
 	{
-		IndentLevel = GetIndentLevel();
+		IndentLevel = GetIndentLevel() - 1;
 	}
-	IndentLevel = FMath::Max(0, IndentLevel - 1);
 
-	return PropertyEditorConstants::GetRowBackgroundColor(IndentLevel);
+	TSharedPtr<FDetailTreeNode> DetailTreeNode = GetOwnerTreeNode();
+	if (DetailTreeNode.IsValid() && 
+		DetailTreeNode->GetDetailsView() != nullptr && 
+		DetailTreeNode->GetDetailsView()->ContainsMultipleTopLevelObjects())
+	{
+		// if the row is in a multiple top level object display (eg. Project Settings), don't display an indent for the initial level
+		--IndentLevel;
+	}
+
+	return FMath::Max(0, IndentLevel);
 }
 
 bool SDetailSingleItemRow::OnContextMenuOpening(FMenuBuilder& MenuBuilder)
diff --git a/Engine/Source/Editor/PropertyEditor/Private/SDetailSingleItemRow.h b/Engine/Source/Editor/PropertyEditor/Private/SDetailSingleItemRow.h
--- a/Engine/Source/Editor/PropertyEditor/Private/SDetailSingleItemRow.h
+++ b/Engine/Source/Editor/PropertyEditor/Private/SDetailSingleItemRow.h
@@ -118,6 +118,8 @@ private:
 	const FSlateBrush* GetExpanderImage() const;
 	FReply OnExpanderClicked();
 
+	int32 GetIndentLevelForBackgroundColor() const;
+
 private:
 	TWeakPtr<IDetailKeyframeHandler> KeyframeHandler;
 	/** Customization for this widget */
diff --git a/Engine/Source/Editor/PropertyEditor/Private/SDetailTableRowBase.h b/Engine/Source/Editor/PropertyEditor/Private/SDetailTableRowBase.h
--- a/Engine/Source/Editor/PropertyEditor/Private/SDetailTableRowBase.h
+++ b/Engine/Source/Editor/PropertyEditor/Private/SDetailTableRowBase.h
@@ -60,6 +60,8 @@ public:
 		return STableRow< TSharedPtr< FDetailTreeNode > >::OnMouseButtonUp( MyGeometry, MouseEvent );
 	}
 
+	TSharedPtr<FDetailTreeNode> GetOwnerTreeNode() const { return OwnerTreeNode.Pin(); }
+
 protected:
 	/**
 	 * Called when the user opens the context menu on this row
