Fixed reset to default widgets not respecting custom IsEnabled attributes on custom widget rows.

diff --git a/Engine/Source/Editor/PropertyEditor/Private/SDetailSingleItemRow.cpp b/Engine/Source/Editor/PropertyEditor/Private/SDetailSingleItemRow.cpp
--- a/Engine/Source/Editor/PropertyEditor/Private/SDetailSingleItemRow.cpp
+++ b/Engine/Source/Editor/PropertyEditor/Private/SDetailSingleItemRow.cpp
@@ -394,8 +394,6 @@ void SDetailSingleItemRow::Construct( const FArguments& InArgs, FDetailLayoutCus
 	{
 		if (Customization->IsValidCustomization())
 		{
-			TAttribute<bool> IsPropertyEditingEnabled = InOwnerTreeNode->IsPropertyEditingEnabled();
-
 			FDetailWidgetRow Row = InCustomization->GetWidgetRow();
 
 			TSharedPtr<SWidget> NameWidget, ValueWidget, ExtensionWidget;
@@ -406,22 +404,28 @@ void SDetailSingleItemRow::Construct( const FArguments& InArgs, FDetailLayoutCus
 				SNew(SConstrainedBox)
 				.MinWidth(Row.ValueWidget.MinWidth)
 				.MaxWidth(Row.ValueWidget.MaxWidth)
-				.IsEnabled(IsPropertyEditingEnabled)
 				[
 					Row.ValueWidget.Widget
 				];
 
 			ExtensionWidget = CreateExtensionWidget(ValueWidget.ToSharedRef(), *Customization, InOwnerTreeNode);
 
-			TAttribute<bool> NameWidgetEnabled;
+			TAttribute<bool> IsEnabledAttribute = InOwnerTreeNode->IsPropertyEditingEnabled();
 			if (Row.IsEnabledAttr.IsBound())
 			{
-				NameWidgetEnabled = Row.IsEnabledAttr;
-				NameWidget->SetEnabled(Row.IsEnabledAttr);
-				ValueWidget->SetEnabled(Row.IsEnabledAttr);
-				ExtensionWidget->SetEnabled(Row.IsEnabledAttr);
+				TAttribute<bool> RowEnabledAttr = Row.IsEnabledAttr;
+				TAttribute<bool> PropertyEnabledAttr = InOwnerTreeNode->IsPropertyEditingEnabled();
+				IsEnabledAttribute = TAttribute<bool>::Create(
+					[RowEnabledAttr, PropertyEnabledAttr]()
+					{
+						return RowEnabledAttr.Get() && PropertyEnabledAttr.Get();
+					});
 			}
 
+			NameWidget->SetEnabled(IsEnabledAttribute);
+			ValueWidget->SetEnabled(IsEnabledAttribute);
+			ExtensionWidget->SetEnabled(IsEnabledAttribute);
+
 			TSharedRef<SHorizontalBox> RowBox = SNew(SHorizontalBox);
 
 			// create outer splitter
@@ -530,7 +534,7 @@ void SDetailSingleItemRow::Construct( const FArguments& InArgs, FDetailLayoutCus
 				{
 					TSharedPtr<SDetailSingleItemRow> InRow = SharedThis(this);
 					TSharedRef<SWidget> ArrayHandle = PropertyEditorHelpers::MakePropertyReorderHandle(PropertyNode.ToSharedRef(), InRow);
-					ArrayHandle->SetEnabled(IsPropertyEditingEnabled);
+					ArrayHandle->SetEnabled(IsEnabledAttribute);
 
 					NameColumnBox->AddSlot()
 						.Padding(5,0,0,0)
@@ -559,15 +563,6 @@ void SDetailSingleItemRow::Construct( const FArguments& InArgs, FDetailLayoutCus
 
 			if (bHasMultipleColumns)
 			{
-				// If the NameWidget has already been disabled, don't re-enable it if IsPropertyEditingEnabled is true.
-				NameWidget->SetEnabled(NameWidgetEnabled.IsBound() ?
-					TAttribute<bool>::Create(TAttribute<bool>::FGetter::CreateLambda(
-						[NameWidgetEnabled, IsPropertyEditingEnabled]()
-						{
-							return NameWidgetEnabled.Get() && IsPropertyEditingEnabled.Get();
-						}))
-					: IsPropertyEditingEnabled);
-
 				NameColumnBox->AddSlot()
 					.HAlign(Row.NameWidget.HorizontalAlignment)
 					.VAlign(Row.NameWidget.VerticalAlignment)
@@ -612,7 +607,7 @@ void SDetailSingleItemRow::Construct( const FArguments& InArgs, FDetailLayoutCus
 			}
 			else
 			{
-				NameColumnBox->SetEnabled(IsPropertyEditingEnabled);
+				NameColumnBox->SetEnabled(IsEnabledAttribute);
 				NameColumnBox->AddSlot()
 					.HAlign(Row.WholeRowWidget.HorizontalAlignment)
 					.VAlign(Row.WholeRowWidget.VerticalAlignment)
@@ -651,7 +646,7 @@ void SDetailSingleItemRow::Construct( const FArguments& InArgs, FDetailLayoutCus
 			if (bCreateResetToDefault)
 			{
 				ResetWidget = SNew(SResetToDefaultPropertyEditor, PropertyHandle)
-					.IsEnabled(IsPropertyEditingEnabled)
+					.IsEnabled(IsEnabledAttribute)
 					.CustomResetToDefault(Row.CustomResetToDefault);
 			}	
 			else
