Missed a file that was causing the wrong indent level colors to be used.

diff --git a/Engine/Source/Editor/PropertyEditor/Private/UserInterface/PropertyEditor/PropertyEditorConstants.cpp b/Engine/Source/Editor/PropertyEditor/Private/UserInterface/PropertyEditor/PropertyEditorConstants.cpp
--- a/Engine/Source/Editor/PropertyEditor/Private/UserInterface/PropertyEditor/PropertyEditorConstants.cpp
+++ b/Engine/Source/Editor/PropertyEditor/Private/UserInterface/PropertyEditor/PropertyEditorConstants.cpp
@@ -2,6 +2,7 @@
 
 #include "UserInterface/PropertyEditor/PropertyEditorConstants.h"
 #include "EditorStyleSet.h"
+#include "Styling/StyleColors.h"
 
 const FName PropertyEditorConstants::PropertyFontStyle( TEXT("PropertyWindow.NormalFont") );
 const FName PropertyEditorConstants::CategoryFontStyle( TEXT("PropertyWindow.BoldFont") );
@@ -17,6 +18,31 @@ const FSlateBrush* PropertyEditorConstants::GetOverlayBrush( const TSharedRef< c
 
 FSlateColor PropertyEditorConstants::GetRowBackgroundColor(int32 IndentLevel) 
 {
-	FLinearColor BackgroundColor = FAppStyle::Get().GetSlateColor("Colors.Background").GetSpecifiedColor();
-	return FSlateColor(BackgroundColor + FLinearColor(2.0f/255, 2.0f/255, 2.0f/255) * (IndentLevel % 4));
+	if (IndentLevel == 0)
+	{
+		return FAppStyle::Get().GetSlateColor("Colors.Background");
+	}
+
+	int32 ColorIndex = 0;
+	int32 Increment = 1;
+
+	for (int i = 0; i < IndentLevel; ++i)
+	{
+		ColorIndex += Increment;
+
+		if (ColorIndex == 0 || ColorIndex == 3)
+		{
+			Increment = -Increment;
+		}
+	}
+
+	static const FLinearColor Colors[] =
+	{
+		COLOR("#2D2D2D"),
+		COLOR("#373737"),
+		COLOR("#434343"),
+		COLOR("#515151")
+	};
+
+	return FSlateColor(Colors[ColorIndex]);
 }
