Asset property editor reskin

diff --git a/Engine/Plugins/Experimental/HairStrands/Source/HairStrandsEditor/Private/GroomMaterialDetails.cpp b/Engine/Plugins/Experimental/HairStrands/Source/HairStrandsEditor/Private/GroomMaterialDetails.cpp
--- a/Engine/Plugins/Experimental/HairStrands/Source/HairStrandsEditor/Private/GroomMaterialDetails.cpp
+++ b/Engine/Plugins/Experimental/HairStrands/Source/HairStrandsEditor/Private/GroomMaterialDetails.cpp
@@ -171,7 +171,7 @@ void FGroomMaterialDetails::AddMaterials(IDetailLayoutBuilder& DetailLayout)
 			//Pass an empty material list owner (owner can be use by the asset picker filter. In this case we do not need it)
 			TArray<FAssetData> MaterialListOwner;
 			MaterialListOwner.Add(GroomAsset);
-			MaterialCategory.AddCustomBuilder(MakeShareable(new FMaterialList(MaterialCategory.GetParentLayout(), MaterialListDelegates, MaterialListOwner, false, true, true)));
+			MaterialCategory.AddCustomBuilder(MakeShareable(new FMaterialList(MaterialCategory.GetParentLayout(), MaterialListDelegates, MaterialListOwner, false, true)));
 		}
 	}
 }
diff --git a/Engine/Source/Editor/DetailCustomizations/Private/DeviceProfileDetails.cpp b/Engine/Source/Editor/DetailCustomizations/Private/DeviceProfileDetails.cpp
--- a/Engine/Source/Editor/DetailCustomizations/Private/DeviceProfileDetails.cpp
+++ b/Engine/Source/Editor/DetailCustomizations/Private/DeviceProfileDetails.cpp
@@ -725,7 +725,7 @@ void FDeviceProfileConsoleVariablesPropertyDetails::CreateRowWidgetForCVarProper
 			.IsFocusable(false)
 			[
 				SNew(SImage)
-				.Image(FEditorStyle::GetBrush("PropertyWindow.Button_Delete"))
+				.Image(FEditorStyle::GetBrush("Icons.X"))
 			]
 		]
 	];
diff --git a/Engine/Source/Editor/EditorStyle/Private/StarshipStyle.cpp b/Engine/Source/Editor/EditorStyle/Private/StarshipStyle.cpp
--- a/Engine/Source/Editor/EditorStyle/Private/StarshipStyle.cpp
+++ b/Engine/Source/Editor/EditorStyle/Private/StarshipStyle.cpp
@@ -2999,17 +2999,13 @@ void FStarshipEditorStyle::FStyle::SetupPropertyEditorStyles()
 			.SetColorAndOpacity(FColor( 96, 194, 253, 255 ).ReinterpretAsLinear())
 		);
 	
-		Set( "PropertyEditor.AssetName", FTextBlockStyle(NormalText)
-			.SetFont( DEFAULT_FONT( "Regular", 12 ) )
-			.SetColorAndOpacity( FLinearColor::White )
-			.SetShadowOffset( FVector2D(1,1) )
-			.SetShadowColorAndOpacity( FLinearColor::Black )
-		);
 
 		Set( "PropertyEditor.AssetName.ColorAndOpacity", FLinearColor::White );
 
-		Set( "PropertyEditor.AssetThumbnailLight", new BOX_BRUSH( "ContentBrowser/ThumbnailLight", FMargin( 5.0f / 64.0f ), SelectionColor ) );
-		Set( "PropertyEditor.AssetThumbnailShadow", new BOX_BRUSH( "ContentBrowser/ThumbnailShadow", FMargin( 4.0f / 64.0f ) ) );
+		Set("PropertyEditor.AssetThumbnailBorder", new FSlateRoundedBoxBrush(FStyleColors::Transparent, 4.0f, FStyleColors::InputOutline, 1.0f));
+		Set("PropertyEditor.AssetThumbnailBorderHovered", new FSlateRoundedBoxBrush(FStyleColors::Transparent, 4.0f, FStyleColors::Hover2, 1.0f));
+		Set("PropertyEditor.AssetTileItem.DropShadow", new BOX_BRUSH("Starship/ContentBrowser/drop-shadow", FMargin(4.0f / 64.0f)));
+
 		Set( "PropertyEditor.AssetClass", FTextBlockStyle(NormalText)
 			.SetFont( DEFAULT_FONT( "Regular", 10 ) )
 			.SetColorAndOpacity( FLinearColor::White )
@@ -3119,20 +3115,19 @@ void FStarshipEditorStyle::FStyle::SetupPropertyEditorStyles()
 		Set( "PropertyWindow.FilteredEditConstColor", new FSlateColorBrush( FColor( 152, 152, 152, 80 ).ReinterpretAsLinear() * FColor(0,255,0,255).ReinterpretAsLinear() ) );
 		Set( "PropertyWindow.CategoryBackground", new BOX_BRUSH( "/PropertyView/CategoryBackground", FMargin(4.f/16.f) ) );
 		Set( "PropertyWindow.CategoryForeground", FLinearColor::Black );
-		Set( "PropertyWindow.Button_Browse", new IMAGE_BRUSH( "Icons/lens_12x", Icon12x12 ) );
-		Set( "PropertyWindow.Button_CreateNewBlueprint", new IMAGE_BRUSH( "Icons/PlusSymbol_12x", Icon12x12 ) );
-		Set( "PropertyWindow.Button_Use", new IMAGE_BRUSH( "Icons/assign_12x", Icon12x12 ) );
+		Set( "PropertyWindow.Button_Browse", new CORE_IMAGE_BRUSH_SVG("Starship/Common/search", Icon16x16));
+		Set( "PropertyWindow.Button_Use", new CORE_IMAGE_BRUSH_SVG("Starship/Common/circle-arrow-left", Icon16x16));
 		Set( "PropertyWindow.Button_Delete", new IMAGE_BRUSH("Icons/Cross_12x", Icon12x12 ) );
 		Set( "PropertyWindow.Button_Clear", new IMAGE_BRUSH( "Icons/Cross_12x", Icon12x12 ) );
-		Set( "PropertyWindow.Button_Edit", new IMAGE_BRUSH( "Icons/wrench_16x", Icon12x12 ) );
+		Set( "PropertyWindow.Button_Edit", new IMAGE_BRUSH_SVG("Starship/Common/edit", Icon16x16));
 		Set( "PropertyWindow.Button_EmptyArray", new IMAGE_BRUSH( "Icons/empty_set_12x", Icon12x12 ) );
-		Set( "PropertyWindow.Button_AddToArray", new IMAGE_BRUSH( "Icons/PlusSymbol_12x", Icon12x12 ) );
-		Set( "PropertyWindow.Button_RemoveFromArray", new IMAGE_BRUSH( "Icons/MinusSymbol_12x", Icon12x12 ) );
+		Set( "PropertyWindow.Button_AddToArray", new CORE_IMAGE_BRUSH_SVG("Starship/Common/plus-circle", Icon16x16));
+		Set( "PropertyWindow.Button_RemoveFromArray", new IMAGE_BRUSH_SVG("Starship/Common/minus", Icon16x16));
 		Set( "PropertyWindow.Button_Ellipsis", new IMAGE_BRUSH( "Icons/ellipsis_12x", Icon12x12 ) );
 		Set( "PropertyWindow.Button_PickAsset", new IMAGE_BRUSH( "Icons/pillarray_12x", Icon12x12 ) );
 		Set( "PropertyWindow.Button_PickActor", new IMAGE_BRUSH( "Icons/levels_16x", Icon12x12 ) );
 		Set( "PropertyWindow.Button_PickActorInteractive", new IMAGE_BRUSH( "Icons/eyedropper_16px", Icon12x12 ) );
-		Set( "PropertyWindow.Button_Refresh", new IMAGE_BRUSH("Icons/refresh_12x", Icon12x12 ) );
+
 		Set( "PropertyWindow.WindowBorder", new BOX_BRUSH( "Common/GroupBorder", FMargin(4.0f/16.0f) ) );
 
 		Set( "DetailsView.NameChangeCommitted", new BOX_BRUSH( "Common/EditableTextSelectionBackground", FMargin(4.f/16.f) ) );
diff --git a/Engine/Source/Editor/LandscapeEditor/Private/LandscapeEditorDetailCustomization_AlphaBrush.cpp b/Engine/Source/Editor/LandscapeEditor/Private/LandscapeEditorDetailCustomization_AlphaBrush.cpp
--- a/Engine/Source/Editor/LandscapeEditor/Private/LandscapeEditorDetailCustomization_AlphaBrush.cpp
+++ b/Engine/Source/Editor/LandscapeEditor/Private/LandscapeEditorDetailCustomization_AlphaBrush.cpp
@@ -66,7 +66,7 @@ public:
 		[
 			SNew(SBorder)
 			.Padding(4.0f)
-			.BorderImage( FEditorStyle::GetBrush("PropertyEditor.AssetThumbnailShadow") )
+			.BorderImage( FEditorStyle::GetBrush("PropertyEditor.AssetTileItem.DropShadow") )
 			.OnMouseDoubleClick(this, &STextureMaskThumbnail::OnAssetThumbnailDoubleClick)
 			//.OnClicked_Static(&FLandscapeEditorDetailCustomization_AlphaBrush::OnTextureButtonClicked)
 			[
diff --git a/Engine/Source/Editor/Persona/Private/PersonaMeshDetails.cpp b/Engine/Source/Editor/Persona/Private/PersonaMeshDetails.cpp
--- a/Engine/Source/Editor/Persona/Private/PersonaMeshDetails.cpp
+++ b/Engine/Source/Editor/Persona/Private/PersonaMeshDetails.cpp
@@ -2577,7 +2577,7 @@ void FPersonaMeshDetails::AddLODLevelCategories(IDetailLayoutBuilder& DetailLayo
 				//Pass an empty material list owner (owner can be use by the asset picker filter. In this case we do not need it)
 				TArray<FAssetData> MaterialListOwner;
 				MaterialListOwner.Add(SkelMesh);
-				MaterialCategory.AddCustomBuilder(MakeShareable(new FMaterialList(MaterialCategory.GetParentLayout(), MaterialListDelegates, MaterialListOwner, false, true, true)));
+				MaterialCategory.AddCustomBuilder(MakeShareable(new FMaterialList(MaterialCategory.GetParentLayout(), MaterialListDelegates, MaterialListOwner, false, true)));
 			}
 		}
 
diff --git a/Engine/Source/Editor/PropertyEditor/Private/MaterialList.cpp b/Engine/Source/Editor/PropertyEditor/Private/MaterialList.cpp
--- a/Engine/Source/Editor/PropertyEditor/Private/MaterialList.cpp
+++ b/Engine/Source/Editor/PropertyEditor/Private/MaterialList.cpp
@@ -11,6 +11,8 @@
 #include "Widgets/Input/SHyperlink.h"
 #include "Widgets/Layout/SBox.h"
 #include "Widgets/Text/STextBlock.h"
+#include "Widgets/Images/SImage.h"
+#include "Styling/SlateIconFinder.h"
 
 #define LOCTEXT_NAMESPACE "PropertyCustomizationHelpers"
 
@@ -112,10 +114,9 @@ public:
 		FOnGenerateWidgetsForMaterial InOnGenerateWidgetsForMaterial, 
 		FOnResetMaterialToDefaultClicked InOnResetToDefaultClicked,
 		int32 InMultipleMaterialCount,
-		bool bShowUsedTextures,
-		bool bDisplayCompactSize)
+		bool bShowUsedTextures)
 	{
-		return MakeShareable( new FMaterialItemView( Material, InOnMaterialChanged, InOnGenerateNameWidgetsForMaterial, InOnGenerateWidgetsForMaterial, InOnResetToDefaultClicked, InMultipleMaterialCount, bShowUsedTextures, bDisplayCompactSize) );
+		return MakeShareable( new FMaterialItemView( Material, InOnMaterialChanged, InOnGenerateNameWidgetsForMaterial, InOnGenerateWidgetsForMaterial, InOnResetToDefaultClicked, InMultipleMaterialCount, bShowUsedTextures) );
 	}
 
 	TSharedRef<SWidget> CreateNameContent()
@@ -142,75 +143,56 @@ public:
 
 	TSharedRef<SWidget> CreateValueContent( const TSharedPtr<FAssetThumbnailPool>& ThumbnailPool, const TArray<FAssetData>& OwnerAssetDataArray)
 	{
-		FIntPoint ThumbnailSize(64, 64);
-
 		return
-			SNew(SHorizontalBox)
-			+ SHorizontalBox::Slot()
+			SNew(SVerticalBox)
+			+SVerticalBox::Slot()
+			.AutoHeight()
+			.Padding( 0.0f )
+			.VAlign(VAlign_Center)
+			.HAlign(HAlign_Fill)
 			[
-				SNew(SVerticalBox)
-				+SVerticalBox::Slot()
-				.AutoHeight()
-				.Padding( 0.0f )
-				.VAlign(VAlign_Center)
-				.HAlign(HAlign_Fill)
+				SNew(SObjectPropertyEntryBox)
+				.ObjectPath(this, &FMaterialItemView::OnGetObjectPath)
+				.AllowClear(false)
+				.AllowedClass(UMaterialInterface::StaticClass())
+				.OnObjectChanged(this, &FMaterialItemView::OnSetObject)
+				.ThumbnailPool(ThumbnailPool)
+				.DisplayCompactSize(true)
+				.OwnerAssetDataArray(OwnerAssetDataArray)
+				.CustomContentSlot()
 				[
-					SNew(SHorizontalBox)
-					+SHorizontalBox::Slot()
-					.FillWidth(1.0f)
+					SNew( SBox )
+					.HAlign(HAlign_Left)
+					.VAlign(VAlign_Center)
 					[
-						SNew( SObjectPropertyEntryBox )
-						.ObjectPath(this, &FMaterialItemView::OnGetObjectPath)
-						.AllowClear(false)
-						.AllowedClass(UMaterialInterface::StaticClass())
-						.OnObjectChanged(this, &FMaterialItemView::OnSetObject)
-						.ThumbnailPool(ThumbnailPool)
-						.DisplayCompactSize(bDisplayCompactSize)
-						.OwnerAssetDataArray(OwnerAssetDataArray)
-						.CustomContentSlot()
+						SNew(SHorizontalBox)
+						+SHorizontalBox::Slot()
+						.VAlign(VAlign_Center)
+						.Padding(0.0f, 0.0f, 3.0f, 0.0f)
+						.AutoWidth()
 						[
-							SNew( SBox )
-							.HAlign(HAlign_Left)
+							// Add a menu for displaying all textures 
+							SNew(SComboButton)
+							.OnGetMenuContent(this, &FMaterialItemView::OnGetTexturesMenuForMaterial)
+							.ComboButtonStyle(FAppStyle::Get(), "SimpleComboButton")
 							.VAlign(VAlign_Center)
+							.IsEnabled( this, &FMaterialItemView::IsTexturesMenuEnabled )
+							.Visibility(bShowUsedTextures ? EVisibility::Visible : EVisibility::Hidden)
+							.ButtonContent()
 							[
-								SNew(SHorizontalBox)
-								+SHorizontalBox::Slot()
-								.VAlign(VAlign_Center)
-								.Padding(0.0f, 0.0f, 3.0f, 0.0f)
-								.AutoWidth()
-								[
-									// Add a menu for displaying all textures 
-									SNew( SComboButton )
-									.OnGetMenuContent( this, &FMaterialItemView::OnGetTexturesMenuForMaterial )
-									.VAlign(VAlign_Center)
-									.ContentPadding(2)
-									.IsEnabled( this, &FMaterialItemView::IsTexturesMenuEnabled )
-									.Visibility( bShowUsedTextures ? EVisibility::Visible : EVisibility::Hidden )
-									.ButtonContent()
-									[
-										SNew( STextBlock )
-										.Font( IDetailLayoutBuilder::GetDetailFont() )
-										.ToolTipText( LOCTEXT("ViewTexturesToolTip", "View the textures used by this material" ) )
-										.Text( LOCTEXT("ViewTextures","Textures") )
-									]
-								]
-								+SHorizontalBox::Slot()
-								.Padding(3.0f, 0.0f)
-								.FillWidth(1.0f)
-								[
-									OnGenerateCustomMaterialWidgets.IsBound() && bDisplayCompactSize ? OnGenerateCustomMaterialWidgets.Execute(MaterialItem.Material.Get(), MaterialItem.SlotIndex) : StaticCastSharedRef<SWidget>(SNullWidget::NullWidget)
-								]
+								SNew(SImage)
+								.Image(FSlateIconFinder::FindIconBrushForClass(UTexture2D::StaticClass()))
+								.ColorAndOpacity(FSlateColor::UseForeground())
 							]
 						]
+						+SHorizontalBox::Slot()
+						.Padding(3.0f, 0.0f)
+						.FillWidth(1.0f)
+						[
+							OnGenerateCustomMaterialWidgets.IsBound() ? OnGenerateCustomMaterialWidgets.Execute(MaterialItem.Material.Get(), MaterialItem.SlotIndex) : SNullWidget::NullWidget
+						]
 					]
 				]
-				+SVerticalBox::Slot()
-				.AutoHeight()
-				.Padding(2)
-				.VAlign( VAlign_Center )
-				[
-					OnGenerateCustomMaterialWidgets.IsBound() && !bDisplayCompactSize ? OnGenerateCustomMaterialWidgets.Execute( MaterialItem.Material.Get(), MaterialItem.SlotIndex ) : StaticCastSharedRef<SWidget>( SNullWidget::NullWidget )
-				]
 			];
 	}
 
@@ -245,8 +227,7 @@ private:
 						FOnGenerateWidgetsForMaterial& InOnGenerateMaterialWidgets, 
 						FOnResetMaterialToDefaultClicked& InOnResetToDefaultClicked,
 						int32 InMultipleMaterialCount,
-						bool bInShowUsedTextures,
-						bool bInDisplayCompactSize)
+						bool bInShowUsedTextures)
 						
 		: MaterialItem( InMaterial )
 		, OnMaterialChanged( InOnMaterialChanged )
@@ -255,7 +236,6 @@ private:
 		, OnResetToDefaultClicked( InOnResetToDefaultClicked )
 		, MultipleMaterialCount( InMultipleMaterialCount )
 		, bShowUsedTextures( bInShowUsedTextures )
-		, bDisplayCompactSize(bInDisplayCompactSize)
 	{
 
 	}
@@ -341,17 +321,15 @@ private:
 	FOnResetMaterialToDefaultClicked OnResetToDefaultClicked;
 	int32 MultipleMaterialCount;
 	bool bShowUsedTextures;
-	bool bDisplayCompactSize;
 };
 
 
-FMaterialList::FMaterialList(IDetailLayoutBuilder& InDetailLayoutBuilder, FMaterialListDelegates& InMaterialListDelegates, const TArray<FAssetData>& InOwnerAssetDataArray, bool bInAllowCollapse, bool bInShowUsedTextures, bool bInDisplayCompactSize)
+FMaterialList::FMaterialList(IDetailLayoutBuilder& InDetailLayoutBuilder, FMaterialListDelegates& InMaterialListDelegates, const TArray<FAssetData>& InOwnerAssetDataArray, bool bInAllowCollapse, bool bInShowUsedTextures)
 	: MaterialListDelegates( InMaterialListDelegates )
 	, DetailLayoutBuilder( InDetailLayoutBuilder )
 	, MaterialListBuilder( new FMaterialListBuilder )
 	, bAllowCollpase(bInAllowCollapse)
 	, bShowUsedTextures(bInShowUsedTextures)
-	, bDisplayCompactSize(bInDisplayCompactSize)
 	, OwnerAssetDataArray(InOwnerAssetDataArray)
 {
 }
@@ -588,7 +566,7 @@ void FMaterialList::AddMaterialItem( FDetailWidgetRow& Row, int32 CurrentSlot, c
 {
 	uint32 NumMaterials = MaterialListBuilder->GetNumMaterialsInSlot(CurrentSlot);
 
-	TSharedRef<FMaterialItemView> NewView = FMaterialItemView::Create( Item, MaterialListDelegates.OnMaterialChanged, MaterialListDelegates.OnGenerateCustomNameWidgets, MaterialListDelegates.OnGenerateCustomMaterialWidgets, MaterialListDelegates.OnResetMaterialToDefaultClicked, NumMaterials, bShowUsedTextures, bDisplayCompactSize);
+	TSharedRef<FMaterialItemView> NewView = FMaterialItemView::Create( Item, MaterialListDelegates.OnMaterialChanged, MaterialListDelegates.OnGenerateCustomNameWidgets, MaterialListDelegates.OnGenerateCustomMaterialWidgets, MaterialListDelegates.OnResetMaterialToDefaultClicked, NumMaterials, bShowUsedTextures);
 
 	TSharedPtr<SWidget> RightSideContent;
 	if( bDisplayLink )
diff --git a/Engine/Source/Editor/PropertyEditor/Private/PropertyCustomizationHelpers.cpp b/Engine/Source/Editor/PropertyEditor/Private/PropertyCustomizationHelpers.cpp
--- a/Engine/Source/Editor/PropertyEditor/Private/PropertyCustomizationHelpers.cpp
+++ b/Engine/Source/Editor/PropertyEditor/Private/PropertyCustomizationHelpers.cpp
@@ -67,6 +67,7 @@ namespace PropertyCustomizationHelpers
 				.OnClicked( this, &SPropertyEditorButton::OnClick )
 				.ToolTipText( InArgs._Text )
 				.ForegroundColor( FSlateColor::UseForeground() )
+				.ContentPadding(0)
 				.IsFocusable(InArgs._IsFocusable)
 				[ 
 					SNew( SImage )
@@ -156,7 +157,7 @@ namespace PropertyCustomizationHelpers
 			SNew( SPropertyEditorButton )
 			.Text( LOCTEXT( "DeleteButtonLabel", "Delete") )
 			.ToolTipText( OptionalToolTipText.Get().IsEmpty() ? LOCTEXT( "DeleteButtonToolTipText", "Delete") : OptionalToolTipText )
-			.Image( FEditorStyle::GetBrush("PropertyWindow.Button_Delete") )
+			.Image( FEditorStyle::GetBrush("Icons.X") )
 			.OnClickAction( OnDeleteClicked )
 			.IsEnabled(IsEnabled)
 			.IsFocusable( false );
@@ -218,7 +219,7 @@ namespace PropertyCustomizationHelpers
 			SNew( SPropertyEditorButton )
 			.Text( LOCTEXT( "NewBlueprintButtonLabel", "New Blueprint") )
 			.ToolTipText( OptionalToolTipText.Get().IsEmpty() ? LOCTEXT( "NewBlueprintButtonToolTipText", "Create New Blueprint") : OptionalToolTipText )
-			.Image( FEditorStyle::GetBrush("PropertyWindow.Button_CreateNewBlueprint") )
+			.Image( FEditorStyle::GetBrush("Icons.PlusCircle") )
 			.OnClickAction( OnNewBlueprintClicked )
 			.IsEnabled(IsEnabled)
 			.IsFocusable( false );
@@ -457,7 +458,7 @@ void SObjectPropertyEntryBox::Construct( const FArguments& InArgs )
 	const TArray<FAssetData>& OwnerAssetDataArray = InArgs._OwnerAssetDataArray;
 
 	bool bDisplayThumbnail = InArgs._DisplayThumbnail;
-	FIntPoint ThumbnailSize(64, 64);
+	FIntPoint ThumbnailSize(48, 48);
 	if (InArgs._ThumbnailSizeOverride.IsSet())
 	{
 		ThumbnailSize = InArgs._ThumbnailSizeOverride.Get();
diff --git a/Engine/Source/Editor/PropertyEditor/Private/UserInterface/PropertyEditor/SPropertyEditorAsset.cpp b/Engine/Source/Editor/PropertyEditor/Private/UserInterface/PropertyEditor/SPropertyEditorAsset.cpp
--- a/Engine/Source/Editor/PropertyEditor/Private/UserInterface/PropertyEditor/SPropertyEditorAsset.cpp
+++ b/Engine/Source/Editor/PropertyEditor/Private/UserInterface/PropertyEditor/SPropertyEditorAsset.cpp
@@ -107,6 +107,14 @@ bool SPropertyEditorAsset::ShouldDisplayThumbnail(const FArguments& InArgs, cons
 	return bShowThumbnail;
 }
 
+const FSlateBrush* SPropertyEditorAsset::GetThumbnailBorder() const
+{
+	static const FName HoveredBorderName("PropertyEditor.AssetThumbnailBorderHovered");
+	static const FName RegularBorderName("PropertyEditor.AssetThumbnailBorder");
+
+	return ThumbnailBorder->IsHovered() ? FAppStyle::Get().GetBrush(HoveredBorderName) : FAppStyle::Get().GetBrush(RegularBorderName);
+}
+
 void SPropertyEditorAsset::InitializeClassFilters(const FProperty* Property)
 {
 	if (Property == nullptr)
@@ -405,12 +413,9 @@ void SPropertyEditorAsset::Construct(const FArguments& InArgs, const TSharedPtr<
 
 	AssetComboButton = SNew(SComboButton)
 		.ToolTipText(TooltipAttribute)
-		.ButtonStyle( FEditorStyle::Get(), "PropertyEditor.AssetComboStyle" )
-		.ForegroundColor(FEditorStyle::GetColor("PropertyEditor.AssetName.ColorAndOpacity"))
 		.OnGetMenuContent( this, &SPropertyEditorAsset::OnGetMenuContent )
 		.OnMenuOpenChanged( this, &SPropertyEditorAsset::OnMenuOpenChanged )
-		.IsEnabled( IsEnabledAttribute )
-		.ContentPadding(2.0f)
+		.IsEnabled(IsEnabledAttribute)
 		.ButtonContent()
 		[
 			SNew(SHorizontalBox)
@@ -429,7 +434,6 @@ void SPropertyEditorAsset::Construct(const FArguments& InArgs, const TSharedPtr<
 			[
 				// Show the name of the asset or actor
 				SNew(STextBlock)
-				.TextStyle( FEditorStyle::Get(), "PropertyEditor.AssetClass" )
 				.Font( FEditorStyle::GetFontStyle( PropertyEditorConstants::PropertyFontStyle ) )
 				.Text(this,&SPropertyEditorAsset::OnGetAssetName)
 			]
@@ -443,7 +447,7 @@ void SPropertyEditorAsset::Construct(const FArguments& InArgs, const TSharedPtr<
 	TSharedPtr<SWidget> ButtonBoxWrapper;
 	TSharedRef<SHorizontalBox> ButtonBox = SNew( SHorizontalBox );
 	
-	TSharedPtr<SVerticalBox> CustomContentBox;
+	TSharedPtr<SHorizontalBox> CustomContentBox;
 
 	if(ShouldDisplayThumbnail(InArgs, ObjectClass))
 	{
@@ -466,93 +470,72 @@ void SPropertyEditorAsset::Construct(const FArguments& InArgs, const TSharedPtr<
 		}
 
 		ValueContentBox->AddSlot()
-		.Padding( 0.0f, 0.0f, 2.0f, 0.0f )
+		.Padding(FMargin(0,5,5,5))
 		.AutoWidth()
 		.VAlign(VAlign_Center)
 		[
-			SNew(SVerticalBox)
-			+ SVerticalBox::Slot()
-			.AutoHeight()
+			SNew(SBorder)
+			.Visibility(EVisibility::SelfHitTestInvisible)
+			.Padding(FMargin(0.0f, 0.0f, 4.0f, 4.0f))
+			.BorderImage(FAppStyle::Get().GetBrush("PropertyEditor.AssetTileItem.DropShadow"))
 			[
-				SAssignNew( ThumbnailBorder, SBorder )
-				.Padding( 5.0f )
-				.BorderImage( this, &SPropertyEditorAsset::GetThumbnailBorder )
-				.OnMouseDoubleClick( this, &SPropertyEditorAsset::OnAssetThumbnailDoubleClick )
+				SNew(SOverlay)
+				+SOverlay::Slot()
+				.Padding(1)
 				[
-					SNew( SBox )
-					.ToolTipText(TooltipAttribute)
-					.WidthOverride( InArgs._ThumbnailSize.X ) 
-					.HeightOverride( InArgs._ThumbnailSize.Y )
+					SAssignNew(ThumbnailBorder, SBorder)
+					.Padding(0)
+					.BorderImage(FStyleDefaults::GetNoBrush())
+					.OnMouseDoubleClick(this, &SPropertyEditorAsset::OnAssetThumbnailDoubleClick)
 					[
-						AssetThumbnail->MakeThumbnailWidget(AssetThumbnailConfig)
+						SNew(SBox)
+						.ToolTipText(TooltipAttribute)
+						.WidthOverride(InArgs._ThumbnailSize.X)
+						.HeightOverride(InArgs._ThumbnailSize.Y)
+						[
+							AssetThumbnail->MakeThumbnailWidget(AssetThumbnailConfig)
+						]
 					]
 				]
+				+ SOverlay::Slot()
+				[
+					SNew(SImage)
+					.Image(this, &SPropertyEditorAsset::GetThumbnailBorder)
+				]
 			]
 		];
 
-		if(InArgs._DisplayCompactSize)
-		{
-			ValueContentBox->AddSlot()
+	
+		ValueContentBox->AddSlot()
+		.Padding(FMargin(0.0f, 5.0f))
+		[
+			SNew(SVerticalBox)
+			+ SVerticalBox::Slot()
+			.VAlign(VAlign_Center)
 			[
-				SNew( SBox )
-				.VAlign( VAlign_Center )
-				[
-					SAssignNew( CustomContentBox, SVerticalBox )
-					+ SVerticalBox::Slot()
-					.AutoHeight()
-					.Padding(0.0f, 4.0f, 0.0f, 0.0f)
-					[
-						SNew(SHorizontalBox)
-						+SHorizontalBox::Slot()
-						.FillWidth(1.0f)
-						.VAlign(VAlign_Center)
-						[
-							AssetComboButton.ToSharedRef()
-						]
-						+ SHorizontalBox::Slot()
-						.AutoWidth()
-						[
-							SAssignNew(ButtonBoxWrapper, SBox)
-							.Padding(FMargin(0.0f, 2.0f, 4.0f, 2.0f))
-							[
-								ButtonBox
-							]
-						]
-					]
-				]
-			];
-		}
-		else
-		{
-			ValueContentBox->AddSlot()
+				AssetComboButton.ToSharedRef()
+			]
+			+ SVerticalBox::Slot()
+			.VAlign(VAlign_Center)
 			[
-				SNew( SBox )
-				.VAlign( VAlign_Center )
+				SAssignNew(CustomContentBox, SHorizontalBox)
+				+SHorizontalBox::Slot()
+				.AutoWidth()
 				[
-					SAssignNew( CustomContentBox, SVerticalBox )
-					+ SVerticalBox::Slot()
-					.Padding(0.0f, 4.0f, 0.0f, 0.0f)
-					[
-						AssetComboButton.ToSharedRef()
-					]
-					+ SVerticalBox::Slot()
-					.AutoHeight()
+					SAssignNew(ButtonBoxWrapper, SBox)
+					.Padding(FMargin(0.0f, 2.0f, 4.0f, 2.0f))
 					[
-						SAssignNew(ButtonBoxWrapper, SBox)
-						.Padding( FMargin( 0.0f, 2.0f, 4.0f, 2.0f ) )
-						[
-							ButtonBox
-						]
+						ButtonBox
 					]
 				]
-			];
-		}
+			]
+		];
 	}
 	else
 	{
 		ValueContentBox->AddSlot()
 		[
-			SAssignNew( CustomContentBox, SVerticalBox )
+			SNew(SVerticalBox)
 			+SVerticalBox::Slot()
 			.VAlign( VAlign_Center )
 			[
@@ -565,12 +548,18 @@ void SPropertyEditorAsset::Construct(const FArguments& InArgs, const TSharedPtr<
 				.AutoWidth()
 				[
 					SAssignNew(ButtonBoxWrapper, SBox)
-					.Padding( FMargin( 4.f, 0.f ) )
+					.Padding(FMargin(4.0f,0.0f))
 					[
 						ButtonBox
 					]
 				]
 			]
+			+ SVerticalBox::Slot()
+			.VAlign(VAlign_Center)
+			.AutoHeight()
+			[
+				SAssignNew(CustomContentBox, SHorizontalBox)
+			]
 		];
 	}
 
@@ -636,18 +625,6 @@ void SPropertyEditorAsset::GetDesiredWidth( float& OutMinDesiredWidth, float &Ou
 	OutMaxDesiredWidth = 350.f;
 }
 
-const FSlateBrush* SPropertyEditorAsset::GetThumbnailBorder() const
-{
-	if ( ThumbnailBorder->IsHovered() )
-	{
-		return FEditorStyle::GetBrush("PropertyEditor.AssetThumbnailLight");
-	}
-	else
-	{
-		return FEditorStyle::GetBrush("PropertyEditor.AssetThumbnailShadow");
-	}
-}
-
 const FSlateBrush* SPropertyEditorAsset::GetStatusIcon() const
 {
 	static FSlateNoResource EmptyBrush = FSlateNoResource();
diff --git a/Engine/Source/Editor/PropertyEditor/Private/UserInterface/PropertyEditor/SPropertyEditorAsset.h b/Engine/Source/Editor/PropertyEditor/Private/UserInterface/PropertyEditor/SPropertyEditorAsset.h
--- a/Engine/Source/Editor/PropertyEditor/Private/UserInterface/PropertyEditor/SPropertyEditorAsset.h
+++ b/Engine/Source/Editor/PropertyEditor/Private/UserInterface/PropertyEditor/SPropertyEditorAsset.h
@@ -43,38 +43,34 @@ public:
 		Error,
 	};
 
-	SLATE_BEGIN_ARGS( SPropertyEditorAsset )
-		: _AssetFont( FEditorStyle::GetFontStyle("PropertyEditor.AssetName.Font") ) 
-		, _ClassFont( FEditorStyle::GetFontStyle("PropertyEditor.AssetClass.Font") ) 
-		, _DisplayThumbnail(true)
+	SLATE_BEGIN_ARGS(SPropertyEditorAsset)
+		: _DisplayThumbnail(true)
 		, _DisplayUseSelected(true)
 		, _DisplayBrowse(true)
 		, _EnableContentPicker(true)
-		, _DisplayCompactSize(false)
+		, _DisplayCompactSize(true)
 		, _ThumbnailPool(nullptr)
-		, _ThumbnailSize(FIntPoint(64, 64))
+		, _ThumbnailSize(FIntPoint(48, 48))
 		, _ObjectPath()
 		, _Class(nullptr)
 		, _CustomContentSlot()
-		{}
-		SLATE_ATTRIBUTE( FSlateFontInfo, AssetFont )
-		SLATE_ATTRIBUTE( FSlateFontInfo, ClassFont )
-		SLATE_ARGUMENT( TOptional<bool>, AllowClear )
-		SLATE_ARGUMENT( bool, DisplayThumbnail )
-		SLATE_ARGUMENT( bool, DisplayUseSelected )
-		SLATE_ARGUMENT( bool, DisplayBrowse )
-		SLATE_ARGUMENT( bool, EnableContentPicker)
-		SLATE_ARGUMENT( bool, DisplayCompactSize)
-		SLATE_ARGUMENT( TSharedPtr<FAssetThumbnailPool>, ThumbnailPool )
-		SLATE_ARGUMENT( FIntPoint, ThumbnailSize )
-		SLATE_ATTRIBUTE( FString, ObjectPath )
-		SLATE_ARGUMENT( UClass*, Class )
-		SLATE_ARGUMENT( TOptional<TArray<UFactory*>>, NewAssetFactories )
-		SLATE_EVENT( FOnSetObject, OnSetObject )
+	{}
+		SLATE_ARGUMENT(TOptional<bool>, AllowClear)
+		SLATE_ARGUMENT(bool, DisplayThumbnail)
+		SLATE_ARGUMENT(bool, DisplayUseSelected)
+		SLATE_ARGUMENT(bool, DisplayBrowse)
+		SLATE_ARGUMENT(bool, EnableContentPicker)
+		SLATE_ARGUMENT(bool, DisplayCompactSize)
+		SLATE_ARGUMENT(TSharedPtr<FAssetThumbnailPool>, ThumbnailPool)
+		SLATE_ARGUMENT(FIntPoint, ThumbnailSize)
+		SLATE_ATTRIBUTE(FString, ObjectPath)
+		SLATE_ARGUMENT(UClass*, Class)
+		SLATE_ARGUMENT(TOptional<TArray<UFactory*>>, NewAssetFactories)
+		SLATE_EVENT(FOnSetObject, OnSetObject)
 		SLATE_EVENT(FOnShouldFilterAsset, OnShouldFilterAsset)
-		SLATE_NAMED_SLOT( FArguments, CustomContentSlot )
-		SLATE_ARGUMENT( TSharedPtr<IPropertyHandle>, PropertyHandle )
-		SLATE_ARGUMENT( TArray<FAssetData>, OwnerAssetDataArray)
+		SLATE_NAMED_SLOT(FArguments, CustomContentSlot)
+		SLATE_ARGUMENT(TSharedPtr<IPropertyHandle>, PropertyHandle)
+		SLATE_ARGUMENT(TArray<FAssetData>, OwnerAssetDataArray)
 
 	SLATE_END_ARGS()
 
diff --git a/Engine/Source/Editor/PropertyEditor/Private/UserInterface/PropertyEditor/SPropertyEditorClass.cpp b/Engine/Source/Editor/PropertyEditor/Private/UserInterface/PropertyEditor/SPropertyEditorClass.cpp
--- a/Engine/Source/Editor/PropertyEditor/Private/UserInterface/PropertyEditor/SPropertyEditorClass.cpp
+++ b/Engine/Source/Editor/PropertyEditor/Private/UserInterface/PropertyEditor/SPropertyEditorClass.cpp
@@ -71,7 +71,7 @@ private:
 
 void SPropertyEditorClass::GetDesiredWidth(float& OutMinDesiredWidth, float& OutMaxDesiredWidth)
 {
-	OutMinDesiredWidth = 125.0f;
+	OutMinDesiredWidth = 200.0f;
 	OutMaxDesiredWidth = 400.0f;
 }
 
@@ -198,7 +198,6 @@ void SPropertyEditorClass::Construct(const FArguments& InArgs, const TSharedPtr<
 
 	SAssignNew(ComboButton, SComboButton)
 		.OnGetMenuContent(this, &SPropertyEditorClass::GenerateClassPicker)
-		.ContentPadding(FMargin(2.0f, 2.0f))
 		.ToolTipText(this, &SPropertyEditorClass::GetDisplayValueAsString)
 		.ButtonContent()
 		[
diff --git a/Engine/Source/Editor/PropertyEditor/Private/UserInterface/PropertyEditor/SPropertyEditorNumeric.h b/Engine/Source/Editor/PropertyEditor/Private/UserInterface/PropertyEditor/SPropertyEditorNumeric.h
--- a/Engine/Source/Editor/PropertyEditor/Private/UserInterface/PropertyEditor/SPropertyEditorNumeric.h
+++ b/Engine/Source/Editor/PropertyEditor/Private/UserInterface/PropertyEditor/SPropertyEditorNumeric.h
@@ -163,7 +163,6 @@ public:
 			// Constructs the UI for bitmask property editing.
 			SAssignNew(PrimaryWidget, SComboButton)
 			.ComboButtonStyle(&ComboBoxStyle.ComboButtonStyle)
-			.ContentPadding(FMargin(4.0, 2.0))
 			.ToolTipText_Lambda([this, CreateBitmaskFlagsArray, Property]
 			{
 				TOptional<NumericType> Value = OnGetValue();
diff --git a/Engine/Source/Editor/PropertyEditor/Public/MaterialList.h b/Engine/Source/Editor/PropertyEditor/Public/MaterialList.h
--- a/Engine/Source/Editor/PropertyEditor/Public/MaterialList.h
+++ b/Engine/Source/Editor/PropertyEditor/Public/MaterialList.h
@@ -137,7 +137,7 @@ class FMaterialList
 	, public TSharedFromThis<FMaterialList>
 {
 public:
-	PROPERTYEDITOR_API FMaterialList( IDetailLayoutBuilder& InDetailLayoutBuilder, FMaterialListDelegates& MaterialListDelegates, const TArray<FAssetData>& InOwnerAssetDataArray, bool bInAllowCollapse = false, bool bInShowUsedTextures = true, bool bInDisplayCompactSize = false);
+	PROPERTYEDITOR_API FMaterialList( IDetailLayoutBuilder& InDetailLayoutBuilder, FMaterialListDelegates& MaterialListDelegates, const TArray<FAssetData>& InOwnerAssetDataArray, bool bInAllowCollapse = false, bool bInShowUsedTextures = true);
 
 	/**
 	 * @return true if materials are being displayed.                                                          
@@ -213,8 +213,6 @@ private:
 	bool bAllowCollpase;
 	/** Whether or not to use the used textures menu for each material entry */
 	bool bShowUsedTextures;
-	/** Whether or not to display a compact form of material entry*/
-	bool bDisplayCompactSize;
 	/** The mesh asset that owns these materials */
 	TArray<FAssetData> OwnerAssetDataArray;
 };
diff --git a/Engine/Source/Editor/StaticMeshEditor/Private/StaticMeshEditorTools.cpp b/Engine/Source/Editor/StaticMeshEditor/Private/StaticMeshEditorTools.cpp
--- a/Engine/Source/Editor/StaticMeshEditor/Private/StaticMeshEditorTools.cpp
+++ b/Engine/Source/Editor/StaticMeshEditor/Private/StaticMeshEditorTools.cpp
@@ -2508,7 +2508,7 @@ void FMeshMaterialsLayout::AddToCategory(IDetailCategoryBuilder& CategoryBuilder
 	MaterialListDelegates.OnCanCopyMaterialItem.BindSP(this, &FMeshMaterialsLayout::OnCanCopyMaterialItem);
 	MaterialListDelegates.OnPasteMaterialItem.BindSP(this, &FMeshMaterialsLayout::OnPasteMaterialItem);
 
-	CategoryBuilder.AddCustomBuilder(MakeShareable(new FMaterialList(CategoryBuilder.GetParentLayout(), MaterialListDelegates, AssetDataArray, false, true, true)));
+	CategoryBuilder.AddCustomBuilder(MakeShareable(new FMaterialList(CategoryBuilder.GetParentLayout(), MaterialListDelegates, AssetDataArray, false, true)));
 }
 
 void FMeshMaterialsLayout::OnCopyMaterialList()
diff --git a/Engine/Source/Editor/UnrealEd/Classes/Editor/PropertyEditorTestObject.h b/Engine/Source/Editor/UnrealEd/Classes/Editor/PropertyEditorTestObject.h
--- a/Engine/Source/Editor/UnrealEd/Classes/Editor/PropertyEditorTestObject.h
+++ b/Engine/Source/Editor/UnrealEd/Classes/Editor/PropertyEditorTestObject.h
@@ -356,6 +356,9 @@ class UPropertyEditorTestObject : public UObject
 	UPROPERTY(EditAnywhere, Category=ArraysOfProperties)
 	int32 StaticArrayOfIntsWithEnumLabels[ArrayIndex_MAX];
 
+	UPROPERTY(VisibleAnywhere, Category = AdvancedProperties)
+	FFloatRange FloatRange;
+
 	// This is a float property tooltip that is overridden
 	UPROPERTY(EditAnywhere, Category=AdvancedProperties, meta=(ClampMin="0.0", ClampMax="100.0", UIMin="0.0", UIMax="50.0", ToolTip="This is a custom tooltip that should be shown"))
 	float FloatPropertyWithClampedRange;
@@ -459,6 +462,9 @@ class UPropertyEditorTestObject : public UObject
 	UPROPERTY(EditAnywhere, Category=AssetPropertyTests)
 	UMaterialInterface* AnyMaterialInterface;
 
+	UPROPERTY(EditAnywhere, Category = AssetPropertyTests, meta=(DisplayThumbnail=false))
+	UMaterialInterface* MaterialNoThumbnail;
+
 	UPROPERTY(EditAnywhere, Category=AssetPropertyTests)
 	AActor* OnlyActorsAllowed;
 
diff --git a/Engine/Source/Editor/UnrealEd/Public/GameModeInfoCustomizer.h b/Engine/Source/Editor/UnrealEd/Public/GameModeInfoCustomizer.h
--- a/Engine/Source/Editor/UnrealEd/Public/GameModeInfoCustomizer.h
+++ b/Engine/Source/Editor/UnrealEd/Public/GameModeInfoCustomizer.h
@@ -155,16 +155,19 @@ public:
 			]
 
 			+SHorizontalBox::Slot()
+			.VAlign(VAlign_Center)
 			.AutoWidth()
 			[
 				PropertyCustomizationHelpers::MakeUseSelectedButton(FSimpleDelegate::CreateSP(this, &FGameModeInfoCustomizer::OnUseSelectedGameModeClicked))
 			]
 			+SHorizontalBox::Slot()
+			.VAlign(VAlign_Center)
 			.AutoWidth()
 			[
 				PropertyCustomizationHelpers::MakeBrowseButton(FSimpleDelegate::CreateSP(this, &FGameModeInfoCustomizer::OnBrowseGameModeClicked))
 			]
 			+SHorizontalBox::Slot()
+			.VAlign(VAlign_Center)
 			.AutoWidth()
 			[
 				PropertyCustomizationHelpers::MakeNewBlueprintButton(FSimpleDelegate::CreateSP(this, &FGameModeInfoCustomizer::OnClickNewGameMode))
