Added a grid line below category headers in the details view.

diff --git a/Engine/Source/Editor/PropertyEditor/Private/DetailCategoryGroupNode.cpp b/Engine/Source/Editor/PropertyEditor/Private/DetailCategoryGroupNode.cpp
--- a/Engine/Source/Editor/PropertyEditor/Private/DetailCategoryGroupNode.cpp
+++ b/Engine/Source/Editor/PropertyEditor/Private/DetailCategoryGroupNode.cpp
@@ -14,7 +14,6 @@ void SDetailCategoryTableRow::Construct( const FArguments& InArgs, TSharedRef<FD
 	TSharedPtr<SWidget> Widget = SNullWidget::NullWidget;
 	FDetailColumnSizeData& ColumnSizeData = InOwnerTreeNode->GetDetailsView()->GetColumnSizeData();
 
-
 	Widget = SNew(SHorizontalBox)
 		+ SHorizontalBox::Slot()
 		.VAlign(VAlign_Center)
@@ -40,17 +39,22 @@ void SDetailCategoryTableRow::Construct( const FArguments& InArgs, TSharedRef<FD
 			InArgs._HeaderContent.IsValid() ? InArgs._HeaderContent.ToSharedRef() : SNullWidget::NullWidget
 		];
 
-	const float VerticalPadding = bIsInnerCategory ? 8 : 10;
+	const float VerticalPadding = bIsInnerCategory ? 6 : 8;
 	
 	ChildSlot
 	.Padding(0)
-	[	
+	[
 		SNew( SBorder )
-		.BorderImage( this, &SDetailCategoryTableRow::GetBackgroundImage )
-		.BorderBackgroundColor( this, &SDetailCategoryTableRow::GetBackgroundColor )
-		.Padding( FMargin(0, VerticalPadding, SDetailTableRowBase::ScrollbarPaddingSize, VerticalPadding) )
+		.BorderImage( FAppStyle::Get().GetBrush( "DetailsView.GridLine") )
+		.Padding( FMargin(0, 0, 0, 1) )
 		[
-			Widget.ToSharedRef()
+			SNew( SBorder )
+			.BorderImage( this, &SDetailCategoryTableRow::GetBackgroundImage )
+			.BorderBackgroundColor( this, &SDetailCategoryTableRow::GetBackgroundColor )
+			.Padding( FMargin(0, VerticalPadding, SDetailTableRowBase::ScrollbarPaddingSize, VerticalPadding) )
+			[
+				Widget.ToSharedRef()
+			]
 		]
 	];
 
