Consolidating icons.
Moved property editing icons to the generic icon set.

diff --git a/Engine/Plugins/AI/EnvironmentQueryEditor/Source/EnvironmentQueryEditor/Private/SEnvQueryProfiler.cpp b/Engine/Plugins/AI/EnvironmentQueryEditor/Source/EnvironmentQueryEditor/Private/SEnvQueryProfiler.cpp
--- a/Engine/Plugins/AI/EnvironmentQueryEditor/Source/EnvironmentQueryEditor/Private/SEnvQueryProfiler.cpp
+++ b/Engine/Plugins/AI/EnvironmentQueryEditor/Source/EnvironmentQueryEditor/Private/SEnvQueryProfiler.cpp
@@ -292,7 +292,7 @@ TSharedRef<SWidget> SEnvQueryProfilerTableRow::GenerateWidgetForColumn(const FNa
 				.IsFocusable(false)
 				[
 					SNew(SImage)
-					.Image(FEditorStyle::GetBrush("PropertyWindow.Button_Browse"))
+					.Image(FEditorStyle::GetBrush("Icons.Search"))
 					.ColorAndOpacity(FSlateColor::UseForeground())
 				]
 			]
diff --git a/Engine/Plugins/Animation/LiveLink/Source/LiveLinkEditor/Private/LiveLinkClientPanel.cpp b/Engine/Plugins/Animation/LiveLink/Source/LiveLinkEditor/Private/LiveLinkClientPanel.cpp
--- a/Engine/Plugins/Animation/LiveLink/Source/LiveLinkEditor/Private/LiveLinkClientPanel.cpp
+++ b/Engine/Plugins/Animation/LiveLink/Source/LiveLinkEditor/Private/LiveLinkClientPanel.cpp
@@ -222,7 +222,7 @@ public:
 					.IsFocusable(false)
 					[
 						SNew(SImage)
-						.Image(FEditorStyle::GetBrush("PropertyWindow.Button_EmptyArray"))
+						.Image(FEditorStyle::GetBrush("Icons.Delete"))
 						.ColorAndOpacity(FSlateColor::UseForeground())
 					];
 			}
@@ -325,7 +325,7 @@ public:
 					.IsFocusable(false)
 					[
 						SNew(SImage)
-						.Image(FEditorStyle::GetBrush("PropertyWindow.Button_EmptyArray"))
+						.Image(FEditorStyle::GetBrush("Icons.Delete"))
 						.ColorAndOpacity(FSlateColor::UseForeground())
 					];
 		}
diff --git a/Engine/Plugins/Editor/AssetManagerEditor/Source/AssetManagerEditor/Private/PrimaryAssetIdCustomization.cpp b/Engine/Plugins/Editor/AssetManagerEditor/Source/AssetManagerEditor/Private/PrimaryAssetIdCustomization.cpp
--- a/Engine/Plugins/Editor/AssetManagerEditor/Source/AssetManagerEditor/Private/PrimaryAssetIdCustomization.cpp
+++ b/Engine/Plugins/Editor/AssetManagerEditor/Source/AssetManagerEditor/Private/PrimaryAssetIdCustomization.cpp
@@ -236,7 +236,7 @@ TSharedRef<SWidget>	SPrimaryAssetIdGraphPin::GetDefaultValueWidget()
 			[
 				SNew(SImage)
 				.ColorAndOpacity(this, &SPrimaryAssetIdGraphPin::OnGetWidgetForeground)
-				.Image(FEditorStyle::GetBrush(TEXT("PropertyWindow.Button_Use")))
+				.Image(FEditorStyle::GetBrush(TEXT("Icons.CircleArrowLeft")))
 			]
 		]
 		+ SHorizontalBox::Slot()
@@ -253,7 +253,7 @@ TSharedRef<SWidget>	SPrimaryAssetIdGraphPin::GetDefaultValueWidget()
 			[
 				SNew(SImage)
 				.ColorAndOpacity(this, &SPrimaryAssetIdGraphPin::OnGetWidgetForeground)
-				.Image(FEditorStyle::GetBrush(TEXT("PropertyWindow.Button_Browse")))
+				.Image(FEditorStyle::GetBrush(TEXT("Icons.Search")))
 			]
 		];
 }
diff --git a/Engine/Plugins/Editor/GameplayTagsEditor/Source/GameplayTagsEditor/Private/SGameplayTagWidget.cpp b/Engine/Plugins/Editor/GameplayTagsEditor/Source/GameplayTagsEditor/Private/SGameplayTagWidget.cpp
--- a/Engine/Plugins/Editor/GameplayTagsEditor/Source/GameplayTagsEditor/Private/SGameplayTagWidget.cpp
+++ b/Engine/Plugins/Editor/GameplayTagsEditor/Source/GameplayTagsEditor/Private/SGameplayTagWidget.cpp
@@ -513,7 +513,7 @@ TSharedRef<ITableRow> SGameplayTagWidget::OnGenerateRow(TSharedPtr<FGameplayTagN
 				.IsFocusable( false )
 				[
 					SNew( SImage )
-					.Image(FEditorStyle::GetBrush("PropertyWindow.Button_AddToArray"))
+					.Image(FEditorStyle::GetBrush("Icons.PlusCircle"))
 					.ColorAndOpacity(FSlateColor::UseForeground())
 				]
 			]
diff --git a/Engine/Plugins/Editor/GameplayTagsEditor/Source/GameplayTagsEditor/Private/SGraphNode_MultiCompareGameplayTag.cpp b/Engine/Plugins/Editor/GameplayTagsEditor/Source/GameplayTagsEditor/Private/SGraphNode_MultiCompareGameplayTag.cpp
--- a/Engine/Plugins/Editor/GameplayTagsEditor/Source/GameplayTagsEditor/Private/SGraphNode_MultiCompareGameplayTag.cpp
+++ b/Engine/Plugins/Editor/GameplayTagsEditor/Source/GameplayTagsEditor/Private/SGraphNode_MultiCompareGameplayTag.cpp
@@ -39,7 +39,7 @@ void SGraphNode_MultiCompareGameplayTag::CreateOutputSideRemoveButton(TSharedPtr
 		.Padding(7, 0, 0, 0)
 		[
 			SNew(SImage)
-			.Image(FEditorStyle::GetBrush(TEXT("PropertyWindow.Button_RemoveFromArray")))
+			.Image(FEditorStyle::GetBrush(TEXT("Icons.Minus")))
 		];
 
 	TSharedPtr<SToolTip> Tooltip;
diff --git a/Engine/Plugins/Experimental/ControlRig/Source/ControlRigEditor/Private/Graph/SControlRigGraphNode.cpp b/Engine/Plugins/Experimental/ControlRig/Source/ControlRigEditor/Private/Graph/SControlRigGraphNode.cpp
--- a/Engine/Plugins/Experimental/ControlRig/Source/ControlRigEditor/Private/Graph/SControlRigGraphNode.cpp
+++ b/Engine/Plugins/Experimental/ControlRig/Source/ControlRigEditor/Private/Graph/SControlRigGraphNode.cpp
@@ -582,7 +582,7 @@ TSharedRef<SWidget> SControlRigGraphNode::AddContainerPinContent(URigVMPin* InIt
 		.Padding(7,0,0,0)
 		[
 			SNew(SImage)
-			.Image(FEditorStyle::GetBrush(TEXT("PropertyWindow.Button_AddToArray")))
+			.Image(FEditorStyle::GetBrush(TEXT("Icons.PlusCircle")))
 		]
 	];
 }
diff --git a/Engine/Plugins/Experimental/HairStrands/Source/HairStrandsEditor/Private/GroomMaterialDetails.cpp b/Engine/Plugins/Experimental/HairStrands/Source/HairStrandsEditor/Private/GroomMaterialDetails.cpp
--- a/Engine/Plugins/Experimental/HairStrands/Source/HairStrandsEditor/Private/GroomMaterialDetails.cpp
+++ b/Engine/Plugins/Experimental/HairStrands/Source/HairStrandsEditor/Private/GroomMaterialDetails.cpp
@@ -153,7 +153,7 @@ void FGroomMaterialDetails::AddMaterials(IDetailLayoutBuilder& DetailLayout)
 			.IsFocusable(false)
 			[
 				SNew(SImage)
-				.Image(FEditorStyle::GetBrush("PropertyWindow.Button_AddToArray"))
+				.Image(FEditorStyle::GetBrush("Icons.PlusCircle"))
 			.ColorAndOpacity(FSlateColor::UseForeground())
 			]
 			]
diff --git a/Engine/Plugins/FX/Niagara/Source/NiagaraEditor/Private/Customizations/NiagaraComponentRendererPropertiesDetails.cpp b/Engine/Plugins/FX/Niagara/Source/NiagaraEditor/Private/Customizations/NiagaraComponentRendererPropertiesDetails.cpp
--- a/Engine/Plugins/FX/Niagara/Source/NiagaraEditor/Private/Customizations/NiagaraComponentRendererPropertiesDetails.cpp
+++ b/Engine/Plugins/FX/Niagara/Source/NiagaraEditor/Private/Customizations/NiagaraComponentRendererPropertiesDetails.cpp
@@ -200,7 +200,7 @@ void FNiagaraComponentRendererPropertiesDetails::CustomizeDetails(IDetailLayoutB
 						.Content()
 						[
 							SNew(SImage)
-							.Image(FEditorStyle::GetBrush("PropertyWindow.Button_EmptyArray"))
+							.Image(FEditorStyle::GetBrush("Icons.Delete"))
 						]
 					]
 				];
diff --git a/Engine/Plugins/FX/Niagara/Source/NiagaraEditor/Private/Widgets/SNiagaraGraphNodeConvert.cpp b/Engine/Plugins/FX/Niagara/Source/NiagaraEditor/Private/Widgets/SNiagaraGraphNodeConvert.cpp
--- a/Engine/Plugins/FX/Niagara/Source/NiagaraEditor/Private/Widgets/SNiagaraGraphNodeConvert.cpp
+++ b/Engine/Plugins/FX/Niagara/Source/NiagaraEditor/Private/Widgets/SNiagaraGraphNodeConvert.cpp
@@ -43,7 +43,7 @@ TSharedRef<SWidget> SNiagaraGraphNodeConvert::CreateTitleWidget(TSharedPtr<SNode
 			.Content()
 			[
 				SNew(SImage)
-				.Image(FEditorStyle::GetBrush("PropertyWindow.Button_Edit"))
+				.Image(FEditorStyle::GetBrush("Icons.Edit"))
 			]
 		];
 }
diff --git a/Engine/Plugins/FX/Niagara/Source/NiagaraEditorWidgets/Private/Stack/SNiagaraStackModuleItem.cpp b/Engine/Plugins/FX/Niagara/Source/NiagaraEditorWidgets/Private/Stack/SNiagaraStackModuleItem.cpp
--- a/Engine/Plugins/FX/Niagara/Source/NiagaraEditorWidgets/Private/Stack/SNiagaraStackModuleItem.cpp
+++ b/Engine/Plugins/FX/Niagara/Source/NiagaraEditorWidgets/Private/Stack/SNiagaraStackModuleItem.cpp
@@ -123,7 +123,7 @@ void SNiagaraStackModuleItem::AddCustomRowWidgets(TSharedRef<SHorizontalBox> Hor
 		.ButtonContent()
 		[
 			SNew(SImage)
-			.Image(FEditorStyle::Get().GetBrush("PropertyWindow.Button_AddToArray"))
+			.Image(FEditorStyle::Get().GetBrush("Icons.PlusCircle"))
 		]
 	];
 
diff --git a/Engine/Plugins/Media/MediaPlayerEditor/Source/MediaPlayerEditor/Private/Widgets/SMediaPlayerEditorPlaylist.cpp b/Engine/Plugins/Media/MediaPlayerEditor/Source/MediaPlayerEditor/Private/Widgets/SMediaPlayerEditorPlaylist.cpp
--- a/Engine/Plugins/Media/MediaPlayerEditor/Source/MediaPlayerEditor/Private/Widgets/SMediaPlayerEditorPlaylist.cpp
+++ b/Engine/Plugins/Media/MediaPlayerEditor/Source/MediaPlayerEditor/Private/Widgets/SMediaPlayerEditorPlaylist.cpp
@@ -102,7 +102,7 @@ void SMediaPlayerEditorPlaylist::Construct(const FArguments& InArgs, UMediaPlaye
 							[
 								SNew(SImage)
 									.ColorAndOpacity(FSlateColor::UseForeground())
-									.Image(FEditorStyle::GetBrush("PropertyWindow.Button_Browse"))
+									.Image(FEditorStyle::GetBrush("Icons.Search"))
 							]
 					]
 
diff --git a/Engine/Plugins/Runtime/PropertyAccess/Source/PropertyAccessEditor/Private/SPropertyBinding.cpp b/Engine/Plugins/Runtime/PropertyAccess/Source/PropertyAccessEditor/Private/SPropertyBinding.cpp
--- a/Engine/Plugins/Runtime/PropertyAccess/Source/PropertyAccessEditor/Private/SPropertyBinding.cpp
+++ b/Engine/Plugins/Runtime/PropertyAccess/Source/PropertyAccessEditor/Private/SPropertyBinding.cpp
@@ -93,7 +93,7 @@ void SPropertyBinding::Construct(const FArguments& InArgs, UBlueprint* InBluepri
 			.ToolTipText(LOCTEXT("GotoFunction", "Goto Function"))
 			[
 				SNew(SImage)
-				.Image(FEditorStyle::GetBrush("PropertyWindow.Button_Browse"))
+				.Image(FEditorStyle::GetBrush("Icons.Search"))
 			]
 		]
 	];
diff --git a/Engine/Plugins/VirtualProduction/DMX/DMXEngine/Source/DMXEditor/Private/Widgets/SAssetPickerButton.cpp b/Engine/Plugins/VirtualProduction/DMX/DMXEngine/Source/DMXEditor/Private/Widgets/SAssetPickerButton.cpp
--- a/Engine/Plugins/VirtualProduction/DMX/DMXEngine/Source/DMXEditor/Private/Widgets/SAssetPickerButton.cpp
+++ b/Engine/Plugins/VirtualProduction/DMX/DMXEngine/Source/DMXEditor/Private/Widgets/SAssetPickerButton.cpp
@@ -83,7 +83,7 @@ void SAssetPickerButton::Construct(const FArguments& InArgs)
 			[
 				SNew(SImage)
 				.ColorAndOpacity( this, &SAssetPickerButton::OnGetWidgetForeground )
-				.Image( FEditorStyle::GetBrush(TEXT("PropertyWindow.Button_Use")) )
+				.Image( FEditorStyle::GetBrush(TEXT("Icons.CircleArrowLeft")) )
 			]
 		]
 		// Browse button
@@ -101,7 +101,7 @@ void SAssetPickerButton::Construct(const FArguments& InArgs)
 			[
 				SNew(SImage)
 				.ColorAndOpacity( this, &SAssetPickerButton::OnGetWidgetForeground )
-				.Image( FEditorStyle::GetBrush(TEXT("PropertyWindow.Button_Browse")) )
+				.Image( FEditorStyle::GetBrush(TEXT("Icons.Search")) )
 			]
 		]
 	];
diff --git a/Engine/Source/Developer/Lumin/LuminPlatformEditor/Private/LuminLocalizedAppNameListWidget.cpp b/Engine/Source/Developer/Lumin/LuminPlatformEditor/Private/LuminLocalizedAppNameListWidget.cpp
--- a/Engine/Source/Developer/Lumin/LuminPlatformEditor/Private/LuminLocalizedAppNameListWidget.cpp
+++ b/Engine/Source/Developer/Lumin/LuminPlatformEditor/Private/LuminLocalizedAppNameListWidget.cpp
@@ -40,7 +40,7 @@ void SLuminLocalizedAppNameListWidget::Construct(const FArguments& Args)
 				.OnClicked(this, &SLuminLocalizedAppNameListWidget::Add)
 				[
 					SNew(SImage)
-					.Image(FEditorStyle::GetBrush("PropertyWindow.Button_AddToArray"))
+					.Image(FEditorStyle::GetBrush("Icons.PlusCircle"))
 					.ColorAndOpacity(FSlateColor::UseForeground())
 				]
 			]
@@ -57,7 +57,7 @@ void SLuminLocalizedAppNameListWidget::Construct(const FArguments& Args)
 				.OnClicked(this, &SLuminLocalizedAppNameListWidget::RemoveAll)
 				[
 					SNew(SImage)
-					.Image(FEditorStyle::GetBrush("PropertyWindow.Button_EmptyArray"))
+					.Image(FEditorStyle::GetBrush("Icons.Delete"))
 					.ColorAndOpacity(FSlateColor::UseForeground())
 				]
 			]
diff --git a/Engine/Source/Developer/Lumin/LuminPlatformEditor/Private/LuminLocalizedAppNameWidget.cpp b/Engine/Source/Developer/Lumin/LuminPlatformEditor/Private/LuminLocalizedAppNameWidget.cpp
--- a/Engine/Source/Developer/Lumin/LuminPlatformEditor/Private/LuminLocalizedAppNameWidget.cpp
+++ b/Engine/Source/Developer/Lumin/LuminPlatformEditor/Private/LuminLocalizedAppNameWidget.cpp
@@ -70,7 +70,7 @@ void SLuminLocalizedAppNameWidget::Construct(const FArguments& Args)
 					.OnClicked(this, &SLuminLocalizedAppNameWidget::OnRemove)
 					[
 						SNew(SImage)
-						.Image(FEditorStyle::GetBrush("PropertyWindow.Button_EmptyArray"))
+						.Image(FEditorStyle::GetBrush("Icons.Delete"))
 						.ColorAndOpacity(FSlateColor::UseForeground())
 					]
 				]
diff --git a/Engine/Source/Developer/Lumin/LuminPlatformEditor/Private/LuminLocalizedIconListWidget.cpp b/Engine/Source/Developer/Lumin/LuminPlatformEditor/Private/LuminLocalizedIconListWidget.cpp
--- a/Engine/Source/Developer/Lumin/LuminPlatformEditor/Private/LuminLocalizedIconListWidget.cpp
+++ b/Engine/Source/Developer/Lumin/LuminPlatformEditor/Private/LuminLocalizedIconListWidget.cpp
@@ -37,7 +37,7 @@ void SLuminLocalizedIconListWidget::Construct(const FArguments& Args)
 				.OnClicked(this, &SLuminLocalizedIconListWidget::Add)
 				[
 					SNew(SImage)
-					.Image(FEditorStyle::GetBrush("PropertyWindow.Button_AddToArray"))
+					.Image(FEditorStyle::GetBrush("Icons.PlusCircle"))
 					.ColorAndOpacity(FSlateColor::UseForeground())
 				]
 			]
@@ -54,7 +54,7 @@ void SLuminLocalizedIconListWidget::Construct(const FArguments& Args)
 				.OnClicked(this, &SLuminLocalizedIconListWidget::RemoveAll)
 				[
 					SNew(SImage)
-					.Image(FEditorStyle::GetBrush("PropertyWindow.Button_EmptyArray"))
+					.Image(FEditorStyle::GetBrush("Icons.Delete"))
 					.ColorAndOpacity(FSlateColor::UseForeground())
 				]
 			]
diff --git a/Engine/Source/Developer/Lumin/LuminPlatformEditor/Private/LuminLocalizedIconWidget.cpp b/Engine/Source/Developer/Lumin/LuminPlatformEditor/Private/LuminLocalizedIconWidget.cpp
--- a/Engine/Source/Developer/Lumin/LuminPlatformEditor/Private/LuminLocalizedIconWidget.cpp
+++ b/Engine/Source/Developer/Lumin/LuminPlatformEditor/Private/LuminLocalizedIconWidget.cpp
@@ -90,7 +90,7 @@ void SLuminLocalizedIconWidget::Construct(const FArguments& Args)
 					.OnClicked(this, &SLuminLocalizedIconWidget::OnRemove)
 					[
 						SNew(SImage)
-						.Image(FEditorStyle::GetBrush("PropertyWindow.Button_EmptyArray"))
+						.Image(FEditorStyle::GetBrush("Icons.Delete"))
 						.ColorAndOpacity(FSlateColor::UseForeground())
 					]
 				]
diff --git a/Engine/Source/Developer/Merge/Private/SMergeAssetPickerView.cpp b/Engine/Source/Developer/Merge/Private/SMergeAssetPickerView.cpp
--- a/Engine/Source/Developer/Merge/Private/SMergeAssetPickerView.cpp
+++ b/Engine/Source/Developer/Merge/Private/SMergeAssetPickerView.cpp
@@ -153,7 +153,7 @@ TSharedRef<SWidget> SMergeAssetPickerPanel::MakeAssetPicker()
 		.ToolTipText(NSLOCTEXT("GraphEditor", "ObjectGraphPin_Browse", "Show the selected asset in the content browser."))
 	[
 		SNew(SImage)
-			.Image(FEditorStyle::GetBrush(TEXT("PropertyWindow.Button_Browse")))
+			.Image(FEditorStyle::GetBrush(TEXT("Icons.Search")))
 	];
 	BrowseButton->SetEnabled(TAttribute<bool>(this, &SMergeAssetPickerPanel::IsBrowseButtonEnabled));
 
@@ -189,7 +189,7 @@ TSharedRef<SWidget> SMergeAssetPickerPanel::MakeAssetPicker()
 			.ToolTipText(NSLOCTEXT("GraphEditor", "ObjectGraphPin_Use", "Use content browser selection."))
 		[
 			SNew(SImage)
-				.Image(FEditorStyle::GetBrush(TEXT("PropertyWindow.Button_Use")))
+				.Image(FEditorStyle::GetBrush(TEXT("Icons.CircleArrowLeft")))
 		]
 	]
 	+ SHorizontalBox::Slot()
diff --git a/Engine/Source/Developer/MessageLog/Private/UserInterface/SMessageLogMessageListRow.cpp b/Engine/Source/Developer/MessageLog/Private/UserInterface/SMessageLogMessageListRow.cpp
--- a/Engine/Source/Developer/MessageLog/Private/UserInterface/SMessageLogMessageListRow.cpp
+++ b/Engine/Source/Developer/MessageLog/Private/UserInterface/SMessageLogMessageListRow.cpp
@@ -164,7 +164,7 @@ void SMessageLogMessageListRow::CreateMessage(const TSharedRef<SHorizontalBox>&
 	{
 		const TSharedRef<FUObjectToken> UObjectToken = StaticCastSharedRef<FUObjectToken>(InMessageToken);
 
-		IconBrushName = FName("PropertyWindow.Button_Browse");
+		IconBrushName = FName("Icons.Search");
 
 		UObject* Object = nullptr;
 
@@ -203,7 +203,7 @@ void SMessageLogMessageListRow::CreateMessage(const TSharedRef<SHorizontalBox>&
 		break;
 	case EMessageToken::EdGraph:
 	{
-		IconBrushName = FName("PropertyWindow.Button_Browse");
+		IconBrushName = FName("Icons.Search");
 		RowContent = CreateHyperlink(InMessageToken, InMessageToken->ToText());
 	}
 	break;
@@ -226,7 +226,7 @@ void SMessageLogMessageListRow::CreateMessage(const TSharedRef<SHorizontalBox>&
 	{
 		const TSharedRef<FAssetNameToken> AssetNameToken = StaticCastSharedRef<FAssetNameToken>(InMessageToken);
 
-		IconBrushName = FName("PropertyWindow.Button_Browse");
+		IconBrushName = FName("Icons.Search");
 		RowContent = CreateHyperlink(InMessageToken, AssetNameToken->ToText());
 	}
 		break;
diff --git a/Engine/Source/Editor/BlueprintGraph/Private/K2Node_MacroInstance.cpp b/Engine/Source/Editor/BlueprintGraph/Private/K2Node_MacroInstance.cpp
--- a/Engine/Source/Editor/BlueprintGraph/Private/K2Node_MacroInstance.cpp
+++ b/Engine/Source/Editor/BlueprintGraph/Private/K2Node_MacroInstance.cpp
@@ -208,7 +208,7 @@ void UK2Node_MacroInstance::GetNodeContextMenuActions(UToolMenu* Menu, UGraphNod
 				"MacroInstanceFindInContentBrowser",
 				NSLOCTEXT("K2Node", "MacroInstanceFindInContentBrowser", "Find in Content Browser"),
 				NSLOCTEXT("K2Node", "MacroInstanceFindInContentBrowserTooltip", "Finds the Blueprint Macro Library that contains this Macro in the Content Browser"),
-				FSlateIcon(FEditorStyle::GetStyleSetName(), "PropertyWindow.Button_Browse"),
+				FSlateIcon(FEditorStyle::GetStyleSetName(), "Icons.Search"),
 				FUIAction( FExecuteAction::CreateStatic( &UK2Node_MacroInstance::FindInContentBrowser, MakeWeakObjectPtr(const_cast<UK2Node_MacroInstance*>(this)) ) )
 				);
 		}
diff --git a/Engine/Source/Editor/DataTableEditor/Private/DataTableEditor.cpp b/Engine/Source/Editor/DataTableEditor/Private/DataTableEditor.cpp
--- a/Engine/Source/Editor/DataTableEditor/Private/DataTableEditor.cpp
+++ b/Engine/Source/Editor/DataTableEditor/Private/DataTableEditor.cpp
@@ -974,7 +974,7 @@ void FDataTableEditor::PostRegenerateMenusAndToolbars()
 				.ForegroundColor(FSlateColor::UseForeground())
 				[
 					SNew(SImage)
-					.Image(FEditorStyle::GetBrush("PropertyWindow.Button_Browse"))
+					.Image(FEditorStyle::GetBrush("Icons.Search"))
 				]
 			]
 			+SHorizontalBox::Slot()
diff --git a/Engine/Source/Editor/DetailCustomizations/Private/DeviceProfileDetails.cpp b/Engine/Source/Editor/DetailCustomizations/Private/DeviceProfileDetails.cpp
--- a/Engine/Source/Editor/DetailCustomizations/Private/DeviceProfileDetails.cpp
+++ b/Engine/Source/Editor/DetailCustomizations/Private/DeviceProfileDetails.cpp
@@ -649,7 +649,7 @@ void FDeviceProfileConsoleVariablesPropertyDetails::CreateConsoleVariablesProper
 					.ButtonContent()
 					[
 						SNew(SImage)
-						.Image(FEditorStyle::GetBrush("PropertyWindow.Button_AddToArray"))
+						.Image(FEditorStyle::GetBrush("Icons.PlusCircle"))
 					]
 					.MenuContent()
 					[
@@ -669,7 +669,7 @@ void FDeviceProfileConsoleVariablesPropertyDetails::CreateConsoleVariablesProper
 					.IsFocusable(false)
 					[
 						SNew(SImage)
-						.Image(FEditorStyle::GetBrush("PropertyWindow.Button_EmptyArray"))
+						.Image(FEditorStyle::GetBrush("Icons.Delete"))
 					]
 				]
 			];
diff --git a/Engine/Source/Editor/DeviceProfileEditor/Private/DeviceProfileConsoleVariableColumn.cpp b/Engine/Source/Editor/DeviceProfileEditor/Private/DeviceProfileConsoleVariableColumn.cpp
--- a/Engine/Source/Editor/DeviceProfileEditor/Private/DeviceProfileConsoleVariableColumn.cpp
+++ b/Engine/Source/Editor/DeviceProfileEditor/Private/DeviceProfileConsoleVariableColumn.cpp
@@ -115,7 +115,7 @@ TSharedRef<class SWidget> FConsoleVariableCellPresenter::ConstructDisplayWidget(
 		.IsFocusable(false)
 		[
 			SNew(SImage)
-			.Image(FEditorStyle::GetBrush("PropertyWindow.Button_Edit"))
+			.Image(FEditorStyle::GetBrush("Icons.Edit"))
 			.ColorAndOpacity(FSlateColor::UseForeground())
 		]
 	];
diff --git a/Engine/Source/Editor/DeviceProfileEditor/Private/DeviceProfileTextureLODSettingsColumn.cpp b/Engine/Source/Editor/DeviceProfileEditor/Private/DeviceProfileTextureLODSettingsColumn.cpp
--- a/Engine/Source/Editor/DeviceProfileEditor/Private/DeviceProfileTextureLODSettingsColumn.cpp
+++ b/Engine/Source/Editor/DeviceProfileEditor/Private/DeviceProfileTextureLODSettingsColumn.cpp
@@ -115,7 +115,7 @@ TSharedRef<class SWidget> FTextureLODSettingsCellPresenter::ConstructDisplayWidg
 		.IsFocusable(false)
 		[
 			SNew(SImage)
-			.Image(FEditorStyle::GetBrush("PropertyWindow.Button_Edit"))
+			.Image(FEditorStyle::GetBrush("Icons.Edit"))
 			.ColorAndOpacity(FSlateColor::UseForeground())
 		]
 	];
diff --git a/Engine/Source/Editor/EditorStyle/Private/StarshipStyle.cpp b/Engine/Source/Editor/EditorStyle/Private/StarshipStyle.cpp
--- a/Engine/Source/Editor/EditorStyle/Private/StarshipStyle.cpp
+++ b/Engine/Source/Editor/EditorStyle/Private/StarshipStyle.cpp
@@ -3106,7 +3106,7 @@ void FStarshipEditorStyle::FStyle::SetupPropertyEditorStyles()
 		Set( "PropertyWindow.FilterFont", DEFAULT_FONT( "Regular", 10 ) );
 
 		FSlateFontInfo MobilityFont = FStyleFonts::Get().SmallBold;
-		MobilityFont.LetterSpacing = 200;
+		MobilityFont.LetterSpacing = 100;
 
 		Set("PropertyWindow.MobilityFont", MobilityFont);
 		Set( "PropertyWindow.NoOverlayColor", new FSlateNoResource() );
@@ -3115,14 +3115,7 @@ void FStarshipEditorStyle::FStyle::SetupPropertyEditorStyles()
 		Set( "PropertyWindow.FilteredEditConstColor", new FSlateColorBrush( FColor( 152, 152, 152, 80 ).ReinterpretAsLinear() * FColor(0,255,0,255).ReinterpretAsLinear() ) );
 		Set( "PropertyWindow.CategoryBackground", new BOX_BRUSH( "/PropertyView/CategoryBackground", FMargin(4.f/16.f) ) );
 		Set( "PropertyWindow.CategoryForeground", FLinearColor::Black );
-		Set( "PropertyWindow.Button_Browse", new CORE_IMAGE_BRUSH_SVG("Starship/Common/search", Icon16x16));
-		Set( "PropertyWindow.Button_Use", new CORE_IMAGE_BRUSH_SVG("Starship/Common/circle-arrow-left", Icon16x16));
-		Set( "PropertyWindow.Button_Delete", new IMAGE_BRUSH("Icons/Cross_12x", Icon12x12 ) );
 		Set( "PropertyWindow.Button_Clear", new IMAGE_BRUSH( "Icons/Cross_12x", Icon12x12 ) );
-		Set( "PropertyWindow.Button_Edit", new IMAGE_BRUSH_SVG("Starship/Common/edit", Icon16x16));
-		Set( "PropertyWindow.Button_EmptyArray", new IMAGE_BRUSH( "Icons/empty_set_12x", Icon12x12 ) );
-		Set( "PropertyWindow.Button_AddToArray", new CORE_IMAGE_BRUSH_SVG("Starship/Common/plus-circle", Icon16x16));
-		Set( "PropertyWindow.Button_RemoveFromArray", new IMAGE_BRUSH_SVG("Starship/Common/minus", Icon16x16));
 		Set( "PropertyWindow.Button_Ellipsis", new IMAGE_BRUSH( "Icons/ellipsis_12x", Icon12x12 ) );
 		Set( "PropertyWindow.Button_PickAsset", new IMAGE_BRUSH( "Icons/pillarray_12x", Icon12x12 ) );
 		Set( "PropertyWindow.Button_PickActor", new IMAGE_BRUSH( "Icons/levels_16x", Icon12x12 ) );
diff --git a/Engine/Source/Editor/GraphEditor/Private/KismetPins/SGraphPinObject.cpp b/Engine/Source/Editor/GraphEditor/Private/KismetPins/SGraphPinObject.cpp
--- a/Engine/Source/Editor/GraphEditor/Private/KismetPins/SGraphPinObject.cpp
+++ b/Engine/Source/Editor/GraphEditor/Private/KismetPins/SGraphPinObject.cpp
@@ -110,7 +110,7 @@ TSharedRef<SWidget>	SGraphPinObject::GetDefaultValueWidget()
 				[
 					SNew(SImage)
 					.ColorAndOpacity( this, &SGraphPinObject::OnGetWidgetForeground )
-					.Image( FEditorStyle::GetBrush(TEXT("PropertyWindow.Button_Use")) )
+					.Image( FEditorStyle::GetBrush(TEXT("Icons.CircleArrowLeft")) )
 				]
 			]
 			// Browse button
@@ -128,7 +128,7 @@ TSharedRef<SWidget>	SGraphPinObject::GetDefaultValueWidget()
 				[
 					SNew(SImage)
 					.ColorAndOpacity( this, &SGraphPinObject::OnGetWidgetForeground )
-					.Image( FEditorStyle::GetBrush(TEXT("PropertyWindow.Button_Browse")) )
+					.Image( FEditorStyle::GetBrush(TEXT("Icons.Search")) )
 				]
 			];
 	}
diff --git a/Engine/Source/Editor/GraphEditor/Private/SGraphNode.cpp b/Engine/Source/Editor/GraphEditor/Private/SGraphNode.cpp
--- a/Engine/Source/Editor/GraphEditor/Private/SGraphNode.cpp
+++ b/Engine/Source/Editor/GraphEditor/Private/SGraphNode.cpp
@@ -1513,7 +1513,7 @@ TSharedRef<SWidget> SGraphNode::AddPinButtonContent(FText PinText, FText PinTool
 		. Padding( 7,0,0,0 )
 		[
 			SNew(SImage)
-			.Image(FEditorStyle::GetBrush(TEXT("PropertyWindow.Button_AddToArray")))
+			.Image(FEditorStyle::GetBrush(TEXT("Icons.PlusCircle")))
 		];
 	}
 	else
@@ -1525,7 +1525,7 @@ TSharedRef<SWidget> SGraphNode::AddPinButtonContent(FText PinText, FText PinTool
 		. Padding( 0,0,7,0 )
 		[
 			SNew(SImage)
-			.Image(FEditorStyle::GetBrush(TEXT("PropertyWindow.Button_AddToArray")))
+			.Image(FEditorStyle::GetBrush(TEXT("Icons.PlusCircle")))
 		]
 		+SHorizontalBox::Slot()
 		.AutoWidth()
diff --git a/Engine/Source/Editor/Kismet/Private/BlueprintEditor.cpp b/Engine/Source/Editor/Kismet/Private/BlueprintEditor.cpp
--- a/Engine/Source/Editor/Kismet/Private/BlueprintEditor.cpp
+++ b/Engine/Source/Editor/Kismet/Private/BlueprintEditor.cpp
@@ -2180,7 +2180,7 @@ void FBlueprintEditor::PostRegenerateMenusAndToolbars()
 				.ForegroundColor( FSlateColor::UseForeground() )
 				[
 					SNew(SImage)
-					.Image(FEditorStyle::GetBrush("PropertyWindow.Button_Browse"))
+					.Image(FEditorStyle::GetBrush("Icons.Search"))
 				]
 			]
 			+SHorizontalBox::Slot()
@@ -2197,7 +2197,7 @@ void FBlueprintEditor::PostRegenerateMenusAndToolbars()
 				.ForegroundColor( FSlateColor::UseForeground() )
 				[
 					SNew(SImage)
-					.Image(FEditorStyle::GetBrush("PropertyWindow.Button_Edit"))
+					.Image(FEditorStyle::GetBrush("Icons.Edit"))
 				]
 			]
 			+ SHorizontalBox::Slot()
diff --git a/Engine/Source/Editor/LevelEditor/Private/LevelEditorToolBar.cpp b/Engine/Source/Editor/LevelEditor/Private/LevelEditorToolBar.cpp
--- a/Engine/Source/Editor/LevelEditor/Private/LevelEditorToolBar.cpp
+++ b/Engine/Source/Editor/LevelEditor/Private/LevelEditorToolBar.cpp
@@ -668,8 +668,8 @@ void LevelEditorActionHelpers::GetBlueprintSettingsSubMenu(UToolMenu* Menu, FBlu
 {
 #define LOCTEXT_NAMESPACE "LevelToolBarViewMenu"
 
-	FSlateIcon EditBPIcon(FEditorStyle::Get().GetStyleSetName(), TEXT("PropertyWindow.Button_Edit"));
-	FSlateIcon NewBPIcon(FEditorStyle::Get().GetStyleSetName(), TEXT("PropertyWindow.Button_AddToArray"));
+	FSlateIcon EditBPIcon(FEditorStyle::Get().GetStyleSetName(), TEXT("Icons.Edit"));
+	FSlateIcon NewBPIcon(FEditorStyle::Get().GetStyleSetName(), TEXT("Icons.PlusCircle"));
 	FText RootClassName = FText::FromString(InSettingsData.RootClass->GetName());
 
 	// If there is currently a valid GameMode Blueprint, offer to edit the Blueprint
diff --git a/Engine/Source/Editor/MovieSceneTools/Private/K2Node_GetSequenceBinding.cpp b/Engine/Source/Editor/MovieSceneTools/Private/K2Node_GetSequenceBinding.cpp
--- a/Engine/Source/Editor/MovieSceneTools/Private/K2Node_GetSequenceBinding.cpp
+++ b/Engine/Source/Editor/MovieSceneTools/Private/K2Node_GetSequenceBinding.cpp
@@ -455,7 +455,7 @@ TSharedPtr<SGraphNode> UK2Node_GetSequenceBinding::CreateVisualWidget()
 					[
 						SNew(SImage)
 						.ColorAndOpacity(this, &SGraphNodeGetSequenceBinding::OnGetWidgetForeground)
-						.Image( FEditorStyle::GetBrush(TEXT("PropertyWindow.Button_Use")) )
+						.Image( FEditorStyle::GetBrush(TEXT("Icons.CircleArrowLeft")) )
 					]
 				]
 
@@ -474,7 +474,7 @@ TSharedPtr<SGraphNode> UK2Node_GetSequenceBinding::CreateVisualWidget()
 					[
 						SNew(SImage)
 						.ColorAndOpacity(this, &SGraphNodeGetSequenceBinding::OnGetWidgetForeground)
-						.Image( FEditorStyle::GetBrush(TEXT("PropertyWindow.Button_Browse")) )
+						.Image( FEditorStyle::GetBrush(TEXT("Icons.Search")) )
 					]
 				]
 			];
diff --git a/Engine/Source/Editor/Persona/Private/AnimTimeline/AnimTimelineTrack_Montage.cpp b/Engine/Source/Editor/Persona/Private/AnimTimeline/AnimTimelineTrack_Montage.cpp
--- a/Engine/Source/Editor/Persona/Private/AnimTimeline/AnimTimelineTrack_Montage.cpp
+++ b/Engine/Source/Editor/Persona/Private/AnimTimeline/AnimTimelineTrack_Montage.cpp
@@ -500,14 +500,14 @@ TSharedRef<SWidget> FAnimTimelineTrack_Montage::BuildMontageSubMenu()
 			MenuBuilder.AddMenuEntry(
 				LOCTEXT("FindParent", "Find parent"),
 				LOCTEXT("FindParentInCBToolTip", "Find parent in Content Browser"),
-				FSlateIcon("EditorStyle", "PropertyWindow.Button_Browse"),
+				FSlateIcon("EditorStyle", "Icons.Search"),
 				FUIAction(FExecuteAction::CreateSP(this, &FAnimTimelineTrack_Montage::OnFindParentClassInContentBrowserClicked))
 			);
 
 			MenuBuilder.AddMenuEntry(
 				LOCTEXT("EditParent", "Edit parent"),
 				LOCTEXT("EditParentToolTip", "Open parent in editor"),
-				FSlateIcon("EditorStyle", "PropertyWindow.Button_Edit"),
+				FSlateIcon("EditorStyle", "Icons.Edit"),
 				FUIAction(FExecuteAction::CreateSP(this, &FAnimTimelineTrack_Montage::OnEditParentClassClicked))
 			);
 		}
diff --git a/Engine/Source/Editor/Persona/Private/PersonaMeshDetails.cpp b/Engine/Source/Editor/Persona/Private/PersonaMeshDetails.cpp
--- a/Engine/Source/Editor/Persona/Private/PersonaMeshDetails.cpp
+++ b/Engine/Source/Editor/Persona/Private/PersonaMeshDetails.cpp
@@ -2555,7 +2555,7 @@ void FPersonaMeshDetails::AddLODLevelCategories(IDetailLayoutBuilder& DetailLayo
 							.IsFocusable(false)
 							[
 								SNew(SImage)
-								.Image(FEditorStyle::GetBrush("PropertyWindow.Button_AddToArray"))
+								.Image(FEditorStyle::GetBrush("Icons.PlusCircle"))
 								.ColorAndOpacity(FSlateColor::UseForeground())
 							]
 						]
@@ -2969,7 +2969,7 @@ void FPersonaMeshDetails::CustomizeSkinWeightProfiles(IDetailLayoutBuilder& Deta
 			.ButtonContent()
 			[
 				SNew(SImage)
-				.Image(FEditorStyle::GetBrush("PropertyWindow.Button_AddToArray"))
+				.Image(FEditorStyle::GetBrush("Icons.PlusCircle"))
 			]
 			.OnGetMenuContent(this, &FPersonaMeshDetails::CreateSkinWeightProfileMenuContent)
 			.ToolTipText(LOCTEXT("ImportSkinWeightButtonToolTip", "Import a new Skin Weight Profile"))
diff --git a/Engine/Source/Editor/Persona/Private/SAssetFamilyShortcutBar.cpp b/Engine/Source/Editor/Persona/Private/SAssetFamilyShortcutBar.cpp
--- a/Engine/Source/Editor/Persona/Private/SAssetFamilyShortcutBar.cpp
+++ b/Engine/Source/Editor/Persona/Private/SAssetFamilyShortcutBar.cpp
@@ -249,7 +249,7 @@ public:
 			MenuBuilder.AddMenuEntry(
 				LOCTEXT("ShowInContentBrowser", "Show In Content Browser"),
 				LOCTEXT("ShowInContentBrowser_ToolTip", "Show this asset in the content browser."),
-				FSlateIcon(FEditorStyle::GetStyleSetName(), "PropertyWindow.Button_Browse"),
+				FSlateIcon(FEditorStyle::GetStyleSetName(), "Icons.Search"),
 				FUIAction(FExecuteAction::CreateSP(this, &SAssetShortcut::HandleShowInContentBrowser)));
 		}
 		MenuBuilder.EndSection();
diff --git a/Engine/Source/Editor/Persona/Private/SSlotNameReferenceWindow.cpp b/Engine/Source/Editor/Persona/Private/SSlotNameReferenceWindow.cpp
--- a/Engine/Source/Editor/Persona/Private/SSlotNameReferenceWindow.cpp
+++ b/Engine/Source/Editor/Persona/Private/SSlotNameReferenceWindow.cpp
@@ -72,7 +72,7 @@ TSharedRef<SWidget> SMontageReferenceListRow::GenerateWidgetForColumn(const FNam
 				.Content()
 				[
 					SNew(SImage)
-					.Image(FEditorStyle::GetBrush("PropertyWindow.Button_Browse"))
+					.Image(FEditorStyle::GetBrush("Icons.Search"))
 				]
 			]
 			+ SHorizontalBox::Slot()
@@ -148,7 +148,7 @@ TSharedRef<SWidget> SBlueprintReferenceListRow::GenerateWidgetForColumn(const FN
 				.Content()
 				[
 					SNew(SImage)
-					.Image(FEditorStyle::GetBrush("PropertyWindow.Button_Browse"))
+					.Image(FEditorStyle::GetBrush("Icons.Search"))
 				]
 			]
 			+ SHorizontalBox::Slot()
diff --git a/Engine/Source/Editor/Persona/Private/SkinWeightProfileCustomization.cpp b/Engine/Source/Editor/Persona/Private/SkinWeightProfileCustomization.cpp
--- a/Engine/Source/Editor/Persona/Private/SkinWeightProfileCustomization.cpp
+++ b/Engine/Source/Editor/Persona/Private/SkinWeightProfileCustomization.cpp
@@ -113,7 +113,7 @@ void FSkinWeightProfileCustomization::CustomizeHeader(TSharedRef<IPropertyHandle
 			.ButtonContent()
 			[
 				SNew(SImage)
-				.Image(FEditorStyle::GetBrush("PropertyWindow.Button_RemoveFromArray"))
+				.Image(FEditorStyle::GetBrush("Icons.Minus"))
 			]
 			.OnGetMenuContent(FOnGetContent::CreateSP(this, &FSkinWeightProfileCustomization::GenerateRemoveMenu))
 		]
diff --git a/Engine/Source/Editor/PropertyEditor/Private/PropertyCustomizationHelpers.cpp b/Engine/Source/Editor/PropertyEditor/Private/PropertyCustomizationHelpers.cpp
--- a/Engine/Source/Editor/PropertyEditor/Private/PropertyCustomizationHelpers.cpp
+++ b/Engine/Source/Editor/PropertyEditor/Private/PropertyCustomizationHelpers.cpp
@@ -109,7 +109,7 @@ namespace PropertyCustomizationHelpers
 			SNew( SPropertyEditorButton )
 			.Text( LOCTEXT( "AddButtonLabel", "Add" ) )
 			.ToolTipText( OptionalToolTipText.Get().IsEmpty() ? LOCTEXT( "AddButtonToolTipText", "Adds Element") : OptionalToolTipText )
-			.Image( FEditorStyle::GetBrush("PropertyWindow.Button_AddToArray") )
+			.Image( FEditorStyle::GetBrush("Icons.PlusCircle") )
 			.OnClickAction( OnAddClicked )
 			.IsEnabled(IsEnabled)
 			.IsFocusable( false );
@@ -121,7 +121,7 @@ namespace PropertyCustomizationHelpers
 			SNew( SPropertyEditorButton )
 			.Text( LOCTEXT( "RemoveButtonLabel", "Remove" ) )
 			.ToolTipText( OptionalToolTipText.Get().IsEmpty() ? LOCTEXT( "RemoveButtonToolTipText", "Removes Element") : OptionalToolTipText )
-			.Image( FEditorStyle::GetBrush("PropertyWindow.Button_RemoveFromArray") )
+			.Image( FEditorStyle::GetBrush("Icons.Minus") )
 			.OnClickAction( OnRemoveClicked )
 			.IsEnabled(IsEnabled)
 			.IsFocusable( false );
@@ -133,7 +133,7 @@ namespace PropertyCustomizationHelpers
 			SNew( SPropertyEditorButton )
 			.Text( LOCTEXT( "EmptyButtonLabel", "Empty" ) )
 			.ToolTipText( OptionalToolTipText.Get().IsEmpty() ? LOCTEXT( "EmptyButtonToolTipText", "Removes All Elements") : OptionalToolTipText )
-			.Image( FEditorStyle::GetBrush("PropertyWindow.Button_EmptyArray") )
+			.Image( FEditorStyle::GetBrush("Icons.Delete") )
 			.OnClickAction( OnEmptyClicked )
 			.IsEnabled(IsEnabled)
 			.IsFocusable( false );
@@ -145,7 +145,7 @@ namespace PropertyCustomizationHelpers
 			SNew( SPropertyEditorButton )
 			.Text( LOCTEXT( "UseButtonLabel", "Use") )
 			.ToolTipText( OptionalToolTipText.Get().IsEmpty() ? LOCTEXT( "UseButtonToolTipText", "Use Selected Asset from Content Browser") : OptionalToolTipText )
-			.Image( FEditorStyle::GetBrush("PropertyWindow.Button_Use") )
+			.Image( FEditorStyle::GetBrush("Icons.CircleArrowLeft") )
 			.OnClickAction( OnUseSelectedClicked )
 			.IsEnabled(IsEnabled)
 			.IsFocusable( false );
@@ -207,7 +207,7 @@ namespace PropertyCustomizationHelpers
 			SNew( SPropertyEditorButton )
 			.Text( LOCTEXT( "BrowseButtonLabel", "Browse") )
 			.ToolTipText( OptionalToolTipText.Get().IsEmpty() ? LOCTEXT( "BrowseButtonToolTipText", "Browse to Asset in Content Browser") : OptionalToolTipText )
-			.Image( FEditorStyle::GetBrush("PropertyWindow.Button_Browse") )
+			.Image( FEditorStyle::GetBrush("Icons.Search") )
 			.OnClickAction( OnFindClicked )
 			.IsEnabled(IsEnabled)
 			.IsFocusable( false );
@@ -250,7 +250,7 @@ namespace PropertyCustomizationHelpers
 
 		return
 			SNew(SComboButton)
-			.ButtonStyle( FEditorStyle::Get(), "HoverHintOnly" )
+			.ComboButtonStyle( FAppStyle::Get(), "SimpleComboButton" )
 			.ContentPadding(2)
 			.ForegroundColor( FSlateColor::UseForeground() )
 			.HasDownArrow(true)
@@ -1390,7 +1390,7 @@ void SMaterialSlotWidget::Construct(const FArguments& InArgs, int32 SlotIndex, b
 		]
 		+ SHorizontalBox::Slot()
 		.FillWidth(1.0f)
-		.Padding(5.0f, 3.0f, 0.0f, 3.0f)
+		.Padding(5.0f, 0.0f, 0.0f,0.0f)
 		[
 			SNew(SBox)
 			.VAlign(VAlign_Center)
diff --git a/Engine/Source/Editor/PropertyEditor/Private/SPerPlatformPropertiesWidget.cpp b/Engine/Source/Editor/PropertyEditor/Private/SPerPlatformPropertiesWidget.cpp
--- a/Engine/Source/Editor/PropertyEditor/Private/SPerPlatformPropertiesWidget.cpp
+++ b/Engine/Source/Editor/PropertyEditor/Private/SPerPlatformPropertiesWidget.cpp
@@ -246,7 +246,7 @@ TSharedRef<SWidget> SPerPlatformPropertiesWidget::MakePerPlatformWidget(FName In
 				.ButtonContent()
 				[
 					SNew(SImage)
-					.Image(FEditorStyle::GetBrush("PropertyWindow.Button_AddToArray"))
+					.Image(FEditorStyle::GetBrush("Icons.PlusCircle"))
 				]
 				.MenuContent()
 				[
diff --git a/Engine/Source/Editor/StaticMeshEditor/Private/StaticMeshEditorTools.cpp b/Engine/Source/Editor/StaticMeshEditor/Private/StaticMeshEditorTools.cpp
--- a/Engine/Source/Editor/StaticMeshEditor/Private/StaticMeshEditorTools.cpp
+++ b/Engine/Source/Editor/StaticMeshEditor/Private/StaticMeshEditorTools.cpp
@@ -2489,7 +2489,7 @@ void FMeshMaterialsLayout::AddToCategory(IDetailCategoryBuilder& CategoryBuilder
 					.IsFocusable(false)
 					[
 						SNew(SImage)
-						.Image(FEditorStyle::GetBrush("PropertyWindow.Button_AddToArray"))
+						.Image(FEditorStyle::GetBrush("Icons.PlusCircle"))
 						.ColorAndOpacity(FSlateColor::UseForeground())
 					]
 				]
diff --git a/Engine/Source/Editor/UMGEditor/Private/Details/SFunctionSelector.cpp b/Engine/Source/Editor/UMGEditor/Private/Details/SFunctionSelector.cpp
--- a/Engine/Source/Editor/UMGEditor/Private/Details/SFunctionSelector.cpp
+++ b/Engine/Source/Editor/UMGEditor/Private/Details/SFunctionSelector.cpp
@@ -81,7 +81,7 @@ void SFunctionSelector::Construct(const FArguments& InArgs, TSharedRef<FWidgetBl
 			.ToolTipText(LOCTEXT("GotoFunction", "Goto Function"))
 			[
 				SNew(SImage)
-				.Image(FEditorStyle::GetBrush("PropertyWindow.Button_Browse"))
+				.Image(FEditorStyle::GetBrush("Icons.Search"))
 			]
 		]
 	];
diff --git a/Engine/Source/Editor/UndoHistory/Private/Widgets/SUndoHistory.cpp b/Engine/Source/Editor/UndoHistory/Private/Widgets/SUndoHistory.cpp
--- a/Engine/Source/Editor/UndoHistory/Private/Widgets/SUndoHistory.cpp
+++ b/Engine/Source/Editor/UndoHistory/Private/Widgets/SUndoHistory.cpp
@@ -180,7 +180,7 @@ void SUndoHistory::Construct( const FArguments& InArgs )
 							.ToolTipText(LOCTEXT("DiscardHistoryButtonToolTip", "Discard the Undo History."))
 							[
 								SNew(SImage)
-									.Image(FEditorStyle::Get().GetBrush("TrashCan_Small"))
+								.Image(FEditorStyle::Get().GetBrush("Icons.Delete"))
 							]
 					]
 
diff --git a/Engine/Source/Editor/UnrealEd/Private/Fbx/SSceneImportNodeTreeView.cpp b/Engine/Source/Editor/UnrealEd/Private/Fbx/SSceneImportNodeTreeView.cpp
--- a/Engine/Source/Editor/UnrealEd/Private/Fbx/SSceneImportNodeTreeView.cpp
+++ b/Engine/Source/Editor/UnrealEd/Private/Fbx/SSceneImportNodeTreeView.cpp
@@ -333,9 +333,9 @@ TSharedPtr<SWidget> SFbxSceneTreeView::OnOpenContextMenu()
 	// We always create a section here, even if there is no parent so that clients can still extend the menu
 	MenuBuilder.BeginSection("FbxSceneTreeViewContextMenuImportSection");
 	{
-		const FSlateIcon PlusIcon(FEditorStyle::GetStyleSetName(), "Plus");
+		const FSlateIcon PlusIcon(FEditorStyle::GetStyleSetName(), "Icons.Plus");
 		MenuBuilder.AddMenuEntry(LOCTEXT("CheckForImport", "Add Selection To Import"), FText(), PlusIcon, FUIAction(FExecuteAction::CreateSP(this, &SFbxSceneTreeView::AddSelectionToImport)));
-		const FSlateIcon MinusIcon(FEditorStyle::GetStyleSetName(), "PropertyWindow.Button_RemoveFromArray");
+		const FSlateIcon MinusIcon(FEditorStyle::GetStyleSetName(), "Icons.Minus");
 		MenuBuilder.AddMenuEntry(LOCTEXT("UncheckForImport", "Remove Selection From Import"), FText(), MinusIcon, FUIAction(FExecuteAction::CreateSP(this, &SFbxSceneTreeView::RemoveSelectionFromImport)));
 	}
 	MenuBuilder.EndSection();
diff --git a/Engine/Source/Editor/UnrealEd/Private/Fbx/SSceneImportStaticMeshListView.cpp b/Engine/Source/Editor/UnrealEd/Private/Fbx/SSceneImportStaticMeshListView.cpp
--- a/Engine/Source/Editor/UnrealEd/Private/Fbx/SSceneImportStaticMeshListView.cpp
+++ b/Engine/Source/Editor/UnrealEd/Private/Fbx/SSceneImportStaticMeshListView.cpp
@@ -301,7 +301,7 @@ TSharedPtr<SWidget> SFbxSceneStaticMeshListView::OnOpenContextMenu()
 
 	const FSlateIcon PlusIcon(FEditorStyle::GetStyleSetName(), "Plus");
 	MenuBuilder.AddMenuEntry(LOCTEXT("CheckForImport", "Add Selection To Import"), FText(), PlusIcon, FUIAction(FExecuteAction::CreateSP(this, &SFbxSceneStaticMeshListView::AddSelectionToImport)));
-	const FSlateIcon MinusIcon(FEditorStyle::GetStyleSetName(), "PropertyWindow.Button_RemoveFromArray");
+	const FSlateIcon MinusIcon(FEditorStyle::GetStyleSetName(), "Icons.Minus");
 	MenuBuilder.AddMenuEntry(LOCTEXT("UncheckForImport", "Remove Selection From Import"), FText(), MinusIcon, FUIAction(FExecuteAction::CreateSP(this, &SFbxSceneStaticMeshListView::RemoveSelectionFromImport)));
 
 	AddBakePivotMenu(MenuBuilder);
diff --git a/Engine/Source/Editor/UnrealEd/Private/Fbx/SSceneMaterialsListView.cpp b/Engine/Source/Editor/UnrealEd/Private/Fbx/SSceneMaterialsListView.cpp
--- a/Engine/Source/Editor/UnrealEd/Private/Fbx/SSceneMaterialsListView.cpp
+++ b/Engine/Source/Editor/UnrealEd/Private/Fbx/SSceneMaterialsListView.cpp
@@ -319,7 +319,7 @@ TSharedPtr<SWidget> SFbxSceneMaterialsListView::OnOpenContextMenu()
 	{
 		const FSlateIcon PlusIcon(FEditorStyle::GetStyleSetName(), "Plus");
 		MenuBuilder.AddMenuEntry(LOCTEXT("CheckForImport", "Add Selection To Import"), FText(), PlusIcon, FUIAction(FExecuteAction::CreateSP(this, &SFbxSceneMaterialsListView::AddSelectionToImport)));
-		const FSlateIcon MinusIcon(FEditorStyle::GetStyleSetName(), "PropertyWindow.Button_RemoveFromArray");
+		const FSlateIcon MinusIcon(FEditorStyle::GetStyleSetName(), "Icons.Minus");
 		MenuBuilder.AddMenuEntry(LOCTEXT("UncheckForImport", "Remove Selection From Import"), FText(), MinusIcon, FUIAction(FExecuteAction::CreateSP(this, &SFbxSceneMaterialsListView::RemoveSelectionFromImport)));
 	}
 	MenuBuilder.EndSection();
diff --git a/Engine/Source/Editor/UnrealEd/Private/Fbx/SSceneReimportNodeTreeView.cpp b/Engine/Source/Editor/UnrealEd/Private/Fbx/SSceneReimportNodeTreeView.cpp
--- a/Engine/Source/Editor/UnrealEd/Private/Fbx/SSceneReimportNodeTreeView.cpp
+++ b/Engine/Source/Editor/UnrealEd/Private/Fbx/SSceneReimportNodeTreeView.cpp
@@ -584,7 +584,7 @@ TSharedPtr<SWidget> SFbxReimportSceneTreeView::OnOpenContextMenu()
 	{
 		const FSlateIcon PlusIcon(FEditorStyle::GetStyleSetName(), "Plus");
 		MenuBuilder.AddMenuEntry(LOCTEXT("CheckForImport", "Add Selection To Import"), FText(), PlusIcon, FUIAction(FExecuteAction::CreateSP(this, &SFbxReimportSceneTreeView::AddSelectionToImport)));
-		const FSlateIcon MinusIcon(FEditorStyle::GetStyleSetName(), "PropertyWindow.Button_RemoveFromArray");
+		const FSlateIcon MinusIcon(FEditorStyle::GetStyleSetName(), "Icons.Minus");
 		MenuBuilder.AddMenuEntry(LOCTEXT("UncheckForImport", "Remove Selection From Import"), FText(), MinusIcon, FUIAction(FExecuteAction::CreateSP(this, &SFbxReimportSceneTreeView::RemoveSelectionFromImport)));
 	}
 	MenuBuilder.EndSection();
diff --git a/Engine/Source/Editor/UnrealEd/Private/Fbx/SSceneReimportSkeletalMeshListView.cpp b/Engine/Source/Editor/UnrealEd/Private/Fbx/SSceneReimportSkeletalMeshListView.cpp
--- a/Engine/Source/Editor/UnrealEd/Private/Fbx/SSceneReimportSkeletalMeshListView.cpp
+++ b/Engine/Source/Editor/UnrealEd/Private/Fbx/SSceneReimportSkeletalMeshListView.cpp
@@ -450,7 +450,7 @@ TSharedPtr<SWidget> SFbxSceneSkeletalMeshReimportListView::OnOpenContextMenu()
 	{
 		const FSlateIcon PlusIcon(FEditorStyle::GetStyleSetName(), "Plus");
 		MenuBuilder.AddMenuEntry(LOCTEXT("CheckForImport", "Add Selection To Import"), FText(), PlusIcon, FUIAction(FExecuteAction::CreateSP(this, &SFbxSceneSkeletalMeshReimportListView::AddSelectionToImport)));
-		const FSlateIcon MinusIcon(FEditorStyle::GetStyleSetName(), "PropertyWindow.Button_RemoveFromArray");
+		const FSlateIcon MinusIcon(FEditorStyle::GetStyleSetName(), "Icons.Minus");
 		MenuBuilder.AddMenuEntry(LOCTEXT("UncheckForImport", "Remove Selection From Import"), FText(), MinusIcon, FUIAction(FExecuteAction::CreateSP(this, &SFbxSceneSkeletalMeshReimportListView::RemoveSelectionFromImport)));
 	}
 	MenuBuilder.EndSection();
diff --git a/Engine/Source/Editor/UnrealEd/Private/Fbx/SSceneReimportStaticMeshListView.cpp b/Engine/Source/Editor/UnrealEd/Private/Fbx/SSceneReimportStaticMeshListView.cpp
--- a/Engine/Source/Editor/UnrealEd/Private/Fbx/SSceneReimportStaticMeshListView.cpp
+++ b/Engine/Source/Editor/UnrealEd/Private/Fbx/SSceneReimportStaticMeshListView.cpp
@@ -450,7 +450,7 @@ TSharedPtr<SWidget> SFbxSceneStaticMeshReimportListView::OnOpenContextMenu()
 	{
 		const FSlateIcon PlusIcon(FEditorStyle::GetStyleSetName(), "Plus");
 		MenuBuilder.AddMenuEntry(LOCTEXT("CheckForImport", "Add Selection To Import"), FText(), PlusIcon, FUIAction(FExecuteAction::CreateSP(this, &SFbxSceneStaticMeshReimportListView::AddSelectionToImport)));
-		const FSlateIcon MinusIcon(FEditorStyle::GetStyleSetName(), "PropertyWindow.Button_RemoveFromArray");
+		const FSlateIcon MinusIcon(FEditorStyle::GetStyleSetName(), "Icons.Minus");
 		MenuBuilder.AddMenuEntry(LOCTEXT("UncheckForImport", "Remove Selection From Import"), FText(), MinusIcon, FUIAction(FExecuteAction::CreateSP(this, &SFbxSceneStaticMeshReimportListView::RemoveSelectionFromImport)));
 	}
 	MenuBuilder.EndSection();
diff --git a/Engine/Source/Editor/UnrealEd/Private/Fbx/SSceneSkeletalMeshListView.cpp b/Engine/Source/Editor/UnrealEd/Private/Fbx/SSceneSkeletalMeshListView.cpp
--- a/Engine/Source/Editor/UnrealEd/Private/Fbx/SSceneSkeletalMeshListView.cpp
+++ b/Engine/Source/Editor/UnrealEd/Private/Fbx/SSceneSkeletalMeshListView.cpp
@@ -300,7 +300,7 @@ TSharedPtr<SWidget> SFbxSceneSkeletalMeshListView::OnOpenContextMenu()
 
 	const FSlateIcon PlusIcon(FEditorStyle::GetStyleSetName(), "Plus");
 	MenuBuilder.AddMenuEntry(LOCTEXT("CheckForImport", "Add Selection To Import"), FText(), PlusIcon, FUIAction(FExecuteAction::CreateSP(this, &SFbxSceneSkeletalMeshListView::AddSelectionToImport)));
-	const FSlateIcon MinusIcon(FEditorStyle::GetStyleSetName(), "PropertyWindow.Button_RemoveFromArray");
+	const FSlateIcon MinusIcon(FEditorStyle::GetStyleSetName(), "Icons.Minus");
 	MenuBuilder.AddMenuEntry(LOCTEXT("UncheckForImport", "Remove Selection From Import"), FText(), MinusIcon, FUIAction(FExecuteAction::CreateSP(this, &SFbxSceneSkeletalMeshListView::RemoveSelectionFromImport)));
 
 	AddBakePivotMenu(MenuBuilder);
diff --git a/Engine/Source/Editor/UnrealEd/Private/Toolkits/SGlobalTabSwitchingDialog.cpp b/Engine/Source/Editor/UnrealEd/Private/Toolkits/SGlobalTabSwitchingDialog.cpp
--- a/Engine/Source/Editor/UnrealEd/Private/Toolkits/SGlobalTabSwitchingDialog.cpp
+++ b/Engine/Source/Editor/UnrealEd/Private/Toolkits/SGlobalTabSwitchingDialog.cpp
@@ -281,7 +281,7 @@ void SGlobalTabSwitchingDialog::OnMainTabListSelectionChanged(FTabListItemPtr In
 					.VAlign(VAlign_Center)
 					[
 						SNew(SImage)
-						.Image(FEditorStyle::GetBrush("PropertyWindow.Button_Browse"))
+						.Image(FEditorStyle::GetBrush("Icons.Search"))
 					]
 					+SHorizontalBox::Slot()
 					.AutoWidth()
diff --git a/Engine/Source/Runtime/AppFramework/Private/Widgets/Colors/SColorThemes.cpp b/Engine/Source/Runtime/AppFramework/Private/Widgets/Colors/SColorThemes.cpp
--- a/Engine/Source/Runtime/AppFramework/Private/Widgets/Colors/SColorThemes.cpp
+++ b/Engine/Source/Runtime/AppFramework/Private/Widgets/Colors/SColorThemes.cpp
@@ -148,9 +148,8 @@ void SColorTrash::Construct( const FArguments& InArgs )
 			SNew(SHorizontalBox)
 			+SHorizontalBox::Slot() .HAlign(HAlign_Center) .FillWidth(1)
 			[
-				SNew(SImage) .Image(InArgs._UsesSmallIcon.Get() ?
-					FCoreStyle::Get().GetBrush("TrashCan_Small") :
-					FCoreStyle::Get().GetBrush("TrashCan"))
+				SNew(SImage) 
+				.Image(FAppStyle::Get().GetBrush("Icons.Delete"))
 			]
 		]
 	];
diff --git a/Engine/Source/Runtime/SlateCore/Private/Styling/StarshipCoreStyle.cpp b/Engine/Source/Runtime/SlateCore/Private/Styling/StarshipCoreStyle.cpp
--- a/Engine/Source/Runtime/SlateCore/Private/Styling/StarshipCoreStyle.cpp
+++ b/Engine/Source/Runtime/SlateCore/Private/Styling/StarshipCoreStyle.cpp
@@ -279,9 +279,6 @@ TSharedRef<ISlateStyle> FStarshipCoreStyle::Create()
 		Style->Set("AppIconPadding", FMargin(11, 11, 3, 5));
 		Style->Set("AppIconPadding.Small", FMargin(4, 4, 0, 0));
 
-		Style->Set("TrashCan", new IMAGE_BRUSH("Icons/TrashCan", FVector2D(64, 64)));
-		Style->Set("TrashCan_Small", new IMAGE_BRUSH("Icons/TrashCan_Small", FVector2D(18, 18)));
-
 		Style->Set("Icons.Denied", new IMAGE_BRUSH("Icons/denied_16x", Icon16x16));
 	
 		Style->Set("Icons.Help", new IMAGE_BRUSH("Icons/icon_help_16x", Icon16x16));
@@ -304,7 +301,7 @@ TSharedRef<ISlateStyle> FStarshipCoreStyle::Create()
 		Style->Set("Icons.Minus", new IMAGE_BRUSH_SVG("Starship/Common/minus", Icon16x16));
 		Style->Set("Icons.PlusCircle", new IMAGE_BRUSH_SVG("Starship/Common/plus-circle", Icon16x16));
 		Style->Set("Icons.X", new IMAGE_BRUSH_SVG("Starship/Common/close", Icon16x16));
-
+		Style->Set("Icons.Delete", new IMAGE_BRUSH_SVG("Starship/Common/delete-outline", Icon16x16));
 		Style->Set("Icons.Save", new IMAGE_BRUSH_SVG("Starship/Common/save", Icon16x16));
 
 		Style->Set("Icons.Import", new IMAGE_BRUSH_SVG("Starship/Common/import", Icon16x16));
