Fix not being able to double click on asset thumbnails to open their editor from the details panel

diff --git a/Engine/Source/Editor/PropertyEditor/Private/UserInterface/PropertyEditor/SPropertyEditorAsset.cpp b/Engine/Source/Editor/PropertyEditor/Private/UserInterface/PropertyEditor/SPropertyEditorAsset.cpp
--- a/Engine/Source/Editor/PropertyEditor/Private/UserInterface/PropertyEditor/SPropertyEditorAsset.cpp
+++ b/Engine/Source/Editor/PropertyEditor/Private/UserInterface/PropertyEditor/SPropertyEditorAsset.cpp
@@ -501,6 +501,7 @@ void SPropertyEditorAsset::Construct(const FArguments& InArgs, const TSharedPtr<
 				[
 					SNew(SImage)
 					.Image(this, &SPropertyEditorAsset::GetThumbnailBorder)
+					.Visibility(EVisibility::SelfHitTestInvisible)
 				]
 			]
 		];
