Details view array drag highlighting fixed.

Details view background color and drop shadow changed to be lighter.

diff --git a/Engine/Source/Editor/PropertyEditor/Private/SDetailSingleItemRow.cpp b/Engine/Source/Editor/PropertyEditor/Private/SDetailSingleItemRow.cpp
--- a/Engine/Source/Editor/PropertyEditor/Private/SDetailSingleItemRow.cpp
+++ b/Engine/Source/Editor/PropertyEditor/Private/SDetailSingleItemRow.cpp
@@ -261,12 +261,14 @@ bool SDetailSingleItemRow::CheckValidDrop(const TSharedPtr<SDetailSingleItemRow>
 FReply SDetailSingleItemRow::OnArrayDrop(const FDragDropEvent& DragDropEvent)
 {
 	bIsHoveredDragTarget = false;
+	
 	TSharedPtr<FArrayRowDragDropOp> ArrayDropOp = DragDropEvent.GetOperationAs< FArrayRowDragDropOp >();
-	TSharedPtr<SDetailSingleItemRow> RowPtr = nullptr;
-	if (ArrayDropOp.IsValid() && ArrayDropOp->Row.IsValid())
+	if (!ArrayDropOp.IsValid())
 	{
-		RowPtr = ArrayDropOp->Row.Pin();
+		return FReply::Unhandled();
 	}
+
+	TSharedPtr<SDetailSingleItemRow> RowPtr = ArrayDropOp->Row.Pin();
 	if (!RowPtr.IsValid())
 	{
 		return FReply::Unhandled();
@@ -301,31 +303,32 @@ FReply SDetailSingleItemRow::OnArrayDrop(const FDragDropEvent& DragDropEvent)
 			DetailsView->GetPropertyUtilities()->NotifyFinishedChangingProperties(MoveEvent);
 		}
 	}
+
 	return FReply::Handled();
 }
 
 TOptional<EItemDropZone> SDetailSingleItemRow::OnArrayCanAcceptDrop(const FDragDropEvent& DragDropEvent, EItemDropZone DropZone, TSharedPtr< FDetailTreeNode > Type)
 {
-	bIsHoveredDragTarget = false;
 	TSharedPtr<FArrayRowDragDropOp> ArrayDropOp = DragDropEvent.GetOperationAs< FArrayRowDragDropOp >();
-	TSharedPtr<SDetailSingleItemRow> RowPtr = nullptr;
-	if (ArrayDropOp.IsValid() && ArrayDropOp->Row.IsValid())
+	if (!ArrayDropOp.IsValid())
 	{
-		RowPtr = ArrayDropOp->Row.Pin();
+		return TOptional<EItemDropZone>();
 	}
+
+	TSharedPtr<SDetailSingleItemRow> RowPtr = ArrayDropOp->Row.Pin();
 	if (!RowPtr.IsValid())
 	{
 		return TOptional<EItemDropZone>();
 	}
 
-	if (CheckValidDrop(RowPtr))
-	{
-		ArrayDropOp->DecoratorText->SetText(NSLOCTEXT("ArrayDragDrop", "PlaceRowHere", "Place Row Here"));
-	}
-	else
+	bool IsValidDrop = CheckValidDrop(RowPtr);
+	if (!IsValidDrop)
 	{
-		ArrayDropOp->DecoratorText->SetText(NSLOCTEXT("ArrayDragDrop", "CannotPlaceRowHere", "Cannot Place Row Here"));
+		bIsHoveredDragTarget = false;
 	}
+
+	ArrayDropOp->IsValidTarget = IsValidDrop;
+
 	return TOptional<EItemDropZone>();
 }
 
@@ -545,14 +548,12 @@ void SDetailSingleItemRow::Construct( const FArguments& InArgs, FDetailLayoutCus
 							ArrayHandle
 						];
 
-					ArrayDragDelegate = FOnTableRowDragEnter::CreateSP(this, &SDetailSingleItemRow::OnArrayDragEnter);
-					ArrayDragLeaveDelegate = FOnTableRowDragLeave::CreateSP(this, &SDetailSingleItemRow::OnArrayDragLeave);
-					ArrayDropDelegate = FOnTableRowDrop::CreateSP(this, &SDetailSingleItemRow::OnArrayDrop);
-					ArrayAcceptDropDelegate = FOnCanAcceptDrop::CreateSP(this, &SDetailSingleItemRow::OnArrayCanAcceptDrop);
 					SwappablePropertyNode = PropertyNode;
 				}
-				else if (CastField<FArrayProperty>(PropertyNode->GetProperty()) != nullptr // Is an array
-					&& CastField<FObjectProperty>(CastField<FArrayProperty>(PropertyNode->GetProperty())->Inner) != nullptr) // Is an object array
+				
+				if (PropertyNode->IsReorderable() || 
+					(CastField<FArrayProperty>(PropertyNode->GetProperty()) != nullptr && 
+					CastField<FObjectProperty>(CastField<FArrayProperty>(PropertyNode->GetProperty())->Inner) != nullptr)) // Is an object array
 				{
 					ArrayDragDelegate = FOnTableRowDragEnter::CreateSP(this, &SDetailSingleItemRow::OnArrayDragEnter);
 					ArrayDragLeaveDelegate = FOnTableRowDragLeave::CreateSP(this, &SDetailSingleItemRow::OnArrayDragLeave);
@@ -774,6 +775,16 @@ FSlateColor SDetailSingleItemRow::GetInnerBackgroundColor() const
 		return FAppStyle::Get().GetSlateColor("Colors.Header");
 	}
 
+	if (bIsHoveredDragTarget)
+	{
+		return FAppStyle::Get().GetSlateColor("Colors.Hover2");
+	}
+
+	if (bIsDragDropObject)
+	{
+		return FAppStyle::Get().GetSlateColor("Colors.Hover");
+	}
+
 	int32 IndentLevel = GetIndentLevelForBackgroundColor();
 	return PropertyEditorConstants::GetRowBackgroundColor(IndentLevel);
 }
@@ -783,6 +794,7 @@ int32 SDetailSingleItemRow::GetIndentLevelForBackgroundColor() const
 	int32 IndentLevel = 0; 
 	if (OwnerTablePtr.IsValid())
 	{
+		// every item is in a category, but we don't want to show an indent for "top-level" properties
 		IndentLevel = GetIndentLevel() - 1;
 	}
 
@@ -1196,17 +1208,39 @@ TSharedPtr<FArrayRowDragDropOp> SArrayRowHandle::CreateDragDropOperation(TShared
 	return Operation;
 }
 
-FArrayRowDragDropOp::FArrayRowDragDropOp(TSharedPtr<SDetailSingleItemRow> InRow)
+FText FArrayRowDragDropOp::GetDecoratorText() const
 {
-	Row = InRow;
+	if (IsValidTarget)
+	{
+		return NSLOCTEXT("ArrayDragDrop", "PlaceRowHere", "Place Row Here");
+	}
+	else
+	{
+		return NSLOCTEXT("ArrayDragDrop", "CannotPlaceRowHere", "Cannot Place Row Here");
+	}
+}
 
-	TSharedPtr<SDetailSingleItemRow> RowPtr = nullptr;
-	if (Row.IsValid())
+const FSlateBrush* FArrayRowDragDropOp::GetDecoratorIcon() const
+{
+	if (IsValidTarget)
 	{
-		RowPtr = Row.Pin();
-		// mark row as being used for drag and drop
-		RowPtr->SetIsDragDrop(true);
+		return FEditorStyle::GetBrush("Graph.ConnectorFeedback.OK");
 	}
+	else 
+	{
+		return FEditorStyle::GetBrush("Graph.ConnectorFeedback.Error");
+	}
+}
+
+FArrayRowDragDropOp::FArrayRowDragDropOp(TSharedPtr<SDetailSingleItemRow> InRow)
+{
+	IsValidTarget = false;
+
+	check(InRow.IsValid());
+	Row = InRow;
+
+	// mark row as being used for drag and drop
+	InRow->SetIsDragDrop(true);
 
 	DecoratorWidget = SNew(SBorder)
 		.Padding(8.f)
@@ -1218,8 +1252,16 @@ FArrayRowDragDropOp::FArrayRowDragDropOp(TSharedPtr<SDetailSingleItemRow> InRow)
 			.AutoWidth()
 			.VAlign(VAlign_Center)
 			[
-				SAssignNew(DecoratorText, STextBlock)
-				.Text(NSLOCTEXT("ArrayDragDrop", "PlaceRowHere", "Place Row Here"))
+				SNew(SImage)
+				.Image_Raw(this, &FArrayRowDragDropOp::GetDecoratorIcon)
+			]
+			+ SHorizontalBox::Slot()
+			.Padding(5,0,0,0)
+			.AutoWidth()
+			.VAlign(VAlign_Center)
+			[
+				SNew(STextBlock)
+				.Text_Raw(this, &FArrayRowDragDropOp::GetDecoratorText)
 			]
 		];
 
@@ -1230,10 +1272,9 @@ void FArrayRowDragDropOp::OnDrop(bool bDropWasHandled, const FPointerEvent& Mous
 {
 	FDecoratedDragDropOp::OnDrop(bDropWasHandled, MouseEvent);
 
-	TSharedPtr<SDetailSingleItemRow> RowPtr = nullptr;
-	if (Row.IsValid())
+	TSharedPtr<SDetailSingleItemRow> RowPtr = Row.Pin();
+	if (RowPtr.IsValid())
 	{
-		RowPtr = Row.Pin();
 		// reset value
 		RowPtr->SetIsDragDrop(false);
 	}
diff --git a/Engine/Source/Editor/PropertyEditor/Private/SDetailSingleItemRow.h b/Engine/Source/Editor/PropertyEditor/Private/SDetailSingleItemRow.h
--- a/Engine/Source/Editor/PropertyEditor/Private/SDetailSingleItemRow.h
+++ b/Engine/Source/Editor/PropertyEditor/Private/SDetailSingleItemRow.h
@@ -149,6 +149,10 @@ public:
 		return DecoratorWidget;
 	}
 
-	TWeakPtr<class SDetailSingleItemRow> Row;
-	TSharedPtr<STextBlock> DecoratorText;
+	TWeakPtr<SDetailSingleItemRow> Row;
+	bool IsValidTarget;
+
+private:
+	FText GetDecoratorText() const;
+	const FSlateBrush* GetDecoratorIcon() const;
 };
diff --git a/Engine/Source/Editor/PropertyEditor/Private/UserInterface/PropertyEditor/PropertyEditorConstants.cpp b/Engine/Source/Editor/PropertyEditor/Private/UserInterface/PropertyEditor/PropertyEditorConstants.cpp
--- a/Engine/Source/Editor/PropertyEditor/Private/UserInterface/PropertyEditor/PropertyEditorConstants.cpp
+++ b/Engine/Source/Editor/PropertyEditor/Private/UserInterface/PropertyEditor/PropertyEditorConstants.cpp
@@ -36,13 +36,17 @@ FSlateColor PropertyEditorConstants::GetRowBackgroundColor(int32 IndentLevel)
 		}
 	}
 
-	static const FLinearColor Colors[] =
+	static const uint8 ColorOffsets[] =
 	{
-		COLOR("#2D2D2D"),
-		COLOR("#373737"),
-		COLOR("#434343"),
-		COLOR("#515151")
+		2, (2 + 4), (6 + 6), (12 + 8)
 	};
 
-	return FSlateColor(Colors[ColorIndex]);
+	FColor BaseColor = FAppStyle::Get().GetSlateColor("Colors.Background").GetSpecifiedColor().ToFColor(true);
+
+	FColor ColorWithOffset(
+		BaseColor.R + ColorOffsets[ColorIndex], 
+		BaseColor.G + ColorOffsets[ColorIndex], 
+		BaseColor.B + ColorOffsets[ColorIndex]);
+
+	return FSlateColor(FLinearColor::FromSRGBColor(ColorWithOffset));
 }
diff --git a/Engine/Source/Editor/PropertyEditor/Private/UserInterface/PropertyEditor/PropertyEditorConstants.h b/Engine/Source/Editor/PropertyEditor/Private/UserInterface/PropertyEditor/PropertyEditorConstants.h
--- a/Engine/Source/Editor/PropertyEditor/Private/UserInterface/PropertyEditor/PropertyEditorConstants.h
+++ b/Engine/Source/Editor/PropertyEditor/Private/UserInterface/PropertyEditor/PropertyEditorConstants.h
@@ -2,6 +2,7 @@
 #pragma once
 
 #include "CoreMinimal.h"
+#include "Styling/SlateColor.h"
 
 struct FSlateBrush;
 
