Strip trailing spaces from Name property values.

Based on PR #7522 (Contributed by [REDACTED_PERSON])

diff --git a/Engine/Source/Editor/PropertyEditor/Private/PropertyHandleImpl.cpp b/Engine/Source/Editor/PropertyEditor/Private/PropertyHandleImpl.cpp
--- a/Engine/Source/Editor/PropertyEditor/Private/PropertyHandleImpl.cpp
+++ b/Engine/Source/Editor/PropertyEditor/Private/PropertyHandleImpl.cpp
@@ -693,33 +693,18 @@ FPropertyAccess::Result FPropertyValueImpl::SetValueAsString( const FString& InV
 
 		FString Value = InValue;
 
-		// Strip any leading underscores and spaces from names.
+		// Strip any leading underscores, and trim the name
 		if( NodeProperty && NodeProperty->IsA( FNameProperty::StaticClass() ) )
 		{
-			while ( true )
+			// trim whitespace and leading underscores
+			int32 Pos = 0;
+			while (Pos < Value.Len() && 
+				(FChar::IsWhitespace(Value[Pos]) || Value[Pos] == '_'))
 			{
-				if ( Value.StartsWith( TEXT("_"), ESearchCase::CaseSensitive ) )
-				{
-					// Strip leading underscores.
-					do
-					{
-						Value.RightInline( Value.Len()-1, false);
-					} while ( Value.StartsWith( TEXT("_"), ESearchCase::CaseSensitive ) );
-				}
-				else if ( Value.StartsWith( TEXT(" "), ESearchCase::CaseSensitive) )
-				{
-					// Strip leading spaces.
-					do
-					{
-						Value.RightInline( Value.Len()-1, false );
-					} while ( Value.StartsWith( TEXT(" "), ESearchCase::CaseSensitive) );
-				}
-				else
-				{
-					// Starting with something valid -- break.
-					break;
-				}
+				++Pos;
 			}
+			Value.RemoveAt(0, Pos, false);
+			Value.TrimEndInline();
 		}
 
 		// If more than one object is selected, an empty field indicates their values for this property differ.
