Cache modified objects in FPropertyNode::NotifyPostChange

diff --git a/Engine/Source/Editor/PropertyEditor/Private/PropertyNode.cpp b/Engine/Source/Editor/PropertyEditor/Private/PropertyNode.cpp
--- a/Engine/Source/Editor/PropertyEditor/Private/PropertyNode.cpp
+++ b/Engine/Source/Editor/PropertyEditor/Private/PropertyNode.cpp
@@ -2457,42 +2457,53 @@ void FPropertyNode::NotifyPostChange( FPropertyChangedEvent& InPropertyChangedEv
 		// Call PostEditChange on the object chain.
 		while ( true )
 		{
-			int32 CurrentObjectIndex = 0;
-			for( TPropObjectIterator Itor( ObjectNode->ObjectIterator() ) ; Itor ; ++Itor )
+			TArray<TWeakObjectPtr<UObject>> Containers;
+			// It's possible that PostEditChangeProperty may cause a construction script to re-run
+			// which will invalidate the PropObjectIterator. We need to instead cache all of the objects
+			// before emitting any change events to ensure there is a PostChange for every PreChange.
+			for (TPropObjectIterator Itor(ObjectNode->ObjectIterator()); Itor; ++Itor)
 			{
-				UObject* Object = Itor->Get();
-				if ( PropertyChain->Num() == 0 )
-				{
-					//copy 
-					FPropertyChangedEvent ChangedEvent = InPropertyChangedEvent;
-					if (CurProperty != InPropertyChangedEvent.Property)
-					{
-						//parent object node property.  Reset other internals and leave the event type as unspecified
-						ChangedEvent = FPropertyChangedEvent(CurProperty, InPropertyChangedEvent.ChangeType);
-					}
-					ChangedEvent.ObjectIteratorIndex = CurrentObjectIndex;
-					if( Object )
-					{
-						Object->PostEditChangeProperty( ChangedEvent );
-					}
-				}
-				else
+				Containers.Add(*Itor);
+			}
+			for (int32 CurrentObjectIndex = 0; CurrentObjectIndex < Containers.Num(); ++CurrentObjectIndex)
+			{
+				UObject* Object = Containers[CurrentObjectIndex].Get();
+				// It is possible that the PostChange for the first object deletes the second object, so
+				// we must check to make sure the object still exists first.
+				if (Object)
 				{
-					FPropertyChangedEvent ChangedEvent = InPropertyChangedEvent;
-					if (CurProperty != InPropertyChangedEvent.Property)
+					if (PropertyChain->Num() == 0)
 					{
-						//parent object node property.  Reset other internals and leave the event type as unspecified
-						ChangedEvent = FPropertyChangedEvent(CurProperty, InPropertyChangedEvent.ChangeType);
+						//copy 
+						FPropertyChangedEvent ChangedEvent = InPropertyChangedEvent;
+						if (CurProperty != InPropertyChangedEvent.Property)
+						{
+							//parent object node property.  Reset other internals and leave the event type as unspecified
+							ChangedEvent = FPropertyChangedEvent(CurProperty, InPropertyChangedEvent.ChangeType);
+						}
+						ChangedEvent.ObjectIteratorIndex = CurrentObjectIndex;
+						if (Object)
+						{
+							Object->PostEditChangeProperty(ChangedEvent);
+						}
 					}
-					FPropertyChangedChainEvent ChainEvent(*PropertyChain, ChangedEvent);
-					ChainEvent.ObjectIteratorIndex = CurrentObjectIndex;
-					if( Object )
+					else
 					{
-						Object->PostEditChangeChainProperty(ChainEvent);
+						FPropertyChangedEvent ChangedEvent = InPropertyChangedEvent;
+						if (CurProperty != InPropertyChangedEvent.Property)
+						{
+							//parent object node property.  Reset other internals and leave the event type as unspecified
+							ChangedEvent = FPropertyChangedEvent(CurProperty, InPropertyChangedEvent.ChangeType);
+						}
+						FPropertyChangedChainEvent ChainEvent(*PropertyChain, ChangedEvent);
+						ChainEvent.ObjectIteratorIndex = CurrentObjectIndex;
+						if (Object)
+						{
+							Object->PostEditChangeChainProperty(ChainEvent);
+						}
 					}
+					LevelDirtyCallback.Request();
 				}
-				LevelDirtyCallback.Request();
-				++CurrentObjectIndex;
 			}
 
 
