Detail categories no longer clip long names halfway if there isn't a custom widget to display in the category header.

diff --git a/Engine/Source/Editor/PropertyEditor/Private/DetailCategoryGroupNode.cpp b/Engine/Source/Editor/PropertyEditor/Private/DetailCategoryGroupNode.cpp
--- a/Engine/Source/Editor/PropertyEditor/Private/DetailCategoryGroupNode.cpp
+++ b/Engine/Source/Editor/PropertyEditor/Private/DetailCategoryGroupNode.cpp
@@ -11,10 +11,9 @@ void SDetailCategoryTableRow::Construct( const FArguments& InArgs, TSharedRef<FD
 	bIsInnerCategory = InArgs._InnerCategory;
 	bShowBorder = InArgs._ShowBorder;
 
-	TSharedPtr<SWidget> Widget = SNullWidget::NullWidget;
 	FDetailColumnSizeData& ColumnSizeData = InOwnerTreeNode->GetDetailsView()->GetColumnSizeData();
 
-	Widget = SNew(SHorizontalBox)
+	TSharedPtr<SHorizontalBox> HeaderBox = SNew(SHorizontalBox)
 		+ SHorizontalBox::Slot()
 		.VAlign(VAlign_Center)
 		.Padding(8, 0, 8, 0)
@@ -24,20 +23,24 @@ void SDetailCategoryTableRow::Construct( const FArguments& InArgs, TSharedRef<FD
 		]
 		+ SHorizontalBox::Slot()
 		.VAlign(VAlign_Center)
-		.FillWidth(0.5f)
+		.FillWidth(1)
 		[
 			SNew(STextBlock)
 			.TransformPolicy(ETextTransformPolicy::ToUpper)
 			.Text(InArgs._DisplayName)
 			.Font(FAppStyle::Get().GetFontStyle(bIsInnerCategory ? PropertyEditorConstants::PropertyFontStyle : PropertyEditorConstants::CategoryFontStyle))
 			.TextStyle(FAppStyle::Get(), "DetailsView.CategoryTextStyle")
-		]
-		+ SHorizontalBox::Slot()
+		];
+
+	if (InArgs._HeaderContent.IsValid())
+	{
+		HeaderBox->AddSlot()
 		.VAlign(VAlign_Center)
-		.FillWidth(0.5f)
+		.FillWidth(1)
 		[
-			InArgs._HeaderContent.IsValid() ? InArgs._HeaderContent.ToSharedRef() : SNullWidget::NullWidget
+			InArgs._HeaderContent.ToSharedRef()
 		];
+	}
 
 	const float VerticalPadding = bIsInnerCategory ? 6 : 8;
 	
@@ -53,7 +56,7 @@ void SDetailCategoryTableRow::Construct( const FArguments& InArgs, TSharedRef<FD
 			.BorderBackgroundColor( this, &SDetailCategoryTableRow::GetBackgroundColor )
 			.Padding( FMargin(0, VerticalPadding, SDetailTableRowBase::ScrollbarPaddingSize, VerticalPadding) )
 			[
-				Widget.ToSharedRef()
+				HeaderBox.ToSharedRef()
 			]
 		]
 	];
