Fixed EditCondition showing up for every row in the property matrix, and crashing the editor if unchecked.

diff --git a/Engine/Source/Editor/PropertyEditor/Private/PropertyEditorHelpers.cpp b/Engine/Source/Editor/PropertyEditor/Private/PropertyEditorHelpers.cpp
--- a/Engine/Source/Editor/PropertyEditor/Private/PropertyEditorHelpers.cpp
+++ b/Engine/Source/Editor/PropertyEditor/Private/PropertyEditorHelpers.cpp
@@ -331,8 +331,6 @@ void SEditConditionWidget::Construct( const FArguments& Args )
 	EditConditionValue = Args._EditConditionValue;
 	OnEditConditionValueChanged = Args._OnEditConditionValueChanged;
 
-	SetVisibility(HasEditConditionToggle() ? EVisibility::Visible : EVisibility::Collapsed);
-
 	ChildSlot
 	[
 		// Some properties become irrelevant depending on the value of other properties.
@@ -341,6 +339,7 @@ void SEditConditionWidget::Construct( const FArguments& Args )
 		SNew(SCheckBox)
 		.OnCheckStateChanged(this, &SEditConditionWidget::OnEditConditionCheckChanged)
 		.IsChecked(this, &SEditConditionWidget::OnGetEditConditionCheckState)
+		.Visibility_Lambda([this]() { return HasEditConditionToggle() ? EVisibility::Visible : EVisibility::Collapsed; })
 	];
 }
 
diff --git a/Engine/Source/Editor/PropertyEditor/Private/UserInterface/PropertyEditor/SPropertyEditorTableRow.cpp b/Engine/Source/Editor/PropertyEditor/Private/UserInterface/PropertyEditor/SPropertyEditorTableRow.cpp
--- a/Engine/Source/Editor/PropertyEditor/Private/UserInterface/PropertyEditor/SPropertyEditorTableRow.cpp
+++ b/Engine/Source/Editor/PropertyEditor/Private/UserInterface/PropertyEditor/SPropertyEditorTableRow.cpp
@@ -75,7 +75,7 @@ TSharedRef< SWidget > SPropertyEditorTableRow::ConstructNameColumnWidget()
 			SNew( SEditConditionWidget )
 			.EditConditionValue(PropertyEditor.ToSharedRef(), &FPropertyEditor::IsEditConditionMet)
 			.OnEditConditionValueChanged_Lambda([this](bool bValue) { PropertyEditor->ToggleEditConditionState(); })
-
+			.Visibility_Lambda([this]() { return PropertyEditor->HasEditCondition() ? EVisibility::Visible : EVisibility::Collapsed; })
 		]
 		+SHorizontalBox::Slot()
 		.AutoWidth()
