Inner categories (produced by eg. instanced UObjects) now correctly indent.

Moved SDetailRowIndent and SConstrainedBox to their own files. Created SDetailExpanderArrow widget so that it can be reused in categories and groups.

#b lauren.barnes

diff --git a/Engine/Source/Editor/PropertyEditor/Private/DetailCategoryGroupNode.cpp b/Engine/Source/Editor/PropertyEditor/Private/DetailCategoryGroupNode.cpp
--- a/Engine/Source/Editor/PropertyEditor/Private/DetailCategoryGroupNode.cpp
+++ b/Engine/Source/Editor/PropertyEditor/Private/DetailCategoryGroupNode.cpp
@@ -3,6 +3,8 @@
 #include "DetailCategoryGroupNode.h"
 #include "PropertyCustomizationHelpers.h"
 #include "PropertyEditorConstants.h"
+#include "SDetailExpanderArrow.h"
+#include "SDetailRowIndent.h"
 
 void SDetailCategoryTableRow::Construct( const FArguments& InArgs, TSharedRef<FDetailTreeNode> InOwnerTreeNode, const TSharedRef<STableViewBase>& InOwnerTableView )
 {
@@ -11,18 +13,28 @@ void SDetailCategoryTableRow::Construct( const FArguments& InArgs, TSharedRef<FD
 	bIsInnerCategory = InArgs._InnerCategory;
 	bShowBorder = InArgs._ShowBorder;
 
+	const float VerticalPadding = bIsInnerCategory ? 6 : 8;
+
 	FDetailColumnSizeData& ColumnSizeData = InOwnerTreeNode->GetDetailsView()->GetColumnSizeData();
 
 	TSharedPtr<SHorizontalBox> HeaderBox = SNew(SHorizontalBox)
 		+ SHorizontalBox::Slot()
+		.HAlign(HAlign_Left)
+		.VAlign(VAlign_Fill)
+		.AutoWidth()
+		[
+			SNew(SDetailRowIndent, SharedThis(this))
+		]
+		+ SHorizontalBox::Slot()
+		.HAlign(HAlign_Left)
 		.VAlign(VAlign_Center)
-		.Padding(8, 0, 8, 0)
 		.AutoWidth()
 		[
-			SNew(SExpanderArrow, SharedThis(this))
+			SNew(SDetailExpanderArrow, SharedThis(this))
 		]
 		+ SHorizontalBox::Slot()
 		.VAlign(VAlign_Center)
+		.Padding(12, VerticalPadding, 0, VerticalPadding)
 		.FillWidth(1)
 		[
 			SNew(STextBlock)
@@ -41,9 +53,9 @@ void SDetailCategoryTableRow::Construct( const FArguments& InArgs, TSharedRef<FD
 			InArgs._HeaderContent.ToSharedRef()
 		];
 	}
-
-	const float VerticalPadding = bIsInnerCategory ? 6 : 8;
 	
+	const float EditConditionWidgetWidth = 33; // this is the width of the edit condition widget displayed on the left in SDetailSingleItemRow
+
 	ChildSlot
 	.Padding(0)
 	[
@@ -51,12 +63,30 @@ void SDetailCategoryTableRow::Construct( const FArguments& InArgs, TSharedRef<FD
 		.BorderImage( FAppStyle::Get().GetBrush( "DetailsView.GridLine") )
 		.Padding( FMargin(0, 0, 0, 1) )
 		[
-			SNew( SBorder )
-			.BorderImage( this, &SDetailCategoryTableRow::GetBackgroundImage )
-			.BorderBackgroundColor( this, &SDetailCategoryTableRow::GetBackgroundColor )
-			.Padding( FMargin(0, VerticalPadding, SDetailTableRowBase::ScrollbarPaddingSize, VerticalPadding) )
+			SNew(SHorizontalBox)
+			+ SHorizontalBox::Slot()
+			.HAlign(HAlign_Left)
+			.AutoWidth()
 			[
-				HeaderBox.ToSharedRef()
+				SNew(SBorder)
+				.BorderImage(FAppStyle::Get().GetBrush("DetailsView.CategoryMiddle"))
+				.BorderBackgroundColor( this, &SDetailCategoryTableRow::GetOuterBackgroundColor )
+				.Padding(0)
+				[
+					SNew(SSpacer)
+					.Size(bIsInnerCategory ? FVector2D(EditConditionWidgetWidth,0) : FVector2D(0,0))
+				]
+			]
+			+ SHorizontalBox::Slot()
+			.HAlign(HAlign_Fill)
+			[
+				SNew( SBorder )
+				.BorderImage( this, &SDetailCategoryTableRow::GetBackgroundImage )
+				.BorderBackgroundColor( this, &SDetailCategoryTableRow::GetInnerBackgroundColor )
+				.Padding(FMargin(0, 0, SDetailTableRowBase::ScrollbarPaddingSize, 0))
+				[
+					HeaderBox.ToSharedRef()
+				]
 			]
 		]
 	];
@@ -96,7 +126,7 @@ const FSlateBrush* SDetailCategoryTableRow::GetBackgroundImage() const
 	return nullptr;
 }
 
-FSlateColor SDetailCategoryTableRow::GetBackgroundColor() const
+FSlateColor SDetailCategoryTableRow::GetInnerBackgroundColor() const
 {
 	if (bShowBorder && bIsInnerCategory)
 	{
@@ -114,6 +144,16 @@ FSlateColor SDetailCategoryTableRow::GetBackgroundColor() const
 	return FSlateColor(FLinearColor::White);
 }
 
+FSlateColor SDetailCategoryTableRow::GetOuterBackgroundColor() const
+{
+	if (bIsHovered)
+	{
+		return FAppStyle::Get().GetSlateColor("Colors.Header");
+	}
+
+	return FAppStyle::Get().GetSlateColor("Colors.Background");
+}
+
 FReply SDetailCategoryTableRow::OnMouseButtonDown( const FGeometry& MyGeometry, const FPointerEvent& MouseEvent )
 {
 	if (MouseEvent.GetEffectingButton() == EKeys::LeftMouseButton)
diff --git a/Engine/Source/Editor/PropertyEditor/Private/DetailCategoryGroupNode.h b/Engine/Source/Editor/PropertyEditor/Private/DetailCategoryGroupNode.h
--- a/Engine/Source/Editor/PropertyEditor/Private/DetailCategoryGroupNode.h
+++ b/Engine/Source/Editor/PropertyEditor/Private/DetailCategoryGroupNode.h
@@ -32,7 +32,8 @@ public:
 private:
 	EVisibility IsSeparatorVisible() const;
 	const FSlateBrush* GetBackgroundImage() const;
-	FSlateColor GetBackgroundColor() const;
+	FSlateColor GetInnerBackgroundColor() const;
+	FSlateColor GetOuterBackgroundColor() const;
 
 	virtual FReply OnMouseButtonDown( const FGeometry& MyGeometry, const FPointerEvent& MouseEvent ) override;
 	virtual FReply OnMouseButtonDoubleClick( const FGeometry& InMyGeometry, const FPointerEvent& InMouseEvent ) override;
diff --git a/Engine/Source/Editor/PropertyEditor/Private/DetailItemNode.cpp b/Engine/Source/Editor/PropertyEditor/Private/DetailItemNode.cpp
--- a/Engine/Source/Editor/PropertyEditor/Private/DetailItemNode.cpp
+++ b/Engine/Source/Editor/PropertyEditor/Private/DetailItemNode.cpp
@@ -4,6 +4,7 @@
 #include "DetailCategoryGroupNode.h"
 #include "DetailGroup.h"
 #include "DetailPropertyRow.h"
+#include "SConstrainedBox.h"
 #include "SDetailSingleItemRow.h"
 #include "IDetailKeyframeHandler.h"
 #include "ObjectPropertyNode.h"
diff --git a/Engine/Source/Editor/PropertyEditor/Private/SConstrainedBox.cpp b/Engine/Source/Editor/PropertyEditor/Private/SConstrainedBox.cpp
new file mode 100644
--- /dev/null
+++ b/Engine/Source/Editor/PropertyEditor/Private/SConstrainedBox.cpp
@@ -0,0 +1,37 @@
+// Copyright Epic Games, Inc. All Rights Reserved.
+
+#include "SConstrainedBox.h"
+
+void SConstrainedBox::Construct(const FArguments& InArgs)
+{
+	MinWidth = InArgs._MinWidth;
+	MaxWidth = InArgs._MaxWidth;
+
+	ChildSlot
+	[
+		InArgs._Content.Widget
+	];
+}
+
+FVector2D SConstrainedBox::ComputeDesiredSize(float LayoutScaleMultiplier) const
+{
+	const float MinWidthVal = MinWidth.Get().Get(0.0f);
+	const float MaxWidthVal = MaxWidth.Get().Get(0.0f);
+
+	if (MinWidthVal == 0.0f && MaxWidthVal == 0.0f)
+	{
+		return SCompoundWidget::ComputeDesiredSize(LayoutScaleMultiplier);
+	}
+	else
+	{
+		FVector2D ChildSize = ChildSlot.GetWidget()->GetDesiredSize();
+
+		float XVal = FMath::Max(MinWidthVal, ChildSize.X);
+		if (MaxWidthVal >= MinWidthVal)
+		{
+			XVal = FMath::Min(MaxWidthVal, XVal);
+		}
+
+		return FVector2D(XVal, ChildSize.Y);
+	}
+}
diff --git a/Engine/Source/Editor/PropertyEditor/Private/SConstrainedBox.h b/Engine/Source/Editor/PropertyEditor/Private/SConstrainedBox.h
new file mode 100644
--- /dev/null
+++ b/Engine/Source/Editor/PropertyEditor/Private/SConstrainedBox.h
@@ -0,0 +1,26 @@
+// Copyright Epic Games, Inc. All Rights Reserved.
+
+#pragma once
+
+#include "Widgets/DeclarativeSyntaxSupport.h"
+#include "Widgets/SCompoundWidget.h"
+
+class SConstrainedBox : public SCompoundWidget
+{
+public:
+	SLATE_BEGIN_ARGS(SConstrainedBox)
+		: _MinWidth()
+		, _MaxWidth()
+	{}
+	SLATE_DEFAULT_SLOT(FArguments, Content)
+		SLATE_ATTRIBUTE(TOptional<float>, MinWidth)
+		SLATE_ATTRIBUTE(TOptional<float>, MaxWidth)
+	SLATE_END_ARGS()
+
+	void Construct(const FArguments& InArgs);
+	virtual FVector2D ComputeDesiredSize(float LayoutScaleMultiplier) const override;
+
+private:
+	TAttribute< TOptional<float> > MinWidth;
+	TAttribute< TOptional<float> > MaxWidth;
+};
diff --git a/Engine/Source/Editor/PropertyEditor/Private/SDetailExpanderArrow.cpp b/Engine/Source/Editor/PropertyEditor/Private/SDetailExpanderArrow.cpp
new file mode 100644
--- /dev/null
+++ b/Engine/Source/Editor/PropertyEditor/Private/SDetailExpanderArrow.cpp
@@ -0,0 +1,90 @@
+// Copyright Epic Games, Inc. All Rights Reserved.
+
+#include "SDetailExpanderArrow.h"
+#include "SDetailTableRowBase.h"
+#include "SConstrainedBox.h"
+#include "Widgets/Images/SImage.h"
+#include "Widgets/Input/SButton.h"
+
+void SDetailExpanderArrow::Construct(const FArguments& InArgs, TSharedRef<SDetailTableRowBase> DetailsRow)
+{
+	Row = DetailsRow;
+
+	ChildSlot
+	[
+		SNew(SConstrainedBox)
+		.MinWidth(20)
+		.Visibility(this, &SDetailExpanderArrow::GetExpanderVisibility)
+		[
+			SAssignNew(ExpanderArrow, SButton)
+			.ButtonStyle(FCoreStyle::Get(), "NoBorder")
+			.VAlign(VAlign_Center)
+			.HAlign(HAlign_Center)
+			.ClickMethod(EButtonClickMethod::MouseDown)
+			.OnClicked(this, &SDetailExpanderArrow::OnExpanderClicked)
+			.ContentPadding(FMargin(5,0,0,0))
+			.IsFocusable(false)
+			[
+				SNew(SImage)
+				.Image(this, &SDetailExpanderArrow::GetExpanderImage)
+				.ColorAndOpacity(FSlateColor::UseForeground())
+			]
+		]
+	];
+}
+
+EVisibility SDetailExpanderArrow::GetExpanderVisibility() const
+{
+	return Row->DoesItemHaveChildren() ? EVisibility::Visible : EVisibility::Collapsed;
+}
+
+const FSlateBrush* SDetailExpanderArrow::GetExpanderImage() const
+{
+	const bool bIsItemExpanded = Row->IsItemExpanded();
+
+	FName ResourceName;
+	if (bIsItemExpanded)
+	{
+		if (ExpanderArrow->IsHovered())
+		{
+			static FName ExpandedHoveredName = "TreeArrow_Expanded_Hovered";
+			ResourceName = ExpandedHoveredName;
+		}
+		else
+		{
+			static FName ExpandedName = "TreeArrow_Expanded";
+			ResourceName = ExpandedName;
+		}
+	}
+	else
+	{
+		if (ExpanderArrow->IsHovered())
+		{
+			static FName CollapsedHoveredName = "TreeArrow_Collapsed_Hovered";
+			ResourceName = CollapsedHoveredName;
+		}
+		else
+		{
+			static FName CollapsedName = "TreeArrow_Collapsed";
+			ResourceName = CollapsedName;
+		}
+	}
+
+	return FAppStyle::Get().GetBrush(ResourceName);
+}
+
+FReply SDetailExpanderArrow::OnExpanderClicked()
+{
+	// Recurse the expansion if "shift" is being pressed
+	const FModifierKeysState ModKeyState = FSlateApplication::Get().GetModifierKeys();
+	if (ModKeyState.IsShiftDown())
+	{
+		Row->Private_OnExpanderArrowShiftClicked();
+	}
+	else
+	{
+		Row->ToggleExpansion();
+	}
+
+	return FReply::Handled();
+}
\ No newline at end of file
diff --git a/Engine/Source/Editor/PropertyEditor/Private/SDetailExpanderArrow.h b/Engine/Source/Editor/PropertyEditor/Private/SDetailExpanderArrow.h
new file mode 100644
--- /dev/null
+++ b/Engine/Source/Editor/PropertyEditor/Private/SDetailExpanderArrow.h
@@ -0,0 +1,28 @@
+// Copyright Epic Games, Inc. All Rights Reserved.
+
+#pragma once
+
+#include "Widgets/DeclarativeSyntaxSupport.h"
+#include "Widgets/SCompoundWidget.h"
+
+class SButton;
+class SDetailTableRowBase;
+
+class SDetailExpanderArrow : public SCompoundWidget
+{
+public:
+	SLATE_BEGIN_ARGS(SDetailExpanderArrow) {}
+	SLATE_END_ARGS()
+
+	void Construct(const FArguments& InArgs, TSharedRef<SDetailTableRowBase> DetailsRow);
+
+private:
+
+	EVisibility GetExpanderVisibility() const;
+	const FSlateBrush* GetExpanderImage() const;
+	FReply OnExpanderClicked();
+	
+private:
+	TSharedPtr<SDetailTableRowBase> Row;
+	TSharedPtr<SButton> ExpanderArrow;
+};
\ No newline at end of file
diff --git a/Engine/Source/Editor/PropertyEditor/Private/SDetailRowIndent.cpp b/Engine/Source/Editor/PropertyEditor/Private/SDetailRowIndent.cpp
new file mode 100644
--- /dev/null
+++ b/Engine/Source/Editor/PropertyEditor/Private/SDetailRowIndent.cpp
@@ -0,0 +1,64 @@
+// Copyright Epic Games, Inc. All Rights Reserved.
+
+#include "SDetailRowIndent.h"
+#include "SConstrainedBox.h"
+#include "SDetailTableRowBase.h"
+#include "PropertyEditorConstants.h"
+#include "Widgets/Layout/SBox.h"
+
+void SDetailRowIndent::Construct(const FArguments& InArgs, TSharedRef<SDetailTableRowBase> DetailsRow)
+{
+	Row = DetailsRow;
+
+	ChildSlot
+	[
+		SNew(SBox)
+		.WidthOverride(this, &SDetailRowIndent::GetIndentWidth)
+	];
+}
+
+int32 SDetailRowIndent::OnPaint(const FPaintArgs& Args, const FGeometry& AllottedGeometry, const FSlateRect& MyCullingRect, FSlateWindowElementList& OutDrawElements, int32 LayerId, const FWidgetStyle& InWidgetStyle, bool bParentEnabled) const
+{
+	const FSlateBrush* BackgroundBrush = FAppStyle::Get().GetBrush("DetailsView.CategoryMiddle");
+	const FSlateBrush* DropShadowBrush = FAppStyle::Get().GetBrush("DetailsView.ArrayDropShadow");
+
+	int32 IndentLevel = Row->GetIndentLevelForBackgroundColor();
+	for (int32 i = 0; i < IndentLevel; ++i)
+	{
+		FSlateColor BackgroundColor = GetRowBackgroundColor(i);
+
+		FSlateDrawElement::MakeBox(
+			OutDrawElements,
+			LayerId,
+			AllottedGeometry.ToPaintGeometry(FVector2D(16 * i, 0), FVector2D(16, AllottedGeometry.GetLocalSize().Y)),
+			BackgroundBrush,
+			ESlateDrawEffect::None,
+			BackgroundColor.GetColor(InWidgetStyle)
+		);
+
+		FSlateDrawElement::MakeBox(
+			OutDrawElements,
+			LayerId + 1,
+			AllottedGeometry.ToPaintGeometry(FVector2D(16 * i, 0), FVector2D(16, AllottedGeometry.GetLocalSize().Y)),
+			DropShadowBrush
+		);
+	}
+		
+	return LayerId + 1;
+}
+
+FOptionalSize SDetailRowIndent::GetIndentWidth() const
+{
+	int32 IndentLevel = Row->GetIndentLevelForBackgroundColor();
+	return IndentLevel * 16.0f;
+}
+
+FSlateColor SDetailRowIndent::GetRowBackgroundColor(int32 i) const
+{
+	if (Row->IsHovered())
+	{
+		return FAppStyle::Get().GetSlateColor("Colors.Header");
+	}
+
+	return PropertyEditorConstants::GetRowBackgroundColor(i);
+}
diff --git a/Engine/Source/Editor/PropertyEditor/Private/SDetailRowIndent.h b/Engine/Source/Editor/PropertyEditor/Private/SDetailRowIndent.h
new file mode 100644
--- /dev/null
+++ b/Engine/Source/Editor/PropertyEditor/Private/SDetailRowIndent.h
@@ -0,0 +1,28 @@
+// Copyright Epic Games, Inc. All Rights Reserved.
+
+#pragma once
+
+#include "Styling/SlateColor.h"
+#include "Types/SlateStructs.h"
+#include "Widgets/DeclarativeSyntaxSupport.h"
+#include "Widgets/SCompoundWidget.h"
+
+class SDetailTableRowBase;
+class SHorizontalBox;
+
+class SDetailRowIndent : public SCompoundWidget
+{ 
+public:
+	SLATE_BEGIN_ARGS(SDetailRowIndent) {}
+	SLATE_END_ARGS()
+
+	void Construct(const FArguments& InArgs, TSharedRef<SDetailTableRowBase> DetailsRow);
+
+private:
+	int32 OnPaint(const FPaintArgs& Args, const FGeometry& AllottedGeometry, const FSlateRect& MyCullingRect, FSlateWindowElementList& OutDrawElements, int32 LayerId, const FWidgetStyle& InWidgetStyle, bool bParentEnabled) const;
+	FOptionalSize GetIndentWidth() const;
+	FSlateColor GetRowBackgroundColor(int32 i) const;
+
+private:
+	TSharedPtr<SDetailTableRowBase> Row;
+};
diff --git a/Engine/Source/Editor/PropertyEditor/Private/SDetailSingleItemRow.cpp b/Engine/Source/Editor/PropertyEditor/Private/SDetailSingleItemRow.cpp
--- a/Engine/Source/Editor/PropertyEditor/Private/SDetailSingleItemRow.cpp
+++ b/Engine/Source/Editor/PropertyEditor/Private/SDetailSingleItemRow.cpp
@@ -16,6 +16,9 @@
 #include "PropertyEditorModule.h"
 #include "Modules/ModuleInterface.h"
 #include "Modules/ModuleManager.h"
+#include "SConstrainedBox.h"
+#include "SDetailExpanderArrow.h"
+#include "SDetailRowIndent.h"
 #include "SResetToDefaultPropertyEditor.h"
 #include "PropertyEditorConstants.h"
 
@@ -47,109 +50,6 @@ namespace SDetailSingleItemRow_Helper
 	}
 }
 
-class SDetailsIndent : public SCompoundWidget
-{ 
-private:
-	TSharedPtr<SDetailSingleItemRow> Row;
-	TSharedPtr<SHorizontalBox> ChildBox;
-
-public:
-
-	SLATE_BEGIN_ARGS(SDetailsIndent) {}
-	SLATE_END_ARGS()
-
-	void Construct(const FArguments& InArgs, TSharedRef<SDetailSingleItemRow> DetailsRow)
-	{
-		Row = DetailsRow;
-
-		ChildSlot
-		[
-			SNew(SBox)
-			.WidthOverride(this, &SDetailsIndent::GetIndentWidth)
-		];
-	}
-
-	int32 OnPaint(const FPaintArgs& Args, const FGeometry& AllottedGeometry, const FSlateRect& MyCullingRect, FSlateWindowElementList& OutDrawElements, int32 LayerId, const FWidgetStyle& InWidgetStyle, bool bParentEnabled) const
-	{
-		const FSlateBrush* BackgroundBrush = FAppStyle::Get().GetBrush("DetailsView.CategoryMiddle");
-		const FSlateBrush* DropShadowBrush = FAppStyle::Get().GetBrush("DetailsView.ArrayDropShadow");
-
-		int32 IndentLevel = Row->GetIndentLevelForBackgroundColor();
-		for (int32 i = 0; i < IndentLevel; ++i)
-		{
-			FSlateColor BackgroundColor = GetRowBackgroundColor(i);
-
-			FSlateDrawElement::MakeBox(
-				OutDrawElements,
-				LayerId,
-				AllottedGeometry.ToPaintGeometry(FVector2D(16 * i, 0), FVector2D(16, AllottedGeometry.GetLocalSize().Y)),
-				BackgroundBrush,
-				ESlateDrawEffect::None,
-				BackgroundColor.GetColor(InWidgetStyle)
-			);
-
-			FSlateDrawElement::MakeBox(
-				OutDrawElements,
-				LayerId + 1,
-				AllottedGeometry.ToPaintGeometry(FVector2D(16 * i, 0), FVector2D(16, AllottedGeometry.GetLocalSize().Y)),
-				DropShadowBrush
-			);
-		}
-		
-		return LayerId + 1;
-	}
-
-	FOptionalSize GetIndentWidth() const
-	{
-		int32 IndentLevel = Row->GetIndentLevelForBackgroundColor();
-		return IndentLevel * 16.0f;
-	}
-
-	FSlateColor GetRowBackgroundColor(int32 i) const
-	{
-		if (Row->IsHovered() && !Row->bIsHoveredDragTarget)
-		{
-			return FAppStyle::Get().GetSlateColor("Colors.Header");
-		}
-
-		return PropertyEditorConstants::GetRowBackgroundColor(i);
-	}
-};
-
-void SConstrainedBox::Construct(const FArguments& InArgs)
-{
-	MinWidth = InArgs._MinWidth;
-	MaxWidth = InArgs._MaxWidth;
-
-	ChildSlot
-	[
-		InArgs._Content.Widget
-	];
-}
-
-FVector2D SConstrainedBox::ComputeDesiredSize(float LayoutScaleMultiplier) const
-{
-	const float MinWidthVal = MinWidth.Get().Get(0.0f);
-	const float MaxWidthVal = MaxWidth.Get().Get(0.0f);
-
-	if (MinWidthVal == 0.0f && MaxWidthVal == 0.0f)
-	{
-		return SCompoundWidget::ComputeDesiredSize(LayoutScaleMultiplier);
-	}
-	else
-	{
-		FVector2D ChildSize = ChildSlot.GetWidget()->GetDesiredSize();
-
-		float XVal = FMath::Max(MinWidthVal, ChildSize.X);
-		if (MaxWidthVal >= MinWidthVal)
-		{
-			XVal = FMath::Min(MaxWidthVal, XVal);
-		}
-
-		return FVector2D(XVal, ChildSize.Y);
-	}
-}
-
 void SDetailSingleItemRow::OnFavoriteMenuToggle()
 {
 	if (Customization->GetPropertyNode().IsValid() && Customization->GetPropertyNode()->CanDisplayFavorite())
@@ -501,7 +401,7 @@ void SDetailSingleItemRow::Construct( const FArguments& InArgs, FDetailLayoutCus
 				.Padding(0,0,0,0)
 				.AutoWidth()
 				[
-					SNew(SDetailsIndent, SharedThis(this))
+					SNew(SDetailRowIndent, SharedThis(this))
 				];
 
 			NameColumnBox->AddSlot()
@@ -510,24 +410,7 @@ void SDetailSingleItemRow::Construct( const FArguments& InArgs, FDetailLayoutCus
 				.Padding(0)
 				.AutoWidth()
 				[
-					SNew(SConstrainedBox)
-					.MinWidth(20)
-					.Visibility(this, &SDetailSingleItemRow::GetExpanderVisibility)
-					[
-						SAssignNew(ExpanderArrow, SButton)
-						.ButtonStyle(FCoreStyle::Get(), "NoBorder")
-						.VAlign(VAlign_Center)
-						.HAlign(HAlign_Center)
-						.ClickMethod(EButtonClickMethod::MouseDown)
-						.OnClicked(this, &SDetailSingleItemRow::OnExpanderClicked)
-						.ContentPadding(FMargin(5,0,0,0))
-						.IsFocusable(false)
-						[
-							SNew(SImage)
-							.Image(this, &SDetailSingleItemRow::GetExpanderImage)
-							.ColorAndOpacity(FSlateColor::UseForeground())
-						]
-					]
+					SNew(SDetailExpanderArrow, SharedThis(this))
 				];
 			
 			TSharedPtr<FPropertyNode> PropertyNode = Customization->GetPropertyNode();
@@ -789,27 +672,6 @@ FSlateColor SDetailSingleItemRow::GetInnerBackgroundColor() const
 	return PropertyEditorConstants::GetRowBackgroundColor(IndentLevel);
 }
 
-int32 SDetailSingleItemRow::GetIndentLevelForBackgroundColor() const
-{
-	int32 IndentLevel = 0; 
-	if (OwnerTablePtr.IsValid())
-	{
-		// every item is in a category, but we don't want to show an indent for "top-level" properties
-		IndentLevel = GetIndentLevel() - 1;
-	}
-
-	TSharedPtr<FDetailTreeNode> DetailTreeNode = GetOwnerTreeNode();
-	if (DetailTreeNode.IsValid() && 
-		DetailTreeNode->GetDetailsView() != nullptr && 
-		DetailTreeNode->GetDetailsView()->ContainsMultipleTopLevelObjects())
-	{
-		// if the row is in a multiple top level object display (eg. Project Settings), don't display an indent for the initial level
-		--IndentLevel;
-	}
-
-	return FMath::Max(0, IndentLevel);
-}
-
 bool SDetailSingleItemRow::OnContextMenuOpening(FMenuBuilder& MenuBuilder)
 {
 	const bool bIsCopyPasteBound = Customization->GetWidgetRow().IsCopyPasteBound();
@@ -912,62 +774,6 @@ bool SDetailSingleItemRow::OnContextMenuOpening(FMenuBuilder& MenuBuilder)
 	return true;
 }
 
-EVisibility SDetailSingleItemRow::GetExpanderVisibility() const
-{
-	return DoesItemHaveChildren() ? EVisibility::Visible : EVisibility::Collapsed;
-}
-
-const FSlateBrush* SDetailSingleItemRow::GetExpanderImage() const
-{
-	const bool bIsItemExpanded = IsItemExpanded();
-
-	FName ResourceName;
-	if (bIsItemExpanded)
-	{
-		if (ExpanderArrow->IsHovered())
-		{
-			static FName ExpandedHoveredName = "TreeArrow_Expanded_Hovered";
-			ResourceName = ExpandedHoveredName;
-		}
-		else
-		{
-			static FName ExpandedName = "TreeArrow_Expanded";
-			ResourceName = ExpandedName;
-		}
-	}
-	else
-	{
-		if (ExpanderArrow->IsHovered())
-		{
-			static FName CollapsedHoveredName = "TreeArrow_Collapsed_Hovered";
-			ResourceName = CollapsedHoveredName;
-		}
-		else
-		{
-			static FName CollapsedName = "TreeArrow_Collapsed";
-			ResourceName = CollapsedName;
-		}
-	}
-
-	return FAppStyle::Get().GetBrush(ResourceName);
-}
-
-FReply SDetailSingleItemRow::OnExpanderClicked()
-{
-	// Recurse the expansion if "shift" is being pressed
-	const FModifierKeysState ModKeyState = FSlateApplication::Get().GetModifierKeys();
-	if(ModKeyState.IsShiftDown())
-	{
-		Private_OnExpanderArrowShiftClicked();
-	}
-	else
-	{
-		ToggleExpansion();
-	}
-
-	return FReply::Handled();
-}
-
 void SDetailSingleItemRow::OnCopyProperty()
 {
 	if (OwnerTreeNode.IsValid())
diff --git a/Engine/Source/Editor/PropertyEditor/Private/SDetailSingleItemRow.h b/Engine/Source/Editor/PropertyEditor/Private/SDetailSingleItemRow.h
--- a/Engine/Source/Editor/PropertyEditor/Private/SDetailSingleItemRow.h
+++ b/Engine/Source/Editor/PropertyEditor/Private/SDetailSingleItemRow.h
@@ -18,26 +18,6 @@ class IDetailKeyframeHandler;
 struct FDetailLayoutCustomization;
 class SDetailSingleItemRow;
 
-class SConstrainedBox : public SCompoundWidget
-{
-public:
-	SLATE_BEGIN_ARGS(SConstrainedBox)
-		: _MinWidth()
-		, _MaxWidth()
-	{}
-	SLATE_DEFAULT_SLOT(FArguments, Content)
-		SLATE_ATTRIBUTE(TOptional<float>, MinWidth)
-		SLATE_ATTRIBUTE(TOptional<float>, MaxWidth)
-	SLATE_END_ARGS()
-
-	void Construct(const FArguments& InArgs);
-	virtual FVector2D ComputeDesiredSize(float LayoutScaleMultiplier) const override;
-
-private:
-	TAttribute< TOptional<float> > MinWidth;
-	TAttribute< TOptional<float> > MaxWidth;
-};
-
 class SArrayRowHandle : public SCompoundWidget
 {
 public:
@@ -114,12 +94,6 @@ private:
 	TSharedPtr<FPropertyNode> GetPropertyNode() const;
 	TSharedPtr<IPropertyHandle> GetPropertyHandle() const;
 
-	EVisibility GetExpanderVisibility() const;
-	const FSlateBrush* GetExpanderImage() const;
-	FReply OnExpanderClicked();
-
-	int32 GetIndentLevelForBackgroundColor() const;
-
 private:
 	TWeakPtr<IDetailKeyframeHandler> KeyframeHandler;
 	/** Customization for this widget */
@@ -129,8 +103,6 @@ private:
 	bool bIsDragDropObject;
 	TSharedPtr<FPropertyNode> SwappablePropertyNode;
 	TSharedPtr<SButton> ExpanderArrow;
-
-	friend class SDetailsIndent;
 };
 
 class FArrayRowDragDropOp : public FDecoratedDragDropOp
diff --git a/Engine/Source/Editor/PropertyEditor/Private/SDetailTableRowBase.cpp b/Engine/Source/Editor/PropertyEditor/Private/SDetailTableRowBase.cpp
--- a/Engine/Source/Editor/PropertyEditor/Private/SDetailTableRowBase.cpp
+++ b/Engine/Source/Editor/PropertyEditor/Private/SDetailTableRowBase.cpp
@@ -3,3 +3,24 @@
 #include "SDetailTableRowBase.h"
 
 float SDetailTableRowBase::ScrollbarPaddingSize = 16.0f;
+
+int32 SDetailTableRowBase::GetIndentLevelForBackgroundColor() const
+{
+	int32 IndentLevel = 0; 
+	if (OwnerTablePtr.IsValid())
+	{
+		// every item is in a category, but we don't want to show an indent for "top-level" properties
+		IndentLevel = GetIndentLevel() - 1;
+	}
+
+	TSharedPtr<FDetailTreeNode> DetailTreeNode = OwnerTreeNode.Pin();
+	if (DetailTreeNode.IsValid() && 
+		DetailTreeNode->GetDetailsView() != nullptr && 
+		DetailTreeNode->GetDetailsView()->ContainsMultipleTopLevelObjects())
+	{
+		// if the row is in a multiple top level object display (eg. Project Settings), don't display an indent for the initial level
+		--IndentLevel;
+	}
+
+	return FMath::Max(0, IndentLevel);
+}
\ No newline at end of file
diff --git a/Engine/Source/Editor/PropertyEditor/Private/SDetailTableRowBase.h b/Engine/Source/Editor/PropertyEditor/Private/SDetailTableRowBase.h
--- a/Engine/Source/Editor/PropertyEditor/Private/SDetailTableRowBase.h
+++ b/Engine/Source/Editor/PropertyEditor/Private/SDetailTableRowBase.h
@@ -60,7 +60,7 @@ public:
 		return STableRow< TSharedPtr< FDetailTreeNode > >::OnMouseButtonUp( MyGeometry, MouseEvent );
 	}
 
-	TSharedPtr<FDetailTreeNode> GetOwnerTreeNode() const { return OwnerTreeNode.Pin(); }
+	int32 GetIndentLevelForBackgroundColor() const;
 
 protected:
 	/**
