Detail property rows that had a custom widget defined for them now correctly call the OverrideResetToDefault delegate given to the custom widget row.

FPropertyHandleImpl also no longer causes an infinite loop of deferred actions when a custom reset to default is set on a property row.

diff --git a/Engine/Source/Editor/PropertyEditor/Private/DetailPropertyRow.cpp b/Engine/Source/Editor/PropertyEditor/Private/DetailPropertyRow.cpp
--- a/Engine/Source/Editor/PropertyEditor/Private/DetailPropertyRow.cpp
+++ b/Engine/Source/Editor/PropertyEditor/Private/DetailPropertyRow.cpp
@@ -15,9 +15,6 @@
 #include "CategoryPropertyNode.h"
 #include "Widgets/Layout/SSpacer.h"
 
-const float FDetailWidgetRow::DefaultValueMinWidth = 125.0f;
-const float FDetailWidgetRow::DefaultValueMaxWidth = 125.0f;
-
 #define LOCTEXT_NAMESPACE	"DetailPropertyRow"
 
 FDetailPropertyRow::FDetailPropertyRow(TSharedPtr<FPropertyNode> InPropertyNode, TSharedRef<FDetailCategoryImpl> InParentCategory, TSharedPtr<FComplexPropertyNode> InExternalRootNode)
@@ -221,6 +218,7 @@ FDetailWidgetRow FDetailPropertyRow::GetWidgetRow()
 			Row.CopyMenuAction = CustomPropertyWidget->CopyMenuAction;
 			Row.PasteMenuAction = CustomPropertyWidget->PasteMenuAction;
 			Row.CustomMenuItems = CustomPropertyWidget->CustomMenuItems;
+			Row.CustomResetToDefault = CustomPropertyWidget->CustomResetToDefault;
 		}
 		
 		return Row;
diff --git a/Engine/Source/Editor/PropertyEditor/Private/DetailWidgetRow.cpp b/Engine/Source/Editor/PropertyEditor/Private/DetailWidgetRow.cpp
--- a/Engine/Source/Editor/PropertyEditor/Private/DetailWidgetRow.cpp
+++ b/Engine/Source/Editor/PropertyEditor/Private/DetailWidgetRow.cpp
@@ -3,9 +3,5 @@
 #include "DetailWidgetRow.h"
 #include "PropertyEditorHelpers.h"
 
-FDetailWidgetRow& FDetailWidgetRow::EditCondition(TAttribute<bool> InEditConditionValue, FOnBooleanValueChanged InOnEditConditionValueChanged)
-{
-	EditConditionValue = InEditConditionValue;
-	OnEditConditionValueChanged = InOnEditConditionValueChanged;
-	return *this;
-}
\ No newline at end of file
+const float FDetailWidgetRow::DefaultValueMinWidth = 125.0f;
+const float FDetailWidgetRow::DefaultValueMaxWidth = 125.0f;
diff --git a/Engine/Source/Editor/PropertyEditor/Private/PropertyHandleImpl.cpp b/Engine/Source/Editor/PropertyEditor/Private/PropertyHandleImpl.cpp
--- a/Engine/Source/Editor/PropertyEditor/Private/PropertyHandleImpl.cpp
+++ b/Engine/Source/Editor/PropertyEditor/Private/PropertyHandleImpl.cpp
@@ -2502,7 +2502,9 @@ void FPropertyHandleBase::OnCustomResetToDefault(const FResetToDefaultOverride&
 		{
 			Implementation->GetPropertyNode()->NotifyPreChange(Implementation->GetPropertyNode()->GetProperty(), Implementation->GetPropertyUtilities()->GetNotifyHook());
 		}
-		CustomResetToDefault.OnResetToDefaultClicked(SharedThis(this));
+
+		FResetToDefaultHandler Delegate = CustomResetToDefault.GetPropertyResetToDefaultDelegate();
+		Delegate.ExecuteIfBound(SharedThis(this));
 
 		// Call PostEditchange on all the objects
 		FPropertyChangedEvent ChangeEvent(Implementation->GetPropertyNode()->GetProperty());
diff --git a/Engine/Source/Editor/PropertyEditor/Private/UserInterface/PropertyEditor/SResetToDefaultPropertyEditor.cpp b/Engine/Source/Editor/PropertyEditor/Private/UserInterface/PropertyEditor/SResetToDefaultPropertyEditor.cpp
--- a/Engine/Source/Editor/PropertyEditor/Private/UserInterface/PropertyEditor/SResetToDefaultPropertyEditor.cpp
+++ b/Engine/Source/Editor/PropertyEditor/Private/UserInterface/PropertyEditor/SResetToDefaultPropertyEditor.cpp
@@ -7,7 +7,7 @@
 
 SResetToDefaultPropertyEditor::~SResetToDefaultPropertyEditor()
 {
-	if (PropertyHandle.IsValid())
+	if (PropertyHandle.IsValid() && OptionalCustomResetToDefault.IsSet())
 	{
 		PropertyHandle->ClearResetToDefaultCustomized();
 	}
diff --git a/Engine/Source/Editor/PropertyEditor/Public/DetailWidgetRow.h b/Engine/Source/Editor/PropertyEditor/Public/DetailWidgetRow.h
--- a/Engine/Source/Editor/PropertyEditor/Public/DetailWidgetRow.h
+++ b/Engine/Source/Editor/PropertyEditor/Public/DetailWidgetRow.h
@@ -249,7 +249,12 @@ public:
 	/** 
 	 * Override the edit condition.
 	 */ 
-	FDetailWidgetRow& EditCondition(TAttribute<bool> EditConditionValue, FOnBooleanValueChanged OnEditConditionValueChanged);
+	FDetailWidgetRow& EditCondition(TAttribute<bool> InEditConditionValue, FOnBooleanValueChanged InOnEditConditionValueChanged)
+	{
+		EditConditionValue = InEditConditionValue;
+		OnEditConditionValueChanged = InOnEditConditionValueChanged;
+		return *this;
+	}
 	
 	/**
 	* Used to provide all the property handles this WidgetRow represent
diff --git a/Engine/Source/Editor/PropertyEditor/Public/IDetailPropertyRow.h b/Engine/Source/Editor/PropertyEditor/Public/IDetailPropertyRow.h
--- a/Engine/Source/Editor/PropertyEditor/Public/IDetailPropertyRow.h
+++ b/Engine/Source/Editor/PropertyEditor/Public/IDetailPropertyRow.h
@@ -188,6 +188,9 @@ public:
 		}
 	}
 
+	/** Get the actual delegate bound to this reset to default handler. */
+	FResetToDefaultHandler GetPropertyResetToDefaultDelegate() const { return OnClickedPropertyDelegate; }
+
 	/** Called by properties to determine whether this override should set on their children */
 	bool PropagatesToChildren() const
 	{
