FPropertyNode now propagates the EditInlineNew flag to children of containers that have the flag. This is necessary because containers of EditInlineNew objects would display differently from single items.

diff --git a/Engine/Source/Editor/PropertyEditor/Private/PropertyNode.cpp b/Engine/Source/Editor/PropertyEditor/Private/PropertyNode.cpp
--- a/Engine/Source/Editor/PropertyEditor/Private/PropertyNode.cpp
+++ b/Engine/Source/Editor/PropertyEditor/Private/PropertyNode.cpp
@@ -162,23 +162,27 @@ void FPropertyNode::InitNode(const FPropertyNodeInitParams& InitParams)
 		const bool bSingleSelectOnly = GetReadAddressUncached( *this, true, nullptr);
 		SetNodeFlags(EPropertyNodeFlags::SingleSelectOnly, bSingleSelectOnly);
 
-		FProperty* MyProperty = Property.Get();
+		const FProperty* MyProperty = Property.Get();
+		const FProperty* OwnerProperty = MyProperty->GetOwnerProperty();
 
 		const bool bIsObjectOrInterface = CastField<FObjectPropertyBase>(MyProperty) || CastField<FInterfaceProperty>(MyProperty);
+		const bool bIsInsideContainer = CastField<FArrayProperty>(OwnerProperty) || CastField<FSetProperty>(OwnerProperty) || CastField<FMapProperty>(OwnerProperty);
 
 		// true if the property can be expanded into the property window; that is, instead of seeing
 		// a pointer to the object, you see the object's properties.
 		static const FName Name_EditInline("EditInline");
 		static const FName Name_ShowInnerProperties("ShowInnerProperties");
 
-		bIsEditInlineNew = bIsObjectOrInterface && GotReadAddresses && MyProperty->HasMetaData(Name_EditInline);
+		// we are EditInlineNew if this property has the flag, or if inside a container that has the flag.
+		bIsEditInlineNew = GotReadAddresses && bIsObjectOrInterface &&
+			(MyProperty->HasMetaData(Name_EditInline)) || (bIsInsideContainer && OwnerProperty->HasMetaData(Name_EditInline));
 		bShowInnerObjectProperties = bIsObjectOrInterface && MyProperty->HasMetaData(Name_ShowInnerProperties);
 
-		if(bIsEditInlineNew)
+		if (bIsEditInlineNew)
 		{
 			SetNodeFlags(EPropertyNodeFlags::EditInlineNew, true);
 		}
-		else if(bShowInnerObjectProperties)
+		else if (bShowInnerObjectProperties)
 		{
 			SetNodeFlags(EPropertyNodeFlags::ShowInnerObjectProperties, true);
 		}
@@ -203,13 +207,7 @@ void FPropertyNode::InitNode(const FPropertyNodeInitParams& InitParams)
 				}
 			}
 		}
-	}
-
-	InitExpansionFlags();
 
-	FProperty* MyProperty = Property.Get();
-	if (MyProperty)
-	{
 		const FString& EditConditionString = MyProperty->GetMetaData(TEXT("EditCondition"));
 
 		// see if the property supports some kind of edit condition and this isn't the "parent" property of a static array
@@ -222,14 +220,19 @@ void FPropertyNode::InitNode(const FPropertyNodeInitParams& InitParams)
 				EditConditionContext = MakeShareable(new FEditConditionContext(*this));
 			}
 		}
-	}
+		
+		bool bRequiresValidation = bIsEditInlineNew || bShowInnerObjectProperties;
+	
+		// We require validation if we are in a container.
+		bRequiresValidation |= MyProperty && (MyProperty->IsA<FArrayProperty>() || MyProperty->IsA<FSetProperty>() || MyProperty->IsA<FMapProperty>());
 
-	bool bRequiresValidation = bIsEditInlineNew || bShowInnerObjectProperties || ( MyProperty && (MyProperty->IsA<FArrayProperty>() || MyProperty->IsA<FSetProperty>() || MyProperty->IsA<FMapProperty>() ));
+		// We require validation if our parent also needs validation (if an array parent was resized all the addresses of children are invalid)
+		bRequiresValidation |= (GetParentNode() && GetParentNode()->HasNodeFlags( EPropertyNodeFlags::RequiresValidation ) != 0);
 
-	// We require validation if our parent also needs validation (if an array parent was resized all the addresses of children are invalid)
-	bRequiresValidation |= (GetParentNode() && GetParentNode()->HasNodeFlags( EPropertyNodeFlags::RequiresValidation ) != 0);
+		SetNodeFlags( EPropertyNodeFlags::RequiresValidation, bRequiresValidation );
+	}
 
-	SetNodeFlags( EPropertyNodeFlags::RequiresValidation, bRequiresValidation );
+	InitExpansionFlags();
 
 	if (InitParams.bAllowChildren)
 	{
diff --git a/Engine/Source/Editor/PropertyEditor/Private/UserInterface/PropertyEditor/SPropertyEditorEditInline.cpp b/Engine/Source/Editor/PropertyEditor/Private/UserInterface/PropertyEditor/SPropertyEditorEditInline.cpp
--- a/Engine/Source/Editor/PropertyEditor/Private/UserInterface/PropertyEditor/SPropertyEditorEditInline.cpp
+++ b/Engine/Source/Editor/PropertyEditor/Private/UserInterface/PropertyEditor/SPropertyEditorEditInline.cpp
@@ -147,8 +147,7 @@ bool SPropertyEditorEditInline::Supports( const FPropertyNode* InTreeNode, int32
 {
 	return InTreeNode
 		&& InTreeNode->HasNodeFlags(EPropertyNodeFlags::EditInlineNew)
-		&& InTreeNode->FindObjectItemParent()
-		&& !InTreeNode->IsPropertyConst();
+		&& InTreeNode->FindObjectItemParent();
 }
 
 bool SPropertyEditorEditInline::Supports( const TSharedRef< class FPropertyEditor >& InPropertyEditor )
