Fixing ambiguous boolean conditional

diff --git a/Engine/Source/Editor/PropertyEditor/Private/PropertyNode.cpp b/Engine/Source/Editor/PropertyEditor/Private/PropertyNode.cpp
--- a/Engine/Source/Editor/PropertyEditor/Private/PropertyNode.cpp
+++ b/Engine/Source/Editor/PropertyEditor/Private/PropertyNode.cpp
@@ -175,7 +175,7 @@ void FPropertyNode::InitNode(const FPropertyNodeInitParams& InitParams)
 
 		// we are EditInlineNew if this property has the flag, or if inside a container that has the flag.
 		bIsEditInlineNew = GotReadAddresses && bIsObjectOrInterface &&
-			(MyProperty->HasMetaData(Name_EditInline)) || (bIsInsideContainer && OwnerProperty->HasMetaData(Name_EditInline));
+			(MyProperty->HasMetaData(Name_EditInline) || (bIsInsideContainer && OwnerProperty->HasMetaData(Name_EditInline)));
 		bShowInnerObjectProperties = bIsObjectOrInterface && MyProperty->HasMetaData(Name_ShowInnerProperties);
 
 		if (bIsEditInlineNew)
