Rolled back change that caused details panel to inline too many things.

diff --git a/Engine/Source/Editor/PropertyEditor/Private/UserInterface/PropertyEditor/SPropertyEditorEditInline.cpp b/Engine/Source/Editor/PropertyEditor/Private/UserInterface/PropertyEditor/SPropertyEditorEditInline.cpp
--- a/Engine/Source/Editor/PropertyEditor/Private/UserInterface/PropertyEditor/SPropertyEditorEditInline.cpp
+++ b/Engine/Source/Editor/PropertyEditor/Private/UserInterface/PropertyEditor/SPropertyEditorEditInline.cpp
@@ -147,7 +147,8 @@ bool SPropertyEditorEditInline::Supports( const FPropertyNode* InTreeNode, int32
 {
 	return InTreeNode
 		&& InTreeNode->HasNodeFlags(EPropertyNodeFlags::EditInlineNew)
-		&& InTreeNode->FindObjectItemParent();
+		&& InTreeNode->FindObjectItemParent()
+		&& !InTreeNode->IsPropertyConst();
 }
 
 bool SPropertyEditorEditInline::Supports( const TSharedRef< class FPropertyEditor >& InPropertyEditor )
