Fixed static analysis issue in PropertyNode.cpp.

diff --git a/Engine/Source/Editor/PropertyEditor/Private/PropertyNode.cpp b/Engine/Source/Editor/PropertyEditor/Private/PropertyNode.cpp
--- a/Engine/Source/Editor/PropertyEditor/Private/PropertyNode.cpp
+++ b/Engine/Source/Editor/PropertyEditor/Private/PropertyNode.cpp
@@ -151,7 +151,8 @@ void FPropertyNode::InitNode(const FPropertyNodeInitParams& InitParams)
 
 	bool bIsEditInlineNew = false;
 	bool bShowInnerObjectProperties = false;
-	if ( !Property.IsValid() )
+	const FProperty* MyProperty = Property.Get();
+	if (MyProperty == nullptr)
 	{
 		// Disable all flags if no property is bound.
 		SetNodeFlags(EPropertyNodeFlags::SingleSelectOnly | EPropertyNodeFlags::EditInlineNew | EPropertyNodeFlags::ShowInnerObjectProperties, false);
@@ -162,7 +163,6 @@ void FPropertyNode::InitNode(const FPropertyNodeInitParams& InitParams)
 		const bool bSingleSelectOnly = GetReadAddressUncached( *this, true, nullptr);
 		SetNodeFlags(EPropertyNodeFlags::SingleSelectOnly, bSingleSelectOnly);
 
-		const FProperty* MyProperty = Property.Get();
 		const FProperty* OwnerProperty = MyProperty->GetOwnerProperty();
 
 		const bool bIsObjectOrInterface = CastField<FObjectPropertyBase>(MyProperty) || CastField<FInterfaceProperty>(MyProperty);
@@ -224,7 +224,7 @@ void FPropertyNode::InitNode(const FPropertyNodeInitParams& InitParams)
 		bool bRequiresValidation = bIsEditInlineNew || bShowInnerObjectProperties;
 	
 		// We require validation if we are in a container.
-		bRequiresValidation |= MyProperty && (MyProperty->IsA<FArrayProperty>() || MyProperty->IsA<FSetProperty>() || MyProperty->IsA<FMapProperty>());
+		bRequiresValidation |= MyProperty->IsA<FArrayProperty>() || MyProperty->IsA<FSetProperty>() || MyProperty->IsA<FMapProperty>();
 
 		// We require validation if our parent also needs validation (if an array parent was resized all the addresses of children are invalid)
 		bRequiresValidation |= (GetParentNode() && GetParentNode()->HasNodeFlags( EPropertyNodeFlags::RequiresValidation ) != 0);
@@ -895,45 +895,45 @@ bool FPropertyNode::IsEditConst() const
 {
 	if (bUpdateEditConstState)
 	{
-		// Ask the objects whether this property can be changed
-		const FObjectPropertyNode* ObjectPropertyNode = FindObjectItemParent();
+	// Ask the objects whether this property can be changed
+	const FObjectPropertyNode* ObjectPropertyNode = FindObjectItemParent();
 
 		bIsEditConst = IsPropertyConst();
 		if (!bIsEditConst && Property != nullptr && ObjectPropertyNode)
-		{
-			// travel up the chain to see if this property's owner struct is editconst - if it is, so is this property
-			FPropertyNode* NextParent = ParentNode;
+	{
+		// travel up the chain to see if this property's owner struct is editconst - if it is, so is this property
+		FPropertyNode* NextParent = ParentNode;
 			while (NextParent != nullptr && CastField<FStructProperty>(NextParent->GetProperty()) != NULL)
+		{
+			if (NextParent->IsEditConst())
 			{
-				if (NextParent->IsEditConst())
-				{
 					bIsEditConst = true;
 					break;
-				}
-				NextParent = NextParent->ParentNode;
 			}
+			NextParent = NextParent->ParentNode;
+		}
 
 			if (!bIsEditConst)
 			{
-				for (TPropObjectConstIterator CurObjectIt(ObjectPropertyNode->ObjectConstIterator()); CurObjectIt; ++CurObjectIt)
+		for (TPropObjectConstIterator CurObjectIt(ObjectPropertyNode->ObjectConstIterator()); CurObjectIt; ++CurObjectIt)
+		{
+			const TWeakObjectPtr<UObject> CurObject = *CurObjectIt;
+			if (CurObject.IsValid())
+			{
+				if (!CurObject->CanEditChange(Property.Get()))
 				{
-					const TWeakObjectPtr<UObject> CurObject = *CurObjectIt;
-					if (CurObject.IsValid())
-					{
-						if (!CurObject->CanEditChange(Property.Get()))
-						{
-							// At least one of the objects didn't like the idea of this property being changed.
+					// At least one of the objects didn't like the idea of this property being changed.
 							bIsEditConst = true;
 							break;
 						}
-					}
 				}
 			}
 		}
+	}
 
-		// check edit condition
+	// check edit condition
 		if (!bIsEditConst && HasEditCondition())
-		{
+	{
 			bIsEditConst = !IsEditConditionMet();
 		}
 
