The details tree now correctly restores expanded categories when clearing the search field.

This was occurring because multiple calls to RequestItemExpanded in the same frame would cause later calls to be ignored because the state it was trying to set was already the current state. However, this ignored that the state would change once the pending FilteredNodesRequestingExpansionState list had been processed at the end of the frame.

diff --git a/Engine/Source/Editor/PropertyEditor/Private/SDetailsViewBase.cpp b/Engine/Source/Editor/PropertyEditor/Private/SDetailsViewBase.cpp
--- a/Engine/Source/Editor/PropertyEditor/Private/SDetailsViewBase.cpp
+++ b/Engine/Source/Editor/PropertyEditor/Private/SDetailsViewBase.cpp
@@ -513,11 +513,7 @@ void SDetailsViewBase::NotifyFinishedChangingProperties(const FPropertyChangedEv
 
 void SDetailsViewBase::RequestItemExpanded(TSharedRef<FDetailTreeNode> TreeNode, bool bExpand)
 {
-	// Don't change expansion state if its already in that state
-	if (DetailTree->IsItemExpanded(TreeNode) != bExpand)
-	{
-		FilteredNodesRequestingExpansionState.Add(TreeNode, bExpand);
-	}
+	FilteredNodesRequestingExpansionState.Add(TreeNode, bExpand);
 }
 
 void SDetailsViewBase::RefreshTree()
@@ -835,11 +831,11 @@ void SDetailsViewBase::Tick( const FGeometry& AllottedGeometry, const double InC
 
 	if (DeferredActions.Num() > 0)
 	{
-		// Execute any deferred actions
+			// Execute any deferred actions
 		for (const FSimpleDelegate& DeferredAction : DeferredActions)
-		{
-			DeferredAction.ExecuteIfBound();
-		}
+			{
+				DeferredAction.ExecuteIfBound();
+			}
 		DeferredActions.Empty();
 	}
 
@@ -959,12 +955,12 @@ void SDetailsViewBase::Tick( const FGeometry& AllottedGeometry, const double InC
 	if (FilteredNodesRequestingExpansionState.Num() > 0)
 	{
 		// change expansion state on the nodes that request it
-		for (TMap<TWeakPtr<FDetailTreeNode>, bool >::TConstIterator It(FilteredNodesRequestingExpansionState); It; ++It)
+		for (TMap<TWeakPtr<FDetailTreeNode>, bool>::TConstIterator It(FilteredNodesRequestingExpansionState); It; ++It)
 		{
-			TWeakPtr<FDetailTreeNode> DetailTreeNode = It.Key();
+			TSharedPtr<FDetailTreeNode> DetailTreeNode = It.Key().Pin();
 			if (DetailTreeNode.IsValid())
 			{
-				DetailTree->SetItemExpansion(DetailTreeNode.Pin().ToSharedRef(), It.Value());
+				DetailTree->SetItemExpansion(DetailTreeNode.ToSharedRef(), It.Value());
 			}
 		}
 
