IWYU. Removed unnecessary includes of SResetToDefaultPropertyEditor.

diff --git a/Engine/Source/Editor/PropertyEditor/Private/DetailPropertyRow.cpp b/Engine/Source/Editor/PropertyEditor/Private/DetailPropertyRow.cpp
--- a/Engine/Source/Editor/PropertyEditor/Private/DetailPropertyRow.cpp
+++ b/Engine/Source/Editor/PropertyEditor/Private/DetailPropertyRow.cpp
@@ -3,7 +3,6 @@
 #include "DetailPropertyRow.h"
 #include "Modules/ModuleManager.h"
 #include "PropertyCustomizationHelpers.h"
-#include "SResetToDefaultPropertyEditor.h"
 #include "DetailItemNode.h"
 #include "DetailCategoryGroupNode.h"
 #include "ObjectPropertyNode.h"
diff --git a/Engine/Source/Editor/PropertyEditor/Private/PropertyEditorHelpers.cpp b/Engine/Source/Editor/PropertyEditor/Private/PropertyEditorHelpers.cpp
--- a/Engine/Source/Editor/PropertyEditor/Private/PropertyEditorHelpers.cpp
+++ b/Engine/Source/Editor/PropertyEditor/Private/PropertyEditorHelpers.cpp
@@ -3,12 +3,11 @@
 #include "PropertyEditorHelpers.h"
 #include "Widgets/Layout/SBorder.h"
 #include "Widgets/Input/SCheckBox.h"
-#include "UserInterface/PropertyEditor/PropertyEditorConstants.h"
 #include "IDocumentation.h"
 
 #include "PropertyHandleImpl.h"
-#include "SResetToDefaultPropertyEditor.h"
 
+#include "UserInterface/PropertyEditor/PropertyEditorConstants.h"
 #include "UserInterface/PropertyEditor/SPropertyEditor.h"
 #include "UserInterface/PropertyEditor/SPropertyEditorNumeric.h"
 #include "UserInterface/PropertyEditor/SPropertyEditorArray.h"
