SDetailsViewBase now allows deferred actions to call EnqueueDeferredAction. These additional deferred actions are then immediately processed in the same frame.

This manifested as an ensure in a details customization calling EnqueueDeferredAction from within CustomizeChildren when a UCurveTable was added to an array.

Since programmers may not understand that it can be called from a deferred action, this restriction would be very hard for programmers to grasp until the assertion triggers.

diff --git a/Engine/Source/Editor/PropertyEditor/Private/SDetailsViewBase.cpp b/Engine/Source/Editor/PropertyEditor/Private/SDetailsViewBase.cpp
--- a/Engine/Source/Editor/PropertyEditor/Private/SDetailsViewBase.cpp
+++ b/Engine/Source/Editor/PropertyEditor/Private/SDetailsViewBase.cpp
@@ -797,7 +797,7 @@ void SDetailsViewBase::Tick( const FGeometry& AllottedGeometry, const double InC
 	FRootPropertyNodeList& RootPropertyNodes = GetRootNodes();
 
 	bool bHadDeferredActions = DeferredActions.Num() > 0;
-	auto PreProcessRootNode = [&bHadDeferredActions, this](TSharedPtr<FComplexPropertyNode> RootPropertyNode) 
+	auto PreProcessRootNode = [bHadDeferredActions, this](TSharedPtr<FComplexPropertyNode> RootPropertyNode) 
 	{
 		check(RootPropertyNode.IsValid());
 
@@ -829,14 +829,22 @@ void SDetailsViewBase::Tick( const FGeometry& AllottedGeometry, const double InC
 		}
 	}
 
-	if (DeferredActions.Num() > 0)
+	if (bHadDeferredActions)
 	{
+		TArray<FSimpleDelegate> DeferredActionsCopy;
+		
+		do
+		{
+			// Execute any deferred actions
+			DeferredActionsCopy = MoveTemp(DeferredActions);
+			DeferredActions.Reset();
+
 			// Execute any deferred actions
-		for (const FSimpleDelegate& DeferredAction : DeferredActions)
+			for (const FSimpleDelegate& DeferredAction : DeferredActionsCopy)
 			{
 				DeferredAction.ExecuteIfBound();
 			}
-		DeferredActions.Empty();
+		} while (!DeferredActions.IsEmpty());
 	}
 
 	if (bHadDeferredActions)
