Details rows now have a toolbar widget on the right that collapses any overflowing extension buttons into a dropdown menu.

Changed how global extension buttons are registered to restrict the available functionality. Instead of being able to register arbitrary widgets, only buttons are now possible.

Details view scrollbar is now displayed against a dark background.

Updated Remote Control UI usage of global extension widgets.

#b lauren.barnes

diff --git a/Engine/Plugins/VirtualProduction/RemoteControl/Source/RemoteControlUI/Private/RemoteControlUIModule.cpp b/Engine/Plugins/VirtualProduction/RemoteControl/Source/RemoteControlUI/Private/RemoteControlUIModule.cpp
--- a/Engine/Plugins/VirtualProduction/RemoteControl/Source/RemoteControlUI/Private/RemoteControlUIModule.cpp
+++ b/Engine/Plugins/VirtualProduction/RemoteControl/Source/RemoteControlUI/Private/RemoteControlUIModule.cpp
@@ -116,39 +116,29 @@ void FRemoteControlUIModule::UnregisterAssetTools()
 	RemoteControlPresetActions.Reset();
 }
 
-void FRemoteControlUIModule::HandleCreatePropertyRowExtension(const FOnGenerateGlobalRowExtensionArgs& InArgs, /*FOnGenerateGlobalRowExtensionArgs::EWidgetPosition InWidgetPosition,*/ TArray<TSharedRef<SWidget>>& OutExtensions)
+void FRemoteControlUIModule::HandleCreatePropertyRowExtension(const FOnGenerateGlobalRowExtensionArgs& InArgs, TArray<FPropertyRowExtensionButton>& OutExtensions)
 {
-	ensureMsgf(false, TEXT("EWidgetPosition was removed in UE5.  Code needs to be fixed up."));
-#if 0
-	if (InWidgetPosition == FOnGenerateGlobalRowExtensionArgs::EWidgetPosition::Right)
-	{
-		return;
-	}
-#endif
-
-	TSharedRef<SWidget> ExposeButton = 
-		SNew(SVerticalBox)
-		+ SVerticalBox::Slot()
-		.FillHeight(1.0f)
-		.VAlign(VAlign_Center)
-		[
-			SNew(SButton)
-			.IsFocusable(false)
-			.ButtonStyle(FEditorStyle::Get(), "NoBorder")
-			.Visibility_Raw(this, &FRemoteControlUIModule::OnGetExposeButtonVisibility, InArgs.PropertyHandle)
-			.OnClicked_Raw(this, &FRemoteControlUIModule::OnToggleExposeProperty, InArgs.PropertyHandle)
-			[
-				SNew(SImage)
-				.Image_Raw(this, &FRemoteControlUIModule::OnGetExposedIcon, InArgs.PropertyHandle)
-			]
-		];
-		
-
-	OutExtensions.Add(MoveTemp(ExposeButton));
-}
-
-const FSlateBrush* FRemoteControlUIModule::OnGetExposedIcon(TSharedPtr<IPropertyHandle> Handle) const
+	FPropertyRowExtensionButton& ExposeButton = OutExtensions.AddDefaulted_GetRef();
+	ExposeButton.Icon = TAttribute<FSlateIcon>::Create(
+		[this, PropertyHandle = InArgs.PropertyHandle]
+		{
+			return OnGetExposedIcon(PropertyHandle);
+		});
+
+	ExposeButton.Label = LOCTEXT("ExposeProperty", "Expose Property");
+	ExposeButton.ToolTip = LOCTEXT("ExposePropertyToolTip", "Expose this property to Remote Control.");
+	ExposeButton.UIAction = FUIAction(
+		FExecuteAction::CreateRaw(this, &FRemoteControlUIModule::OnToggleExposeProperty, InArgs.PropertyHandle),
+		FCanExecuteAction::CreateRaw(this, &FRemoteControlUIModule::CanToggleExposeProperty, InArgs.PropertyHandle),
+		FGetActionCheckState::CreateRaw(this, &FRemoteControlUIModule::GetPropertyExposedCheckState, InArgs.PropertyHandle),
+		FIsActionButtonVisible::CreateRaw(this, &FRemoteControlUIModule::CanToggleExposeProperty, InArgs.PropertyHandle)
+	);
+}
+
+FSlateIcon FRemoteControlUIModule::OnGetExposedIcon(TSharedPtr<IPropertyHandle> Handle) const
 {
+	FName BrushName("NoBrush");
+
 	if (TSharedPtr<SRemoteControlPanel> Panel = WeakActivePanel.Pin())
 	{
 		if (Panel->GetPreset())
@@ -156,19 +146,19 @@ const FSlateBrush* FRemoteControlUIModule::OnGetExposedIcon(TSharedPtr<IProperty
 			EPropertyExposeStatus Status = GetPropertyExposeStatus(Handle);
 			if (Status == EPropertyExposeStatus::Exposed)
 			{
-				return FEditorStyle::GetBrush(TEXT("Level.VisibleIcon16x"));
+				BrushName = "Level.VisibleIcon16x";
 			}
 			else
 			{
-				return FEditorStyle::GetBrush(TEXT("Level.NotVisibleIcon16x"));
+				BrushName = "Level.NotVisibleIcon16x";
 			}
 		}
 	}
 
-	return FEditorStyle::GetNoBrush();
+	return FSlateIcon(FEditorStyle::Get().GetStyleSetName(), BrushName);
 }
 
-EVisibility FRemoteControlUIModule::OnGetExposeButtonVisibility(TSharedPtr<IPropertyHandle> Handle) const
+bool FRemoteControlUIModule::CanToggleExposeProperty(TSharedPtr<IPropertyHandle> Handle) const
 {
 	if (TSharedPtr<SRemoteControlPanel> Panel = WeakActivePanel.Pin())
 	{
@@ -177,30 +167,42 @@ EVisibility FRemoteControlUIModule::OnGetExposeButtonVisibility(TSharedPtr<IProp
 			EPropertyExposeStatus ExposeStatus = GetPropertyExposeStatus(Handle);
 			if (ExposeStatus == EPropertyExposeStatus::Exposed || ExposeStatus == EPropertyExposeStatus::Unexposed)
 			{
-				return EVisibility::Visible;
+				return true;
 			}
-			else
+		}
+	}
+
+	return false;
+}
+
+ECheckBoxState FRemoteControlUIModule::GetPropertyExposedCheckState(TSharedPtr<IPropertyHandle> Handle) const
+{
+	if (TSharedPtr<SRemoteControlPanel> Panel = WeakActivePanel.Pin())
+	{
+		if (Panel->GetPreset())
+		{
+			EPropertyExposeStatus ExposeStatus = GetPropertyExposeStatus(Handle);
+			if (ExposeStatus == EPropertyExposeStatus::Exposed)
 			{
-				// Show no icon when property is unexposable.
-				return EVisibility::Hidden;
+				return ECheckBoxState::Checked;
 			}
 		}
 	}
-	return EVisibility::Collapsed;
+
+	return ECheckBoxState::Unchecked;
 }
 
-FReply FRemoteControlUIModule::OnToggleExposeProperty(TSharedPtr<IPropertyHandle> Handle)
+void FRemoteControlUIModule::OnToggleExposeProperty(TSharedPtr<IPropertyHandle> Handle)
 {
 	if (!ensureMsgf(Handle && Handle->IsValidHandle(), TEXT("Property could not be exposed because the handle was invalid.")))
 	{
-		return FReply::Handled();
+		return;
 	}
 
 	if (TSharedPtr<SRemoteControlPanel> Panel = WeakActivePanel.Pin())
 	{
 		Panel->ToggleProperty(Handle);
 	}
-	return FReply::Handled();
 }
 
 FRemoteControlUIModule::EPropertyExposeStatus FRemoteControlUIModule::GetPropertyExposeStatus(const TSharedPtr<IPropertyHandle>& Handle) const
diff --git a/Engine/Plugins/VirtualProduction/RemoteControl/Source/RemoteControlUI/Private/RemoteControlUIModule.h b/Engine/Plugins/VirtualProduction/RemoteControl/Source/RemoteControlUI/Private/RemoteControlUIModule.h
--- a/Engine/Plugins/VirtualProduction/RemoteControl/Source/RemoteControlUI/Private/RemoteControlUIModule.h
+++ b/Engine/Plugins/VirtualProduction/RemoteControl/Source/RemoteControlUI/Private/RemoteControlUIModule.h
@@ -77,16 +77,19 @@ private:
 	void UnregisterDetailRowExtension();
 
 	/** Handle creating the property row extensions.  */
-	void HandleCreatePropertyRowExtension(const FOnGenerateGlobalRowExtensionArgs& InArgs, /*FOnGenerateGlobalRowExtensionArgs::EWidgetPosition InWidgetPosition,*/ TArray<TSharedRef<SWidget>>& OutExtensions);
+	void HandleCreatePropertyRowExtension(const FOnGenerateGlobalRowExtensionArgs& InArgs, TArray<FPropertyRowExtensionButton>& OutExtensions);
 
 	/** Handle getting the icon displayed in the property row extension. */
-	const FSlateBrush* OnGetExposedIcon(TSharedPtr<IPropertyHandle> Handle) const;
+	FSlateIcon OnGetExposedIcon(TSharedPtr<IPropertyHandle> Handle) const;
 
 	/** Handle getting the expose button visibility. */
-	EVisibility OnGetExposeButtonVisibility(TSharedPtr<IPropertyHandle> Handle) const;
+	bool CanToggleExposeProperty(TSharedPtr<IPropertyHandle> Handle) const;
+
+	/** Is the property currently exposed? */
+	ECheckBoxState GetPropertyExposedCheckState(TSharedPtr<IPropertyHandle> Handle) const;
 
 	/** Handle clicking the expose button. */
-	FReply OnToggleExposeProperty(TSharedPtr<IPropertyHandle> Handle);
+	void OnToggleExposeProperty(TSharedPtr<IPropertyHandle> Handle);
 
 	/** Returns whether a property is exposed, unexposed or unexposable. */
 	EPropertyExposeStatus GetPropertyExposeStatus(const TSharedPtr<IPropertyHandle>& Handle) const;
diff --git a/Engine/Source/Editor/EditorStyle/Private/StarshipStyle.cpp b/Engine/Source/Editor/EditorStyle/Private/StarshipStyle.cpp
--- a/Engine/Source/Editor/EditorStyle/Private/StarshipStyle.cpp
+++ b/Engine/Source/Editor/EditorStyle/Private/StarshipStyle.cpp
@@ -2254,14 +2254,7 @@ void FStarshipEditorStyle::FStyle::SetupSequencerStyles()
 			.SetShadowColorAndOpacity(FLinearColor(0, 0, 0, 0.9f)));
 		Set("Sequencer.BreadcrumbIcon", new IMAGE_BRUSH("Common/SmallArrowRight", Icon10x10));
 
-
-		const FButtonStyle DetailsKeyButton = FButtonStyle(NoBorder)
-		.SetNormal(IMAGE_BRUSH("Sequencer/AddKey_Details", FVector2D(11, 11)))
-		.SetHovered(IMAGE_BRUSH("Sequencer/AddKey_Details", FVector2D(11, 11), SelectionColor))
-		.SetPressed(IMAGE_BRUSH("Sequencer/AddKey_Details", FVector2D(11, 11), SelectionColor_Pressed))
-		.SetNormalPadding(FMargin(0, 1))
-		.SetPressedPadding(FMargin(0, 2, 0, 0));
-		Set("Sequencer.AddKey.Details", DetailsKeyButton);
+		Set("Sequencer.AddKey.Details", new IMAGE_BRUSH("Sequencer/AddKey_Details", FVector2D(11, 11)));
 
 		const FSplitterStyle OutlinerSplitterStyle = FSplitterStyle()
 		.SetHandleNormalBrush(FSlateNoResource())
@@ -3147,7 +3140,7 @@ void FStarshipEditorStyle::FStyle::SetupPropertyEditorStyles()
 		FSlateFontInfo MobilityFont = FStyleFonts::Get().Small;
 		MobilityFont.LetterSpacing = 100;
 
-		Set("PropertyWindow.MobilityFont", MobilityFont);
+		Set( "PropertyWindow.MobilityFont", MobilityFont );
 		Set( "PropertyWindow.NoOverlayColor", new FSlateNoResource() );
 		Set( "PropertyWindow.EditConstColor", new FSlateColorBrush( FColor( 152, 152, 152, 80 ) ) );
 		Set( "PropertyWindow.FilteredColor", new FSlateColorBrush( FColor( 0, 255, 0, 80 ) ) );
@@ -3233,7 +3226,33 @@ void FStarshipEditorStyle::FStyle::SetupPropertyEditorStyles()
 			.SetFont(GetFontStyle("DetailsView.CategoryFontStyle"))
 			.SetColorAndOpacity(FStyleColors::ForegroundHeader)
 		);
-		
+
+		FButtonStyle DetailsExtensionMenuButton = FButtonStyle(FStarshipCoreStyle::GetCoreStyle().GetWidgetStyle<FButtonStyle>("NoBorder"))
+			.SetNormalForeground(FStyleColors::Foreground)
+			.SetHoveredForeground(FStyleColors::ForegroundHover)
+			.SetPressedForeground(FStyleColors::ForegroundHover)
+			.SetDisabledForeground(FStyleColors::Foreground)
+			.SetNormalPadding(FMargin(0, 2, 0, 2))
+			.SetPressedPadding(FMargin(0, 3, 0, 1));
+
+		Set("DetailsView.ExtensionToolBar.Button", DetailsExtensionMenuButton);
+
+		FToolBarStyle DetailsExtensionToolBarStyle = FToolBarStyle(FStarshipCoreStyle::GetCoreStyle().GetWidgetStyle<FToolBarStyle>("SlimToolBar"))
+			.SetBackground(FSlateNoResource())
+			.SetIconSize(Icon16x16)
+			.SetLabelPadding(FMargin(0))
+			.SetComboButtonPadding(FMargin(0))
+			.SetBlockPadding(FMargin(2.0f, 0.0f))
+			.SetIndentedBlockPadding(FMargin(0))
+			.SetButtonPadding(FMargin(2, 0))
+			.SetCheckBoxPadding(FMargin(3.0f, 0.0f))
+			.SetButtonStyle(DetailsExtensionMenuButton)
+			.SetSeparatorBrush(FSlateNoResource())
+			.SetSeparatorPadding(FMargin(4.0f, 0.0f))
+			.SetExpandBrush(CORE_IMAGE_BRUSH_SVG("Starship/Common/ellipsis-vertical-narrow", FVector2D(6, 24)));
+
+		Set("DetailsView.ExtensionToolBar", DetailsExtensionToolBarStyle);
+
 		Set("DetailsView.ArrayDropShadow", new IMAGE_BRUSH("Common/ArrayDropShadow", FVector2D(32,2)));
 
 		Set( "DetailsView.TreeView.TableRow", FTableRowStyle()
diff --git a/Engine/Source/Editor/PropertyEditor/Private/DetailCategoryGroupNode.cpp b/Engine/Source/Editor/PropertyEditor/Private/DetailCategoryGroupNode.cpp
--- a/Engine/Source/Editor/PropertyEditor/Private/DetailCategoryGroupNode.cpp
+++ b/Engine/Source/Editor/PropertyEditor/Private/DetailCategoryGroupNode.cpp
@@ -17,7 +17,7 @@ void SDetailCategoryTableRow::Construct( const FArguments& InArgs, TSharedRef<FD
 
 	FDetailColumnSizeData& ColumnSizeData = InOwnerTreeNode->GetDetailsView()->GetColumnSizeData();
 
-	TSharedPtr<SHorizontalBox> HeaderBox = SNew(SHorizontalBox)
+	TSharedRef<SHorizontalBox> HeaderBox = SNew(SHorizontalBox)
 		+ SHorizontalBox::Slot()
 		.HAlign(HAlign_Left)
 		.VAlign(VAlign_Fill)
@@ -61,7 +61,7 @@ void SDetailCategoryTableRow::Construct( const FArguments& InArgs, TSharedRef<FD
 	[
 		SNew( SBorder )
 		.BorderImage( FAppStyle::Get().GetBrush( "DetailsView.GridLine") )
-		.Padding( FMargin(0, 0, 0, 1) )
+		.Padding( FMargin(0, 0, SDetailTableRowBase::ScrollbarPaddingSize, 1) )
 		[
 			SNew(SHorizontalBox)
 			+ SHorizontalBox::Slot()
@@ -83,9 +83,9 @@ void SDetailCategoryTableRow::Construct( const FArguments& InArgs, TSharedRef<FD
 				SNew( SBorder )
 				.BorderImage( this, &SDetailCategoryTableRow::GetBackgroundImage )
 				.BorderBackgroundColor( this, &SDetailCategoryTableRow::GetInnerBackgroundColor )
-				.Padding(FMargin(0, 0, SDetailTableRowBase::ScrollbarPaddingSize, 0))
+				.Padding(FMargin(0, 0, 0, 0))
 				[
-					HeaderBox.ToSharedRef()
+					HeaderBox
 				]
 			]
 		]
diff --git a/Engine/Source/Editor/PropertyEditor/Private/SDetailSingleItemRow.cpp b/Engine/Source/Editor/PropertyEditor/Private/SDetailSingleItemRow.cpp
--- a/Engine/Source/Editor/PropertyEditor/Private/SDetailSingleItemRow.cpp
+++ b/Engine/Source/Editor/PropertyEditor/Private/SDetailSingleItemRow.cpp
@@ -1,26 +1,29 @@
 // Copyright Epic Games, Inc. All Rights Reserved.
 
 #include "SDetailSingleItemRow.h"
-#include "ObjectPropertyNode.h"
-#include "Widgets/Images/SImage.h"
-#include "Widgets/Input/SButton.h"
-#include "Settings/EditorExperimentalSettings.h"
+
+#include "DetailGroup.h"
+#include "DetailPropertyRow.h"
 #include "DetailWidgetRow.h"
+#include "Editor.h"
 #include "IDetailKeyframeHandler.h"
 #include "IDetailPropertyExtensionHandler.h"
-#include "DetailPropertyRow.h"
-#include "DetailGroup.h"
-#include "HAL/PlatformApplicationMisc.h"
-#include "Editor.h"
-#include "PropertyHandleImpl.h"
+#include "ObjectPropertyNode.h"
+#include "PropertyEditorConstants.h"
 #include "PropertyEditorModule.h"
-#include "Modules/ModuleInterface.h"
-#include "Modules/ModuleManager.h"
+#include "PropertyHandleImpl.h"
 #include "SConstrainedBox.h"
 #include "SDetailExpanderArrow.h"
 #include "SDetailRowIndent.h"
 #include "SResetToDefaultPropertyEditor.h"
-#include "PropertyEditorConstants.h"
+
+#include "HAL/PlatformApplicationMisc.h"
+#include "Modules/ModuleInterface.h"
+#include "Modules/ModuleManager.h"
+#include "Settings/EditorExperimentalSettings.h"
+#include "Widgets/Images/SImage.h"
+#include "Widgets/Input/SButton.h"
+#include "Widgets/Input/SComboButton.h"
 
 namespace DetailWidgetConstants
 {
@@ -50,83 +53,6 @@ namespace SDetailSingleItemRow_Helper
 	}
 }
 
-void SDetailSingleItemRow::OnFavoriteMenuToggle()
-{
-	if (Customization->GetPropertyNode().IsValid() && Customization->GetPropertyNode()->CanDisplayFavorite())
-	{
-		bool bToggle = !Customization->GetPropertyNode()->IsFavorite();
-		Customization->GetPropertyNode()->SetFavorite(bToggle);
-		if (OwnerTreeNode.IsValid())
-		{
-			//////////////////////////////////////////////////////////////////////////
-			// Calculate properly the scrolling offset (by item) to make sure the mouse stay over the same property
-
-			//Get the node item number in case it is expand we have to recursively count all childrens
-			int32 ExpandSize = 0;
-			if (OwnerTreeNode.Pin()->ShouldBeExpanded())
-			{
-				SDetailSingleItemRow_Helper::RecursivelyGetItemShow(OwnerTreeNode.Pin().ToSharedRef(), ExpandSize);
-			}
-			else
-			{
-				//if the item is not expand count is 1
-				ExpandSize = 1;
-			}
-			
-			//Get the number of favorite child (simple and advance) to know if the favorite category will be create or remove
-			int32 SimplePropertiesNum = 0;
-			int32 AdvancePropertiesNum = 0;
-
-			FDetailLayoutBuilderImpl& DetailLayout = OwnerTreeNode.Pin()->GetParentCategory()->GetParentLayoutImpl();
-
-			const FName FavoritesCategoryName = TEXT("Favorites");
-			bool HasCategoryFavorite = DetailLayout.HasCategory(FavoritesCategoryName);
-			if(HasCategoryFavorite)
-			{
-				DetailLayout.DefaultCategory(FavoritesCategoryName).GetCategoryInformation(SimplePropertiesNum, AdvancePropertiesNum);
-			}
-
-			// Check if the property we toggle is an advance property
-			bool IsAdvanceProperty = Customization->GetPropertyNode()->HasNodeFlags(EPropertyNodeFlags::IsAdvanced) == 0 ? false : true;
-
-			// Compute the scrolling offset by item
-			int32 ScrollingOffsetAdd = ExpandSize;
-			int32 ScrollingOffsetRemove = -ExpandSize;
-			if (HasCategoryFavorite)
-			{
-				// Adding the advance button in a category add 1 item
-				ScrollingOffsetAdd += (IsAdvanceProperty && AdvancePropertiesNum == 0) ? 1 : 0;
-
-				if (IsAdvanceProperty && AdvancePropertiesNum == 1)
-				{
-					//Removing the advance button count as 1 item
-					ScrollingOffsetRemove -= 1;
-				}
-				if (AdvancePropertiesNum + SimplePropertiesNum == 1)
-				{
-					//Removing a full category count as 2 items
-					ScrollingOffsetRemove -= 2;
-				}
-			}
-			else
-			{
-				// Adding new category (2 items) adding advance button (1 item)
-				ScrollingOffsetAdd += IsAdvanceProperty ? 3 : 2;
-				
-				// We should never remove an item from favorite if there is no favorite category
-				// Set the remove offset to 0
-				ScrollingOffsetRemove = 0;
-			}
-
-			// Apply the calculated offset
-			OwnerTreeNode.Pin()->GetDetailsView()->MoveScrollOffset(bToggle ? ScrollingOffsetAdd : ScrollingOffsetRemove);
-
-			// Refresh the tree
-			OwnerTreeNode.Pin()->GetDetailsView()->ForceRefresh();
-		}
-	}
-}
-
 void SDetailSingleItemRow::OnArrayDragEnter(const FDragDropEvent& DragDropEvent)
 {
 	bIsHoveredDragTarget = true;
@@ -146,7 +72,9 @@ bool SDetailSingleItemRow::CheckValidDrop(const TSharedPtr<SDetailSingleItemRow>
 		{
 			int32 OriginalIndex = SwappingPropertyNode->GetArrayIndex();
 			int32 NewIndex = SwappablePropertyNode->GetArrayIndex();
-			TSharedPtr<IPropertyHandle> SwappingHandle = PropertyEditorHelpers::GetPropertyHandle(SwappingPropertyNode.ToSharedRef(), OwnerTreeNode.Pin()->GetDetailsView()->GetNotifyHook(), OwnerTreeNode.Pin()->GetDetailsView()->GetPropertyUtilities());
+
+			IDetailsViewPrivate* DetailsView = OwnerTreeNode.Pin()->GetDetailsView();
+			TSharedPtr<IPropertyHandle> SwappingHandle = PropertyEditorHelpers::GetPropertyHandle(SwappingPropertyNode.ToSharedRef(), DetailsView->GetNotifyHook(), DetailsView->GetPropertyUtilities());
 			TSharedPtr<IPropertyHandleArray> ParentHandle = SwappingHandle->GetParentHandle()->AsArray();
 
 			if (ParentHandle.IsValid() && SwappablePropertyNode->GetParentNode() == SwappingPropertyNode->GetParentNode())
@@ -174,34 +102,37 @@ FReply SDetailSingleItemRow::OnArrayDrop(const FDragDropEvent& DragDropEvent)
 		return FReply::Unhandled();
 	}
 
-	if (CheckValidDrop(RowPtr))
+	if (!CheckValidDrop(RowPtr))
 	{
-		TSharedPtr<FPropertyNode> SwappingPropertyNode = RowPtr->SwappablePropertyNode;
-		TSharedPtr<IPropertyHandle> SwappingHandle = PropertyEditorHelpers::GetPropertyHandle(SwappingPropertyNode.ToSharedRef(), OwnerTreeNode.Pin()->GetDetailsView()->GetNotifyHook(), OwnerTreeNode.Pin()->GetDetailsView()->GetPropertyUtilities());
-		TSharedPtr<IPropertyHandleArray> ParentHandle = SwappingHandle->GetParentHandle()->AsArray();
-		int32 OriginalIndex = SwappingPropertyNode->GetArrayIndex();
-		int32 NewIndex = SwappablePropertyNode->GetArrayIndex();
-
-		// Need to swap the moving and target expansion states before saving
-		bool bOriginalSwappableExpansion = SwappablePropertyNode->HasNodeFlags(EPropertyNodeFlags::Expanded) != 0;
-		bool bOriginalSwappingExpansion = SwappingPropertyNode->HasNodeFlags(EPropertyNodeFlags::Expanded) != 0;
-		SwappablePropertyNode->SetNodeFlags(EPropertyNodeFlags::Expanded, bOriginalSwappingExpansion);
-		SwappingPropertyNode->SetNodeFlags(EPropertyNodeFlags::Expanded, bOriginalSwappableExpansion);
+		return FReply::Unhandled();
+	}
 
-		IDetailsViewPrivate* DetailsView = OwnerTreeNode.Pin()->GetDetailsView();
-		DetailsView->SaveExpandedItems(SwappablePropertyNode->GetParentNodeSharedPtr().ToSharedRef());
-		FScopedTransaction Transaction(NSLOCTEXT("UnrealEd", "MoveRow", "Move Row"));
+	IDetailsViewPrivate* DetailsView = OwnerTreeNode.Pin()->GetDetailsView();
 
-		SwappingHandle->GetParentHandle()->NotifyPreChange();
+	TSharedPtr<FPropertyNode> SwappingPropertyNode = RowPtr->SwappablePropertyNode;
+	TSharedPtr<IPropertyHandle> SwappingHandle = PropertyEditorHelpers::GetPropertyHandle(SwappingPropertyNode.ToSharedRef(), DetailsView->GetNotifyHook(), DetailsView->GetPropertyUtilities());
+	TSharedPtr<IPropertyHandleArray> ParentHandle = SwappingHandle->GetParentHandle()->AsArray();
+	int32 OriginalIndex = SwappingPropertyNode->GetArrayIndex();
+	int32 NewIndex = SwappablePropertyNode->GetArrayIndex();
 
-		ParentHandle->MoveElementTo(OriginalIndex, NewIndex);
+	// Need to swap the moving and target expansion states before saving
+	bool bOriginalSwappableExpansion = SwappablePropertyNode->HasNodeFlags(EPropertyNodeFlags::Expanded) != 0;
+	bool bOriginalSwappingExpansion = SwappingPropertyNode->HasNodeFlags(EPropertyNodeFlags::Expanded) != 0;
+	SwappablePropertyNode->SetNodeFlags(EPropertyNodeFlags::Expanded, bOriginalSwappingExpansion);
+	SwappingPropertyNode->SetNodeFlags(EPropertyNodeFlags::Expanded, bOriginalSwappableExpansion);
 
-		FPropertyChangedEvent MoveEvent(SwappingHandle->GetParentHandle()->GetProperty(), EPropertyChangeType::Unspecified);
-		SwappingHandle->GetParentHandle()->NotifyPostChange(EPropertyChangeType::Unspecified);
-		if (DetailsView->GetPropertyUtilities().IsValid())
-		{
-			DetailsView->GetPropertyUtilities()->NotifyFinishedChangingProperties(MoveEvent);
-		}
+	DetailsView->SaveExpandedItems(SwappablePropertyNode->GetParentNodeSharedPtr().ToSharedRef());
+	FScopedTransaction Transaction(NSLOCTEXT("UnrealEd", "MoveRow", "Move Row"));
+
+	SwappingHandle->GetParentHandle()->NotifyPreChange();
+
+	ParentHandle->MoveElementTo(OriginalIndex, NewIndex);
+
+	FPropertyChangedEvent MoveEvent(SwappingHandle->GetParentHandle()->GetProperty(), EPropertyChangeType::Unspecified);
+	SwappingHandle->GetParentHandle()->NotifyPostChange(EPropertyChangeType::Unspecified);
+	if (DetailsView->GetPropertyUtilities().IsValid())
+	{
+		DetailsView->GetPropertyUtilities()->NotifyFinishedChangingProperties(MoveEvent);
 	}
 
 	return FReply::Handled();
@@ -266,7 +197,8 @@ TSharedPtr<IPropertyHandle> SDetailSingleItemRow::GetPropertyHandle() const
 	TSharedPtr<FPropertyNode> PropertyNode = GetPropertyNode();
 	if (PropertyNode.IsValid())
 	{
-		Handle = PropertyEditorHelpers::GetPropertyHandle(PropertyNode.ToSharedRef(), OwnerTreeNode.Pin()->GetDetailsView()->GetNotifyHook(), OwnerTreeNode.Pin()->GetDetailsView()->GetPropertyUtilities());
+		IDetailsViewPrivate* DetailsView = OwnerTreeNode.Pin()->GetDetailsView();
+		Handle = PropertyEditorHelpers::GetPropertyHandle(PropertyNode.ToSharedRef(), DetailsView->GetNotifyHook(), DetailsView->GetPropertyUtilities());
 	}
 	else if (Customization->GetWidgetRow().PropertyHandles.Num() > 0)
 	{
@@ -299,7 +231,7 @@ void SDetailSingleItemRow::Construct( const FArguments& InArgs, FDetailLayoutCus
 		{
 			FDetailWidgetRow Row = InCustomization->GetWidgetRow();
 
-			TSharedPtr<SWidget> NameWidget, ValueWidget, ExtensionWidget;
+			TSharedPtr<SWidget> NameWidget, ValueWidget;
 
 			NameWidget = Row.NameWidget.Widget;
 
@@ -311,23 +243,24 @@ void SDetailSingleItemRow::Construct( const FArguments& InArgs, FDetailLayoutCus
 					Row.ValueWidget.Widget
 				];
 
-			ExtensionWidget = CreateExtensionWidget(ValueWidget.ToSharedRef(), *Customization, InOwnerTreeNode);
-
-			TAttribute<bool> IsEnabledAttribute = InOwnerTreeNode->IsPropertyEditingEnabled();
+			TAttribute<bool> IsEnabledAttribute;
 			if (Row.IsEnabledAttr.IsBound())
 			{
-				TAttribute<bool> RowEnabledAttr = Row.IsEnabledAttr;
 				TAttribute<bool> PropertyEnabledAttr = InOwnerTreeNode->IsPropertyEditingEnabled();
+				TAttribute<bool> RowEnabledAttr = Row.IsEnabledAttr;
 				IsEnabledAttribute = TAttribute<bool>::Create(
 					[RowEnabledAttr, PropertyEnabledAttr]()
 					{
 						return RowEnabledAttr.Get() && PropertyEnabledAttr.Get();
 					});
 			}
+			else
+			{
+				IsEnabledAttribute = InOwnerTreeNode->IsPropertyEditingEnabled();
+			}
 
 			NameWidget->SetEnabled(IsEnabledAttribute);
 			ValueWidget->SetEnabled(IsEnabledAttribute);
-			ExtensionWidget->SetEnabled(IsEnabledAttribute);
 
 			TSharedRef<SHorizontalBox> RowBox = SNew(SHorizontalBox);
 
@@ -455,7 +388,6 @@ void SDetailSingleItemRow::Construct( const FArguments& InArgs, FDetailLayoutCus
 						NameWidget.ToSharedRef()
 					];
 		
-
 				InnerSplitter->AddSlot()
 					.Value(ColumnSizeData.NameColumnWidth)
 					.OnSlotResized(ColumnSizeData.OnNameColumnResized)
@@ -463,6 +395,9 @@ void SDetailSingleItemRow::Construct( const FArguments& InArgs, FDetailLayoutCus
 						NameColumnBox
 					];
 
+				TSharedRef<SWidget> ExtensionWidget = CreateExtensionWidget();
+				ExtensionWidget->SetEnabled(IsEnabledAttribute);
+
 				// create Value column:
 				// | Left  | Name | Value | Right |
 				InnerSplitter->AddSlot()
@@ -485,7 +420,7 @@ void SDetailSingleItemRow::Construct( const FArguments& InArgs, FDetailLayoutCus
 						.Padding(5,0,0,0)
 						.AutoWidth()
 						[
-							ExtensionWidget.ToSharedRef()
+							ExtensionWidget
 						]
 					];
 			}
@@ -506,84 +441,60 @@ void SDetailSingleItemRow::Construct( const FArguments& InArgs, FDetailLayoutCus
 					];
 			}
 
-			TSharedRef<SHorizontalBox> RightColumnBox = SNew(SHorizontalBox)
-				.Clipping(EWidgetClipping::OnDemand);
-
-			// reset to default widget
-			bool bCreateResetToDefault = Row.CustomResetToDefault.IsSet();
-
-			TSharedPtr<IPropertyHandle> PropertyHandle = GetPropertyHandle();
+			TArray<FPropertyRowExtensionButton> ExtensionButtons;
 
-			if (PropertyHandle.IsValid())
-			{
-				if (PropertyHandle->HasMetaData("NoResetToDefault") || PropertyHandle->GetInstanceMetaData("NoResetToDefault"))
+			FPropertyRowExtensionButton& ResetToDefault = ExtensionButtons.AddDefaulted_GetRef();
+			ResetToDefault.Icon = TAttribute<FSlateIcon>::Create([this]()
 				{
-					bCreateResetToDefault = false;
-				}
-				else
+					return IsResetToDefaultEnabled() ?
+						FSlateIcon(FAppStyle::Get().GetStyleSetName(), "PropertyWindow.DiffersFromDefault") :
+						FSlateIcon(FAppStyle::Get().GetStyleSetName(), "NoBrush");
+				});
+			ResetToDefault.Label = NSLOCTEXT("PropertyEditor", "ResetToDefault", "Reset to Default");
+			ResetToDefault.ToolTip = TAttribute<FText>::Create([this]()
 				{
-					bCreateResetToDefault = true;
-				}
-			} 
-
-			TSharedRef<SWidget> ResetWidget = SNullWidget::NullWidget;
-			if (bCreateResetToDefault)
-			{
-				ResetWidget = SNew(SResetToDefaultPropertyEditor, PropertyHandle)
-					.IsEnabled(IsEnabledAttribute)
-					.CustomResetToDefault(Row.CustomResetToDefault);
-			}	
-			else
+					return IsResetToDefaultEnabled() ?
+						NSLOCTEXT("PropertyEditor", "ResetToDefaultToolTip", "Reset this property to its default value.") :
+						FText::GetEmpty();
+				});
+			ResetToDefault.UIAction = FUIAction(
+				FExecuteAction::CreateSP(this, &SDetailSingleItemRow::OnResetToDefaultClicked),
+				FCanExecuteAction::CreateSP(this, &SDetailSingleItemRow::IsResetToDefaultEnabled));
+
+			FPropertyRowExtensionButton& CreateKey = ExtensionButtons.AddDefaulted_GetRef();
+			CreateKey.Icon = FSlateIcon(FAppStyle::Get().GetStyleSetName(), "Sequencer.AddKey.Details");
+			CreateKey.Label = NSLOCTEXT("PropertyEditor", "CreateKey", "Create Key");
+			CreateKey.ToolTip = NSLOCTEXT("PropertyEditor", "CreateKeyToolTip", "Add a keyframe for this property.");
+			CreateKey.UIAction = FUIAction(
+				FExecuteAction::CreateSP(this, &SDetailSingleItemRow::OnAddKeyframeClicked),
+				FCanExecuteAction::CreateSP(this, &SDetailSingleItemRow::IsKeyframeButtonEnabled),
+				FGetActionCheckState(),
+				FIsActionButtonVisible::CreateSP(this, &SDetailSingleItemRow::IsKeyframeButtonVisible));
+
+			CreateGlobalExtensionWidgets(ExtensionButtons);
+
+			FSlimHorizontalToolBarBuilder ToolbarBuilder(TSharedPtr<FUICommandList>(), FMultiBoxCustomization::None);
+			ToolbarBuilder.SetLabelVisibility(EVisibility::Collapsed);
+			ToolbarBuilder.SetStyle(&FAppStyle::Get(), "DetailsView.ExtensionToolBar");
+
+			for (const FPropertyRowExtensionButton& Extension : ExtensionButtons)
 			{
-				// leave space in the reset to default column regardless
-				const FSlateBrush* DiffersFromDefaultBrush = FEditorStyle::GetBrush("PropertyWindow.DiffersFromDefault");
-				ResetWidget = SNew(SSpacer)
-					.Size(DiffersFromDefaultBrush != nullptr ? DiffersFromDefaultBrush->ImageSize : FVector2D(8.0f, 8.0f));
+				ToolbarBuilder.AddToolBarButton(Extension.UIAction, NAME_None, Extension.Label, Extension.ToolTip, Extension.Icon);
 			}
 
+			TSharedRef<SHorizontalBox> RightColumnBox = SNew(SHorizontalBox);
 			RightColumnBox->AddSlot()
-				.HAlign(HAlign_Left)
-				.VAlign(VAlign_Center)
-				.Padding(5,0,0,0)
-				.AutoWidth()
-				[
-					ResetWidget
-				];
-
-			// keyframe button
-			RightColumnBox->AddSlot()
-				.HAlign(HAlign_Left)
+				.HAlign(HAlign_Fill)
 				.VAlign(VAlign_Center)
-				.Padding(5,0,0,0)
-				.AutoWidth()
+				.Padding(5,0,5,0)
 				[
-					CreateKeyframeButton(*Customization, InOwnerTreeNode)
+					ToolbarBuilder.MakeWidget()
 				];
 
-			// fetch global extension widgets 
-			FPropertyEditorModule& PropertyEditorModule = FModuleManager::Get().GetModuleChecked<FPropertyEditorModule>("PropertyEditor");
-
-			TArray<TSharedRef<SWidget>> GlobalWidgetExtensions;
-			FOnGenerateGlobalRowExtensionArgs Args { GetPropertyHandle(), GetPropertyNode(), OwnerTreeNode };
-			PropertyEditorModule.GetGlobalRowExtensionDelegate().Broadcast(Args, GlobalWidgetExtensions);
-
-			if (GlobalWidgetExtensions.Num() > 0)
-			{
-				for (TSharedRef<SWidget>& GlobalRowExtension : GlobalWidgetExtensions)
-				{
-					RightColumnBox->AddSlot()
-						.HAlign(HAlign_Left)
-						.VAlign(VAlign_Center)
-						.AutoWidth()
-						[
-							MoveTemp(GlobalRowExtension)
-						];
-				}
-			}
-
 			OuterSplitter->AddSlot()
 				.Value(ColumnSizeData.RightColumnWidth)
 				.OnSlotResized(ColumnSizeData.OnRightColumnResized)
+				.MinSize(40)
 			[
 				RightColumnBox
 			];
@@ -600,12 +511,13 @@ void SDetailSingleItemRow::Construct( const FArguments& InArgs, FDetailLayoutCus
 	[
 		SNew( SBorder )
 		.BorderImage(FAppStyle::Get().GetBrush( "DetailsView.GridLine"))
-		.Padding( FMargin(0, 0, 0, 1) )
+		.Padding( FMargin(0, 0, SDetailTableRowBase::ScrollbarPaddingSize, 1) )
 		[
 			SNew( SBorder )
 			.BorderImage(FAppStyle::Get().GetBrush("DetailsView.CategoryMiddle"))
 			.BorderBackgroundColor(this, &SDetailSingleItemRow::GetOuterBackgroundColor)
-			.Padding(FMargin( 0.0f, 0.0f, SDetailTableRowBase::ScrollbarPaddingSize, 0.0f ))
+			.Padding(0)
+			.Clipping(EWidgetClipping::ClipToBounds)
 			[
 				Widget
 			]
@@ -624,6 +536,46 @@ void SDetailSingleItemRow::Construct( const FArguments& InArgs, FDetailLayoutCus
 	);
 }
 
+bool SDetailSingleItemRow::IsResetToDefaultEnabled() const
+{
+	FDetailWidgetRow WidgetRow = Customization->GetWidgetRow();
+	TSharedPtr<IPropertyHandle> PropertyHandle = GetPropertyHandle();
+
+	if (PropertyHandle.IsValid())
+	{
+		if (PropertyHandle->HasMetaData("NoResetToDefault") || PropertyHandle->GetInstanceMetaData("NoResetToDefault"))
+		{
+			return false;
+		}
+	}
+
+	if (WidgetRow.CustomResetToDefault.IsSet())
+	{
+		return WidgetRow.CustomResetToDefault.GetValue().IsResetToDefaultVisible(PropertyHandle);
+	}
+	else if (PropertyHandle.IsValid())
+	{
+		return PropertyHandle->CanResetToDefault();
+	} 
+
+	return false;
+}
+
+void SDetailSingleItemRow::OnResetToDefaultClicked() const
+{
+	FDetailWidgetRow WidgetRow = Customization->GetWidgetRow();
+	TSharedPtr<IPropertyHandle> PropertyHandle = GetPropertyHandle();
+
+	if (WidgetRow.CustomResetToDefault.IsSet())
+	{
+		WidgetRow.CustomResetToDefault.GetValue().OnResetToDefaultClicked(PropertyHandle);
+	}
+	else if (PropertyHandle.IsValid())
+	{
+		PropertyHandle->ResetToDefault();
+	}
+}
+
 /** Get the background color of the outer part of the row, which contains the edit condition and extension widgets. */
 FSlateColor SDetailSingleItemRow::GetOuterBackgroundColor() const
 {
@@ -768,7 +720,6 @@ bool SDetailSingleItemRow::OnContextMenuOpening(FMenuBuilder& MenuBuilder)
 				CustomMenuData.SlateIcon,
 				CustomMenuData.Action);
 		}
-
 	}
 
 	return true;
@@ -781,7 +732,8 @@ void SDetailSingleItemRow::OnCopyProperty()
 		TSharedPtr<FPropertyNode> PropertyNode = GetPropertyNode();
 		if (PropertyNode.IsValid())
 		{
-			TSharedPtr<IPropertyHandle> Handle = PropertyEditorHelpers::GetPropertyHandle(PropertyNode.ToSharedRef(), OwnerTreeNode.Pin()->GetDetailsView()->GetNotifyHook(), OwnerTreeNode.Pin()->GetDetailsView()->GetPropertyUtilities());
+			IDetailsViewPrivate* DetailsView = OwnerTreeNode.Pin()->GetDetailsView();
+			TSharedPtr<IPropertyHandle> Handle = PropertyEditorHelpers::GetPropertyHandle(PropertyNode.ToSharedRef(), DetailsView->GetNotifyHook(), DetailsView->GetPropertyUtilities());
 
 			FString Value;
 			if (Handle->GetValueAsFormattedString(Value, PPF_Copy) == FPropertyAccess::Success)
@@ -808,7 +760,8 @@ void SDetailSingleItemRow::OnPasteProperty()
 		{
 			FScopedTransaction Transaction(NSLOCTEXT("UnrealEd", "PasteProperty", "Paste Property"));
 
-			TSharedPtr<IPropertyHandle> Handle = PropertyEditorHelpers::GetPropertyHandle(PropertyNode.ToSharedRef(), OwnerTreeNode.Pin()->GetDetailsView()->GetNotifyHook(), OwnerTreeNode.Pin()->GetDetailsView()->GetPropertyUtilities());
+			IDetailsViewPrivate* DetailsView = OwnerTreeNode.Pin()->GetDetailsView();
+			TSharedPtr<IPropertyHandle> Handle = PropertyEditorHelpers::GetPropertyHandle(PropertyNode.ToSharedRef(), DetailsView->GetNotifyHook(), DetailsView->GetPropertyUtilities());
 
 			Handle->SetValueFromFormattedString(ClipboardContent);
 
@@ -899,24 +852,24 @@ bool SDetailSingleItemRow::CanPasteProperty() const
 	return !ClipboardContent.IsEmpty();
 }
 
-TSharedRef<SWidget> SDetailSingleItemRow::CreateExtensionWidget(TSharedRef<SWidget> ValueWidget, FDetailLayoutCustomization& InCustomization, TSharedRef<FDetailTreeNode> InTreeNode)
+TSharedRef<SWidget> SDetailSingleItemRow::CreateExtensionWidget() const
 {
 	TSharedPtr<SWidget> ExtensionWidget = SNullWidget::NullWidget;
 
-	if(InTreeNode->GetParentCategory().IsValid())
+	TSharedPtr<FDetailTreeNode> OwnerTreeNodePinned = OwnerTreeNode.Pin();
+	if (OwnerTreeNodePinned.IsValid())
 	{ 
-		IDetailsViewPrivate* DetailsView = InTreeNode->GetDetailsView();
+		IDetailsViewPrivate* DetailsView = OwnerTreeNodePinned->GetDetailsView();
 		TSharedPtr<IDetailPropertyExtensionHandler> ExtensionHandler = DetailsView->GetExtensionHandler();
-
-		if(ExtensionHandler.IsValid() && InCustomization.HasPropertyNode())
+		if (Customization->HasPropertyNode() && ExtensionHandler.IsValid())
 		{
-			TSharedPtr<IPropertyHandle> Handle = PropertyEditorHelpers::GetPropertyHandle(InCustomization.GetPropertyNode().ToSharedRef(), nullptr, nullptr);
+			TSharedPtr<IPropertyHandle> Handle = PropertyEditorHelpers::GetPropertyHandle(Customization->GetPropertyNode().ToSharedRef(), nullptr, nullptr);
 
-			FObjectPropertyNode* ObjectItemParent = InCustomization.GetPropertyNode()->FindObjectItemParent();
-			UClass* ObjectClass = (ObjectItemParent) ? ObjectItemParent->GetObjectBaseClass() : nullptr;
-			if(Handle->IsValidHandle() && ObjectClass && ExtensionHandler->IsPropertyExtendable(ObjectClass, *Handle))
+			FObjectPropertyNode* ObjectItemParent = Customization->GetPropertyNode()->FindObjectItemParent();
+			UClass* ObjectClass = ObjectItemParent ? ObjectItemParent->GetObjectBaseClass() : nullptr;
+			if (Handle->IsValidHandle() && ObjectClass && ExtensionHandler->IsPropertyExtendable(ObjectClass, *Handle))
 			{
-				FDetailLayoutBuilderImpl& DetailLayout = OwnerTreeNode.Pin()->GetParentCategory()->GetParentLayoutImpl();
+				FDetailLayoutBuilderImpl& DetailLayout = OwnerTreeNodePinned->GetParentCategory()->GetParentLayoutImpl();
 				ExtensionWidget = ExtensionHandler->GenerateExtensionWidget(DetailLayout, ObjectClass, Handle);
 			}
 		}
@@ -925,53 +878,137 @@ TSharedRef<SWidget> SDetailSingleItemRow::CreateExtensionWidget(TSharedRef<SWidg
 	return ExtensionWidget.ToSharedRef();
 }
 
-TSharedRef<SWidget> SDetailSingleItemRow::CreateKeyframeButton( FDetailLayoutCustomization& InCustomization, TSharedRef<FDetailTreeNode> InTreeNode )
+void SDetailSingleItemRow::OnFavoriteMenuToggle()
 {
-	IDetailsViewPrivate* DetailsView = InTreeNode->GetDetailsView();
+	if (!Customization->GetPropertyNode().IsValid() || !Customization->GetPropertyNode()->CanDisplayFavorite())
+	{
+		return; 
+	}
+
+	bool bToggled = !Customization->GetPropertyNode()->IsFavorite();
+	Customization->GetPropertyNode()->SetFavorite(bToggled);
+
+	TSharedPtr<FDetailTreeNode> OwnerTreeNodePinned = OwnerTreeNode.Pin();
+	if (!OwnerTreeNodePinned.IsValid())
+	{
+		return;
+	}
+
+	//////////////////////////////////////////////////////////////////////////
+	// Calculate properly the scrolling offset (by item) to make sure the mouse stay over the same property
+
+	//Get the node item number in case it is expand we have to recursively count all childrens
+	int32 ExpandSize = 0;
+	if (OwnerTreeNodePinned->ShouldBeExpanded())
+	{
+		SDetailSingleItemRow_Helper::RecursivelyGetItemShow(OwnerTreeNodePinned.ToSharedRef(), ExpandSize);
+	}
+	else
+	{
+		//if the item is not expand count is 1
+		ExpandSize = 1;
+	}
+			
+	//Get the number of favorite child (simple and advance) to know if the favorite category will be create or remove
+	int32 SimplePropertiesNum = 0;
+	int32 NumAdvancedProperties = 0;
 
-	KeyframeHandler = DetailsView->GetKeyframeHandler();
+	FDetailLayoutBuilderImpl& DetailLayout = OwnerTreeNodePinned->GetParentCategory()->GetParentLayoutImpl();
 
-	EVisibility SetKeyVisibility = EVisibility::Collapsed;
+	const FName FavoritesCategoryName = TEXT("Favorites");
+	bool HasFavoritesCategory = DetailLayout.HasCategory(FavoritesCategoryName);
+	if (HasFavoritesCategory)
+	{
+		DetailLayout.DefaultCategory(FavoritesCategoryName).GetCategoryInformation(SimplePropertiesNum, NumAdvancedProperties);
+	}
+
+	// Check if the property we toggle is an advanced property
+	bool IsAdvancedProperty = Customization->GetPropertyNode()->HasNodeFlags(EPropertyNodeFlags::IsAdvanced) != 0;
 
-	if (InCustomization.HasPropertyNode() && KeyframeHandler.IsValid() )
+	// Compute the scrolling offset by item
+	int32 ScrollingOffsetAdd = ExpandSize;
+	int32 ScrollingOffsetRemove = -ExpandSize;
+	if (HasFavoritesCategory)
+	{
+		// Adding the advance button in a category add 1 item
+		ScrollingOffsetAdd += (IsAdvancedProperty && NumAdvancedProperties == 0) ? 1 : 0;
+
+		if (IsAdvancedProperty && NumAdvancedProperties == 1)
+		{
+			//Removing the advance button count as 1 item
+			ScrollingOffsetRemove -= 1;
+		}
+		if (NumAdvancedProperties + SimplePropertiesNum == 1)
+		{
+			//Removing a full category count as 2 items
+			ScrollingOffsetRemove -= 2;
+		}
+	}
+	else
 	{
-		TSharedPtr<IPropertyHandle> Handle = PropertyEditorHelpers::GetPropertyHandle(InCustomization.GetPropertyNode().ToSharedRef(), nullptr, nullptr);
+		// Adding new category (2 items) adding advance button (1 item)
+		ScrollingOffsetAdd += IsAdvancedProperty ? 3 : 2;
+				
+		// We should never remove an item from favorite if there is no favorite category
+		// Set the remove offset to 0
+		ScrollingOffsetRemove = 0;
+	}
+
+	// Apply the calculated offset
+	IDetailsViewPrivate* DetailsView = OwnerTreeNodePinned->GetDetailsView();
+	DetailsView->MoveScrollOffset(bToggled ? ScrollingOffsetAdd : ScrollingOffsetRemove);
+
+	// Refresh the tree
+	DetailsView->ForceRefresh();
+}
+
+void SDetailSingleItemRow::CreateGlobalExtensionWidgets(TArray<FPropertyRowExtensionButton>& OutExtensions) const
+{
+	// fetch global extension widgets 
+	FPropertyEditorModule& PropertyEditorModule = FModuleManager::Get().GetModuleChecked<FPropertyEditorModule>("PropertyEditor");
+
+	FOnGenerateGlobalRowExtensionArgs Args { GetPropertyHandle(), GetPropertyNode(), OwnerTreeNode };
+	PropertyEditorModule.GetGlobalRowExtensionDelegate().Broadcast(Args, OutExtensions);
+}
 
-		FObjectPropertyNode* ObjectItemParent = InCustomization.GetPropertyNode()->FindObjectItemParent();
-		UClass* ObjectClass = ObjectItemParent != nullptr ? ObjectItemParent->GetObjectBaseClass() : nullptr;
-		SetKeyVisibility = ObjectClass != nullptr && KeyframeHandler.Pin()->IsPropertyKeyable(ObjectClass, *Handle) ? EVisibility::Visible : EVisibility::Collapsed;
+bool SDetailSingleItemRow::IsKeyframeButtonVisible() const
+{
+	TSharedPtr<IDetailKeyframeHandler> KeyframeHandler = OwnerTreeNode.Pin()->GetDetailsView()->GetKeyframeHandler();
+	if (!Customization->HasPropertyNode() || !KeyframeHandler.IsValid())
+	{
+		return false;
 	}
 
-	return
-		SNew(SButton)
-		.HAlign(HAlign_Left)
-		.VAlign(VAlign_Center)
-		.ContentPadding(0.0f)
-		.ButtonStyle(FEditorStyle::Get(), "Sequencer.AddKey.Details")
-		.Visibility(SetKeyVisibility)
-		.IsEnabled( this, &SDetailSingleItemRow::IsKeyframeButtonEnabled, InTreeNode )
-		.ToolTipText( NSLOCTEXT("PropertyView", "AddKeyframeButton_ToolTip", "Adds a keyframe for this property to the current animation") )
-		.OnClicked( this, &SDetailSingleItemRow::OnAddKeyframeClicked );
+	TSharedPtr<IPropertyHandle> Handle = PropertyEditorHelpers::GetPropertyHandle(Customization->GetPropertyNode().ToSharedRef(), nullptr, nullptr);
+
+	FObjectPropertyNode* ObjectItemParent = Customization->GetPropertyNode()->FindObjectItemParent();
+	UClass* ObjectClass = ObjectItemParent != nullptr ? ObjectItemParent->GetObjectBaseClass() : nullptr;
+	return ObjectClass != nullptr && KeyframeHandler->IsPropertyKeyable(ObjectClass, *Handle);
 }
 
-bool SDetailSingleItemRow::IsKeyframeButtonEnabled(TSharedRef<FDetailTreeNode> InTreeNode) const
+bool SDetailSingleItemRow::IsKeyframeButtonEnabled() const
 {
-	return
-		InTreeNode->IsPropertyEditingEnabled().Get() &&
+	TSharedPtr<FDetailTreeNode> OwnerTreeNodePinned = OwnerTreeNode.Pin();
+	if (OwnerTreeNodePinned.IsValid() == false)
+	{
+		return false;
+	}
+
+	TSharedPtr<IDetailKeyframeHandler> KeyframeHandler = OwnerTreeNodePinned->GetDetailsView()->GetKeyframeHandler();
+	return OwnerTreeNodePinned->IsPropertyEditingEnabled().Get() &&
 		KeyframeHandler.IsValid() &&
-		KeyframeHandler.Pin()->IsPropertyKeyingEnabled();
+		KeyframeHandler->IsPropertyKeyingEnabled();
 }
 
-FReply SDetailSingleItemRow::OnAddKeyframeClicked()
+void SDetailSingleItemRow::OnAddKeyframeClicked()
 {	
-	if( KeyframeHandler.IsValid() )
+	TSharedPtr<IDetailKeyframeHandler> KeyframeHandler = OwnerTreeNode.Pin()->GetDetailsView()->GetKeyframeHandler();
+	if (KeyframeHandler.IsValid())
 	{
 		TSharedPtr<IPropertyHandle> Handle = PropertyEditorHelpers::GetPropertyHandle(Customization->GetPropertyNode().ToSharedRef(), nullptr, nullptr);
 
-		KeyframeHandler.Pin()->OnKeyPropertyClicked(*Handle);
+		KeyframeHandler->OnKeyPropertyClicked(*Handle);
 	}
-
-	return FReply::Handled();
 }
 
 bool SDetailSingleItemRow::IsHighlighted() const
diff --git a/Engine/Source/Editor/PropertyEditor/Private/SDetailSingleItemRow.h b/Engine/Source/Editor/PropertyEditor/Private/SDetailSingleItemRow.h
--- a/Engine/Source/Editor/PropertyEditor/Private/SDetailSingleItemRow.h
+++ b/Engine/Source/Editor/PropertyEditor/Private/SDetailSingleItemRow.h
@@ -3,16 +3,18 @@
 #pragma once
 
 #include "CoreMinimal.h"
-#include "Input/Reply.h"
-#include "Widgets/SWidget.h"
 #include "DetailTreeNode.h"
-#include "Widgets/Views/STableViewBase.h"
-#include "Widgets/DeclarativeSyntaxSupport.h"
-#include "SDetailsViewBase.h"
+#include "PropertyCustomizationHelpers.h"
 #include "SDetailTableRowBase.h"
-#include "DragAndDrop/DecoratedDragDropOp.h"
+#include "SDetailsViewBase.h"
 #include "ScopedTransaction.h"
-#include "PropertyCustomizationHelpers.h"
+
+#include "DragAndDrop/DecoratedDragDropOp.h"
+#include "Framework/Commands/Commands.h"
+#include "Input/Reply.h"
+#include "Widgets/DeclarativeSyntaxSupport.h"
+#include "Widgets/SWidget.h"
+#include "Widgets/Views/STableViewBase.h"
 
 class IDetailKeyframeHandler;
 struct FDetailLayoutCustomization;
@@ -70,10 +72,16 @@ private:
 	bool CanPasteProperty() const;
 	FSlateColor GetOuterBackgroundColor() const;
 	FSlateColor GetInnerBackgroundColor() const;
-	TSharedRef<SWidget> CreateExtensionWidget( TSharedRef<SWidget> ValueWidget, FDetailLayoutCustomization& InCustomization, TSharedRef<FDetailTreeNode> InTreeNode );
-	TSharedRef<SWidget> CreateKeyframeButton( FDetailLayoutCustomization& InCustomization, TSharedRef<FDetailTreeNode> InTreeNode );
-	bool IsKeyframeButtonEnabled(TSharedRef<FDetailTreeNode> InTreeNode) const;
-	FReply OnAddKeyframeClicked();
+
+	void CreateGlobalExtensionWidgets(TArray<FPropertyRowExtensionButton>& ExtensionButtons) const;
+	TSharedRef<SWidget> CreateExtensionWidget() const;
+	TSharedRef<SWidget> CreateKeyframeButton();
+	void OnResetToDefaultClicked() const;
+	bool IsResetToDefaultEnabled() const;
+
+	bool IsKeyframeButtonVisible() const;
+	bool IsKeyframeButtonEnabled() const;
+	void OnAddKeyframeClicked();
 	bool IsHighlighted() const;
 
 	void OnFavoriteMenuToggle();
@@ -95,7 +103,6 @@ private:
 	TSharedPtr<IPropertyHandle> GetPropertyHandle() const;
 
 private:
-	TWeakPtr<IDetailKeyframeHandler> KeyframeHandler;
 	/** Customization for this widget */
 	FDetailLayoutCustomization* Customization;
 	bool bAllowFavoriteSystem;
diff --git a/Engine/Source/Editor/PropertyEditor/Public/PropertyEditorDelegates.h b/Engine/Source/Editor/PropertyEditor/Public/PropertyEditorDelegates.h
--- a/Engine/Source/Editor/PropertyEditor/Public/PropertyEditorDelegates.h
+++ b/Engine/Source/Editor/PropertyEditor/Public/PropertyEditorDelegates.h
@@ -3,6 +3,8 @@
 #pragma once
 
 #include "CoreMinimal.h"
+#include "Framework/Commands/UIAction.h"
+#include "Textures/SlateIcon.h"
 #include "UObject/WeakObjectPtr.h"
 
 class FPropertyPath;
@@ -84,13 +86,29 @@ struct FOnGenerateGlobalRowExtensionArgs
 	TWeakPtr<class IDetailTreeNode> OwnerTreeNode;
 };
 
+/** 
+ * A property row extension button is displayed at the end of a property row, either inline as a button, 
+ * or in a dropdown when not all buttons can fit.
+ */
+struct FPropertyRowExtensionButton
+{
+	/** The icon to display for the button. */
+	TAttribute<FSlateIcon> Icon;
+	/** The label to display for the button when shown in the dropdown. */
+	TAttribute<FText> Label;
+	/** The tooltip to display for the button. */
+	TAttribute<FText> ToolTip;
+	/** The UIAction to use for the button - this includes on execute, can execute and visibility handlers. */
+	FUIAction UIAction;
+};
+
 /**
  * Delegate called to get add an extension to a property row's name column.
  * To use, bind an handler to the delegate that adds an extension to the out array parameter.
  * When called, EWidgetPosition indicates the position for which the delegate is gathering extensions.
  * ie. The favorite system is implemented by adding the star widget when the delegate is called with the left position.
  */
-DECLARE_MULTICAST_DELEGATE_TwoParams(FOnGenerateGlobalRowExtension, const FOnGenerateGlobalRowExtensionArgs& /*InArgs*/, TArray<TSharedRef<class SWidget>>& /*OutExtensions*/);
+DECLARE_MULTICAST_DELEGATE_TwoParams(FOnGenerateGlobalRowExtension, const FOnGenerateGlobalRowExtensionArgs& /*InArgs*/, TArray<FPropertyRowExtensionButton>& /*OutExtensions*/);
 
 /**
  * Callback executed to query the custom layout of details
