Performance improvements for SDetailSingleItemRow.

Cached the reset to default enabled state.

Made FSlateIcons static.

Cached FDetailWidgetRow instead of potentially generating widgets multiple times in SDetailSingleItemRow. Added FDetailWidgetRow operator= that was preventing the caching of the row.

Removed extra padding that was being introduced by the new extension button toolbar.

diff --git a/Engine/Source/Editor/EditorStyle/Private/StarshipStyle.cpp b/Engine/Source/Editor/EditorStyle/Private/StarshipStyle.cpp
--- a/Engine/Source/Editor/EditorStyle/Private/StarshipStyle.cpp
+++ b/Engine/Source/Editor/EditorStyle/Private/StarshipStyle.cpp
@@ -3259,13 +3259,15 @@ void FStarshipEditorStyle::FStyle::SetupPropertyEditorStyles()
 			.SetIconSize(Icon16x16)
 			.SetLabelPadding(FMargin(0))
 			.SetComboButtonPadding(FMargin(0))
-			.SetBlockPadding(FMargin(2.0f, 0.0f))
+			.SetBlockPadding(FMargin(2, 0))
 			.SetIndentedBlockPadding(FMargin(0))
+			.SetIndentedBlockPadding(FMargin(0))
+			.SetBackgroundPadding(FMargin(0))
 			.SetButtonPadding(FMargin(2, 0))
-			.SetCheckBoxPadding(FMargin(3.0f, 0.0f))
+			.SetCheckBoxPadding(FMargin(3, 0))
 			.SetButtonStyle(DetailsExtensionMenuButton)
 			.SetSeparatorBrush(FSlateNoResource())
-			.SetSeparatorPadding(FMargin(4.0f, 0.0f))
+			.SetSeparatorPadding(FMargin(4, 0))
 			.SetExpandBrush(CORE_IMAGE_BRUSH_SVG("Starship/Common/ellipsis-vertical-narrow", FVector2D(6, 24)));
 
 		Set("DetailsView.ExtensionToolBar", DetailsExtensionToolBarStyle);
diff --git a/Engine/Source/Editor/PropertyEditor/Private/SDetailSingleItemRow.cpp b/Engine/Source/Editor/PropertyEditor/Private/SDetailSingleItemRow.cpp
--- a/Engine/Source/Editor/PropertyEditor/Private/SDetailSingleItemRow.cpp
+++ b/Engine/Source/Editor/PropertyEditor/Private/SDetailSingleItemRow.cpp
@@ -200,15 +200,39 @@ TSharedPtr<IPropertyHandle> SDetailSingleItemRow::GetPropertyHandle() const
 		IDetailsViewPrivate* DetailsView = OwnerTreeNode.Pin()->GetDetailsView();
 		Handle = PropertyEditorHelpers::GetPropertyHandle(PropertyNode.ToSharedRef(), DetailsView->GetNotifyHook(), DetailsView->GetPropertyUtilities());
 	}
-	else if (Customization->GetWidgetRow().PropertyHandles.Num() > 0)
+	else if (WidgetRow.PropertyHandles.Num() > 0)
 	{
 		// @todo: Handle more than 1 property handle?
-		Handle = Customization->GetWidgetRow().PropertyHandles[0];
+		Handle = WidgetRow.PropertyHandles[0];
 	}
 
 	return Handle;
 }
 
+bool SDetailSingleItemRow::UpdateResetToDefault()
+{
+	TSharedPtr<IPropertyHandle> PropertyHandle = GetPropertyHandle();
+	if (PropertyHandle.IsValid())
+	{
+		if (PropertyHandle->HasMetaData("NoResetToDefault") || PropertyHandle->GetInstanceMetaData("NoResetToDefault"))
+		{
+			return false;
+		}
+	}
+
+	if (WidgetRow.CustomResetToDefault.IsSet())
+	{
+		return WidgetRow.CustomResetToDefault.GetValue().IsResetToDefaultVisible(PropertyHandle);
+	}
+	else if (PropertyHandle.IsValid())
+	{
+		return PropertyHandle->CanResetToDefault();
+	}
+
+	return false;
+
+}
+
 void SDetailSingleItemRow::Construct( const FArguments& InArgs, FDetailLayoutCustomization* InCustomization, bool bHasMultipleColumns, TSharedRef<FDetailTreeNode> InOwnerTreeNode, const TSharedRef<STableViewBase>& InOwnerTableView )
 {
 	OwnerTreeNode = InOwnerTreeNode;
@@ -229,25 +253,23 @@ void SDetailSingleItemRow::Construct( const FArguments& InArgs, FDetailLayoutCus
 	{
 		if (Customization->IsValidCustomization())
 		{
-			FDetailWidgetRow Row = InCustomization->GetWidgetRow();
-
-			TSharedPtr<SWidget> NameWidget, ValueWidget;
+			WidgetRow = Customization->GetWidgetRow();
 
-			NameWidget = Row.NameWidget.Widget;
+			TSharedPtr<SWidget> NameWidget = WidgetRow.NameWidget.Widget;
 
-			ValueWidget =
+			TSharedPtr<SWidget> ValueWidget =
 				SNew(SConstrainedBox)
-				.MinWidth(Row.ValueWidget.MinWidth)
-				.MaxWidth(Row.ValueWidget.MaxWidth)
+				.MinWidth(WidgetRow.ValueWidget.MinWidth)
+				.MaxWidth(WidgetRow.ValueWidget.MaxWidth)
 				[
-					Row.ValueWidget.Widget
+					WidgetRow.ValueWidget.Widget
 				];
 
 			TAttribute<bool> IsEnabledAttribute;
-			if (Row.IsEnabledAttr.IsBound())
+			if (WidgetRow.IsEnabledAttr.IsBound())
 			{
 				TAttribute<bool> PropertyEnabledAttr = InOwnerTreeNode->IsPropertyEditingEnabled();
-				TAttribute<bool> RowEnabledAttr = Row.IsEnabledAttr;
+				TAttribute<bool> RowEnabledAttr = WidgetRow.IsEnabledAttr;
 				IsEnabledAttribute = TAttribute<bool>::Create(
 					[RowEnabledAttr, PropertyEnabledAttr]()
 					{
@@ -290,8 +312,8 @@ void SDetailSingleItemRow::Construct( const FArguments& InArgs, FDetailLayoutCus
 					.MinWidth(20)
 					[
 						SNew(SEditConditionWidget)
-						.EditConditionValue(Row.EditConditionValue)
-						.OnEditConditionValueChanged(Row.OnEditConditionValueChanged)
+						.EditConditionValue(WidgetRow.EditConditionValue)
+						.OnEditConditionValueChanged(WidgetRow.OnEditConditionValueChanged)
 					]
 				];
 
@@ -381,8 +403,8 @@ void SDetailSingleItemRow::Construct( const FArguments& InArgs, FDetailLayoutCus
 			if (bHasMultipleColumns)
 			{
 				NameColumnBox->AddSlot()
-					.HAlign(Row.NameWidget.HorizontalAlignment)
-					.VAlign(Row.NameWidget.VerticalAlignment)
+					.HAlign(WidgetRow.NameWidget.HorizontalAlignment)
+					.VAlign(WidgetRow.NameWidget.VerticalAlignment)
 					.Padding(DetailWidgetConstants::LeftRowPadding)
 					[
 						NameWidget.ToSharedRef()
@@ -407,8 +429,8 @@ void SDetailSingleItemRow::Construct( const FArguments& InArgs, FDetailLayoutCus
 						SNew(SHorizontalBox)
 						.Clipping(EWidgetClipping::OnDemand)
 						+ SHorizontalBox::Slot()
-						.HAlign(Row.ValueWidget.HorizontalAlignment)
-						.VAlign(Row.ValueWidget.VerticalAlignment)
+						.HAlign(WidgetRow.ValueWidget.HorizontalAlignment)
+						.VAlign(WidgetRow.ValueWidget.VerticalAlignment)
 						.Padding(DetailWidgetConstants::RightRowPadding)
 						[
 							ValueWidget.ToSharedRef()
@@ -428,11 +450,11 @@ void SDetailSingleItemRow::Construct( const FArguments& InArgs, FDetailLayoutCus
 			{
 				NameColumnBox->SetEnabled(IsEnabledAttribute);
 				NameColumnBox->AddSlot()
-					.HAlign(Row.WholeRowWidget.HorizontalAlignment)
-					.VAlign(Row.WholeRowWidget.VerticalAlignment)
+					.HAlign(WidgetRow.WholeRowWidget.HorizontalAlignment)
+					.VAlign(WidgetRow.WholeRowWidget.VerticalAlignment)
 					.Padding(DetailWidgetConstants::LeftRowPadding)
 					[
-						Row.WholeRowWidget.Widget
+						WidgetRow.WholeRowWidget.Widget
 					];
 
 				InnerSplitter->AddSlot()
@@ -444,22 +466,29 @@ void SDetailSingleItemRow::Construct( const FArguments& InArgs, FDetailLayoutCus
 			TArray<FPropertyRowExtensionButton> ExtensionButtons;
 
 			FPropertyRowExtensionButton& ResetToDefault = ExtensionButtons.AddDefaulted_GetRef();
-			ResetToDefault.Icon = TAttribute<FSlateIcon>::Create([this]()
-				{
-					return IsResetToDefaultEnabled() ?
-						FSlateIcon(FAppStyle::Get().GetStyleSetName(), "PropertyWindow.DiffersFromDefault") :
-						FSlateIcon(FAppStyle::Get().GetStyleSetName(), "NoBrush");
-				});
 			ResetToDefault.Label = NSLOCTEXT("PropertyEditor", "ResetToDefault", "Reset to Default");
-			ResetToDefault.ToolTip = TAttribute<FText>::Create([this]()
-				{
-					return IsResetToDefaultEnabled() ?
-						NSLOCTEXT("PropertyEditor", "ResetToDefaultToolTip", "Reset this property to its default value.") :
-						FText::GetEmpty();
-				});
 			ResetToDefault.UIAction = FUIAction(
 				FExecuteAction::CreateSP(this, &SDetailSingleItemRow::OnResetToDefaultClicked),
-				FCanExecuteAction::CreateSP(this, &SDetailSingleItemRow::IsResetToDefaultEnabled));
+				FCanExecuteAction::CreateSP(this, &SDetailSingleItemRow::IsResetToDefaultEnabled)
+			);
+
+			// We could just collapse the Reset to Default button by setting the FIsActionButtonVisible delegate,
+			// but this would cause the reset to defaults not to reserve space in the toolbar and not be aligned across all rows.
+			// Instead, we show an empty icon and tooltip and disable the button.
+			static FSlateIcon EnabledResetToDefaultIcon(FAppStyle::Get().GetStyleSetName(), "PropertyWindow.DiffersFromDefault");
+			static FSlateIcon DisabledResetToDefaultIcon(FAppStyle::Get().GetStyleSetName(), "NoBrush");
+			ResetToDefault.Icon = TAttribute<FSlateIcon>::Create([this]()
+			{
+				return IsResetToDefaultEnabled() ? 
+					EnabledResetToDefaultIcon :
+					DisabledResetToDefaultIcon;
+			});
+			ResetToDefault.ToolTip = TAttribute<FText>::Create([this]() 
+			{
+				return IsResetToDefaultEnabled() ?
+					NSLOCTEXT("PropertyEditor", "ResetToDefaultToolTip", "Reset this property to its default value.") :
+					FText::GetEmpty();
+			});
 
 			FPropertyRowExtensionButton& CreateKey = ExtensionButtons.AddDefaulted_GetRef();
 			CreateKey.Icon = FSlateIcon(FAppStyle::Get().GetStyleSetName(), "Sequencer.AddKey.Details");
@@ -538,34 +567,12 @@ void SDetailSingleItemRow::Construct( const FArguments& InArgs, FDetailLayoutCus
 
 bool SDetailSingleItemRow::IsResetToDefaultEnabled() const
 {
-	FDetailWidgetRow WidgetRow = Customization->GetWidgetRow();
-	TSharedPtr<IPropertyHandle> PropertyHandle = GetPropertyHandle();
-
-	if (PropertyHandle.IsValid())
-	{
-		if (PropertyHandle->HasMetaData("NoResetToDefault") || PropertyHandle->GetInstanceMetaData("NoResetToDefault"))
-		{
-			return false;
-		}
-	}
-
-	if (WidgetRow.CustomResetToDefault.IsSet())
-	{
-		return WidgetRow.CustomResetToDefault.GetValue().IsResetToDefaultVisible(PropertyHandle);
-	}
-	else if (PropertyHandle.IsValid())
-	{
-		return PropertyHandle->CanResetToDefault();
-	} 
-
-	return false;
+	return bCachedResetToDefaultEnabled;
 }
 
 void SDetailSingleItemRow::OnResetToDefaultClicked() const
 {
-	FDetailWidgetRow WidgetRow = Customization->GetWidgetRow();
 	TSharedPtr<IPropertyHandle> PropertyHandle = GetPropertyHandle();
-
 	if (WidgetRow.CustomResetToDefault.IsSet())
 	{
 		WidgetRow.CustomResetToDefault.GetValue().OnResetToDefaultClicked(PropertyHandle);
@@ -626,15 +633,15 @@ FSlateColor SDetailSingleItemRow::GetInnerBackgroundColor() const
 
 bool SDetailSingleItemRow::OnContextMenuOpening(FMenuBuilder& MenuBuilder)
 {
-	const bool bIsCopyPasteBound = Customization->GetWidgetRow().IsCopyPasteBound();
+	const bool bIsCopyPasteBound = WidgetRow.IsCopyPasteBound();
 
 	FUIAction CopyAction;
 	FUIAction PasteAction;
 
 	if (bIsCopyPasteBound)
 	{
-		CopyAction = Customization->GetWidgetRow().CopyMenuAction;
-		PasteAction = Customization->GetWidgetRow().PasteMenuAction;
+		CopyAction = WidgetRow.CopyMenuAction;
+		PasteAction = WidgetRow.PasteMenuAction;
 	}
 	else
 	{
@@ -702,8 +709,7 @@ bool SDetailSingleItemRow::OnContextMenuOpening(FMenuBuilder& MenuBuilder)
 		FSlateIcon(FEditorStyle::Get().GetStyleSetName(), FavoriteIcon),
 		FavoriteAction);
 
-	const TArray<FDetailWidgetRow::FCustomMenuData>& CustomMenuActions = Customization->GetWidgetRow().CustomMenuItems;
-	if (CustomMenuActions.Num() > 0)
+	if (WidgetRow.CustomMenuItems.Num() > 0)
 	{
 		// Hide separator line if it only contains the SearchWidget, making the next 2 elements the top of the list
 		if (MenuBuilder.GetMultiBox()->GetBlocks().Num() > 1)
@@ -711,9 +717,9 @@ bool SDetailSingleItemRow::OnContextMenuOpening(FMenuBuilder& MenuBuilder)
 			MenuBuilder.AddMenuSeparator();
 		}
 
-		for (const FDetailWidgetRow::FCustomMenuData& CustomMenuData : CustomMenuActions)
+		for (const FDetailWidgetRow::FCustomMenuData& CustomMenuData : WidgetRow.CustomMenuItems)
 		{
-			//Add the menu entry
+			// Add the menu entry
 			MenuBuilder.AddMenuEntry(
 				CustomMenuData.Name,
 				CustomMenuData.Tooltip,
@@ -844,7 +850,7 @@ bool SDetailSingleItemRow::CanPasteProperty() const
 	}
 
 	FString ClipboardContent;
-	if( OwnerTreeNode.IsValid() )
+	if (OwnerTreeNode.IsValid())
 	{
 		FPlatformApplicationMisc::ClipboardPaste(ClipboardContent);
 	}
@@ -1021,6 +1027,11 @@ void SDetailSingleItemRow::SetIsDragDrop(bool bInIsDragDrop)
 	bIsDragDropObject = bInIsDragDrop;
 }
 
+void SDetailSingleItemRow::Tick(const FGeometry& AllottedGeometry, const double InCurrentTime, const float InDeltaTime)
+{
+	bCachedResetToDefaultEnabled = UpdateResetToDefault();
+}
+
 void SArrayRowHandle::Construct(const FArguments& InArgs)
 {
 	ParentRow = InArgs._ParentRow;
diff --git a/Engine/Source/Editor/PropertyEditor/Private/SDetailSingleItemRow.h b/Engine/Source/Editor/PropertyEditor/Private/SDetailSingleItemRow.h
--- a/Engine/Source/Editor/PropertyEditor/Private/SDetailSingleItemRow.h
+++ b/Engine/Source/Editor/PropertyEditor/Private/SDetailSingleItemRow.h
@@ -63,9 +63,9 @@ public:
 
 	void SetIsDragDrop(bool bInIsDragDrop);
 
+	virtual void Tick(const FGeometry& AllottedGeometry, const double InCurrentTime, const float InDeltaTime);
 protected:
 	virtual bool OnContextMenuOpening( FMenuBuilder& MenuBuilder ) override;
-
 private:
 	void OnCopyProperty();
 	void OnPasteProperty();
@@ -102,12 +102,15 @@ private:
 	TSharedPtr<FPropertyNode> GetPropertyNode() const;
 	TSharedPtr<IPropertyHandle> GetPropertyHandle() const;
 
+	bool UpdateResetToDefault();
 private:
 	/** Customization for this widget */
 	FDetailLayoutCustomization* Customization;
+	FDetailWidgetRow WidgetRow;
 	bool bAllowFavoriteSystem;
 	bool bIsHoveredDragTarget;
 	bool bIsDragDropObject;
+	bool bCachedResetToDefaultEnabled;
 	TSharedPtr<FPropertyNode> SwappablePropertyNode;
 	TSharedPtr<SButton> ExpanderArrow;
 };
diff --git a/Engine/Source/Editor/PropertyEditor/Public/DetailWidgetRow.h b/Engine/Source/Editor/PropertyEditor/Public/DetailWidgetRow.h
--- a/Engine/Source/Editor/PropertyEditor/Public/DetailWidgetRow.h
+++ b/Engine/Source/Editor/PropertyEditor/Public/DetailWidgetRow.h
@@ -26,15 +26,24 @@ public:
 		, VerticalAlignment( InVAlign )
 		, MinWidth( InMinWidth )
 		, MaxWidth( InMaxWidth )
-		, ParentDecl( InParentDecl )
+		, ParentDecl( &InParentDecl )
 	{
+	}
 
+	FDetailWidgetDecl( class FDetailWidgetRow& InParentDecl, const FDetailWidgetDecl& Other )
+		: Widget( Other.Widget )
+		, HorizontalAlignment( Other.HorizontalAlignment )
+		, VerticalAlignment( Other.VerticalAlignment )
+		, MinWidth( Other.MinWidth )
+		, MaxWidth( Other.MaxWidth )
+		, ParentDecl( &InParentDecl )
+	{
 	}
 
 	FDetailWidgetRow& operator[]( TSharedRef<SWidget> InWidget )
 	{
 		Widget = InWidget;
-		return ParentDecl;
+		return *ParentDecl;
 	}
 
 	FDetailWidgetDecl& VAlign( EVerticalAlignment InAlignment )
@@ -81,7 +90,7 @@ public:
 	TOptional<float> MinWidth;
 	TOptional<float> MaxWidth;
 private:
-	class FDetailWidgetRow& ParentDecl;
+	class FDetailWidgetRow* ParentDecl;
 };
 
 
@@ -108,7 +117,26 @@ public:
 		, RowTagName()
 	{
 	}
-	
+
+	FDetailWidgetRow& operator=(const FDetailWidgetRow& Other)
+	{
+		NameWidget = FDetailWidgetDecl(*this, Other.NameWidget);
+		ValueWidget = FDetailWidgetDecl(*this, Other.ValueWidget);
+		WholeRowWidget = FDetailWidgetDecl(*this, Other.WholeRowWidget);
+		VisibilityAttr = Other.VisibilityAttr;
+		IsEnabledAttr = Other.IsEnabledAttr;
+		FilterTextString = Other.FilterTextString;
+		CopyMenuAction = Other.CopyMenuAction;
+		PasteMenuAction = Other.PasteMenuAction;
+		CustomMenuItems = Other.CustomMenuItems;
+		RowTagName = Other.RowTagName;
+		CustomResetToDefault = Other.CustomResetToDefault;
+		EditConditionValue = Other.EditConditionValue;
+		OnEditConditionValueChanged = Other.OnEditConditionValueChanged;
+		PropertyHandles = Other.PropertyHandles;
+		return *this;
+	}
+
 	/**
 	 * Assigns content to the entire row
 	 */
@@ -304,7 +332,6 @@ public:
 	};
 	/** Custom Action on this row */
 	TArray<FCustomMenuData> CustomMenuItems;
-
 	/* Tag to identify this row */
 	FName RowTagName;
 	/** Custom reset to default handler */
