Added GetNodeDetailsView() to IDetailTreeNode, which returns the IDetailsView* that the node is part of.

Note: The odd naming is because GetDetailsView() is already in use in FDetailTreeNode, which returns an IDetailsViewPrivate*, and changing that would require refactoring many callsites and potentially breaking licensee code.

diff --git a/Engine/Source/Editor/PropertyEditor/Private/DetailAdvancedDropdownNode.h b/Engine/Source/Editor/PropertyEditor/Private/DetailAdvancedDropdownNode.h
--- a/Engine/Source/Editor/PropertyEditor/Private/DetailAdvancedDropdownNode.h
+++ b/Engine/Source/Editor/PropertyEditor/Private/DetailAdvancedDropdownNode.h
@@ -32,7 +32,8 @@ public:
 	{}
 private:
 	/** IDetailTreeNode Interface */
-	virtual IDetailsViewPrivate* GetDetailsView() const override{ return ParentCategory.GetDetailsView(); }
+	virtual IDetailsView* GetNodeDetailsView() const override { return ParentCategory.GetNodeDetailsView(); }
+	virtual IDetailsViewPrivate* GetDetailsView() const override { return ParentCategory.GetDetailsView(); }
 	virtual TSharedRef< ITableRow > GenerateWidgetForTableView( const TSharedRef<STableViewBase>& OwnerTable, bool bAllowFavoriteSystem) override;
 	virtual bool GenerateStandaloneWidget(FDetailWidgetRow& OutRow) const override;
 	virtual void GetChildren(FDetailNodeList& OutChildren) override {}
diff --git a/Engine/Source/Editor/PropertyEditor/Private/DetailCategoryBuilderImpl.h b/Engine/Source/Editor/PropertyEditor/Private/DetailCategoryBuilderImpl.h
--- a/Engine/Source/Editor/PropertyEditor/Private/DetailCategoryBuilderImpl.h
+++ b/Engine/Source/Editor/PropertyEditor/Private/DetailCategoryBuilderImpl.h
@@ -180,6 +180,7 @@ public:
 	virtual void SetSortOrder(int32 InSortOrder) override;
 
 	/** FDetailTreeNode interface */
+	virtual IDetailsView* GetNodeDetailsView() const override { return DetailLayoutBuilder.Pin()->GetDetailsView(); }
 	virtual IDetailsViewPrivate* GetDetailsView() const override { return DetailLayoutBuilder.Pin()->GetDetailsView(); }
 	virtual TSharedRef< ITableRow > GenerateWidgetForTableView(const TSharedRef<STableViewBase>& OwnerTable, bool bAllowFavoriteSystem) override;
 	virtual bool GenerateStandaloneWidget(FDetailWidgetRow& OutRow) const override;
diff --git a/Engine/Source/Editor/PropertyEditor/Private/DetailCategoryGroupNode.h b/Engine/Source/Editor/PropertyEditor/Private/DetailCategoryGroupNode.h
--- a/Engine/Source/Editor/PropertyEditor/Private/DetailCategoryGroupNode.h
+++ b/Engine/Source/Editor/PropertyEditor/Private/DetailCategoryGroupNode.h
@@ -57,7 +57,8 @@ public:
 	bool GetHasSplitter() const { return bHasSplitter; }
 
 private:
-	virtual IDetailsViewPrivate* GetDetailsView() const override{ return ParentCategory.GetDetailsView(); }
+	virtual IDetailsView* GetNodeDetailsView() const override { return ParentCategory.GetNodeDetailsView(); }
+	virtual IDetailsViewPrivate* GetDetailsView() const override { return ParentCategory.GetDetailsView(); }
 	virtual void OnItemExpansionChanged( bool bIsExpanded, bool bShouldSaveState) override {}
 	virtual bool ShouldBeExpanded() const override { return true; }
 	virtual ENodeVisibility GetVisibility() const override { return bShouldBeVisible ? ENodeVisibility::Visible : ENodeVisibility::HiddenDueToFiltering; }
@@ -68,7 +69,7 @@ private:
 	virtual TSharedPtr<IPropertyHandle> CreatePropertyHandle() const override { return nullptr; }
 	virtual void GetFilterStrings(TArray<FString>& OutFilterStrings) const override { OutFilterStrings.Add(GroupName.ToString()); };
 
-	virtual void GetChildren(FDetailNodeList& OutChildren )  override;
+	virtual void GetChildren(FDetailNodeList& OutChildren ) override;
 	virtual void FilterNode( const FDetailFilter& InFilter ) override;
 	virtual void Tick( float DeltaTime ) override {}
 	virtual bool ShouldShowOnlyChildren() const override { return false; }
diff --git a/Engine/Source/Editor/PropertyEditor/Private/DetailItemNode.h b/Engine/Source/Editor/PropertyEditor/Private/DetailItemNode.h
--- a/Engine/Source/Editor/PropertyEditor/Private/DetailItemNode.h
+++ b/Engine/Source/Editor/PropertyEditor/Private/DetailItemNode.h
@@ -53,6 +53,7 @@ public:
 	bool HasGeneratedChildren() const { return Children.Num() > 0;}
 
 	/** FDetailTreeNode interface */
+	virtual IDetailsView* GetNodeDetailsView() const override { return ParentCategory.Pin()->GetNodeDetailsView(); }
 	virtual IDetailsViewPrivate* GetDetailsView() const override{ return ParentCategory.Pin()->GetDetailsView(); }
 	virtual TSharedRef< ITableRow > GenerateWidgetForTableView( const TSharedRef<STableViewBase>& OwnerTable, bool bAllowFavoriteSystem) override;
 	virtual bool GenerateStandaloneWidget(FDetailWidgetRow& OutRow) const override;
diff --git a/Engine/Source/Editor/PropertyEditor/Private/DetailMultiTopLevelObjectRootNode.h b/Engine/Source/Editor/PropertyEditor/Private/DetailMultiTopLevelObjectRootNode.h
--- a/Engine/Source/Editor/PropertyEditor/Private/DetailMultiTopLevelObjectRootNode.h
+++ b/Engine/Source/Editor/PropertyEditor/Private/DetailMultiTopLevelObjectRootNode.h
@@ -47,6 +47,7 @@ class FDetailMultiTopLevelObjectRootNode : public FDetailTreeNode, public TShare
 public:
 	FDetailMultiTopLevelObjectRootNode(const FDetailNodeList& InChildNodes, const TSharedPtr<IDetailRootObjectCustomization>& RootObjectCustomization, IDetailsViewPrivate* InDetailsView, const FObjectPropertyNode* RootNode);
 private:
+	virtual IDetailsView* GetNodeDetailsView() const override { return DetailsView; }
 	virtual IDetailsViewPrivate* GetDetailsView() const override { return DetailsView; }
 	virtual void OnItemExpansionChanged(bool bIsExpanded, bool bShouldSaveState) override;
 	virtual bool ShouldBeExpanded() const override;
diff --git a/Engine/Source/Editor/PropertyEditor/Public/IDetailTreeNode.h b/Engine/Source/Editor/PropertyEditor/Public/IDetailTreeNode.h
--- a/Engine/Source/Editor/PropertyEditor/Public/IDetailTreeNode.h
+++ b/Engine/Source/Editor/PropertyEditor/Public/IDetailTreeNode.h
@@ -78,6 +78,11 @@ class IDetailTreeNode
 public:
 	virtual ~IDetailTreeNode() {}
 
+	/** 
+	 * @return Get the details view that this node is contained in.
+	 */
+	virtual class IDetailsView* GetNodeDetailsView() const = 0;
+
 	/**
 	 * @return The type of this node.  Should be used to determine any external styling to apply to the generated r ow
 	 */
@@ -111,6 +116,9 @@ public:
 	 */
 	virtual FName GetNodeName() const = 0;
 
+	/**
+	 * Get the property row that this node is represented by.
+	 */
 	virtual TSharedPtr<class IDetailPropertyRow> GetRow() const = 0;
 
 	/**
