Advanced dropdown nodes in the details view no longer have borders that intersect with the scrollbar.

diff --git a/Engine/Source/Editor/PropertyEditor/Private/DetailAdvancedDropdownNode.cpp b/Engine/Source/Editor/PropertyEditor/Private/DetailAdvancedDropdownNode.cpp
--- a/Engine/Source/Editor/PropertyEditor/Private/DetailAdvancedDropdownNode.cpp
+++ b/Engine/Source/Editor/PropertyEditor/Private/DetailAdvancedDropdownNode.cpp
@@ -34,75 +34,54 @@ public:
 		TSharedPtr<SWidget> ContentWidget;
 		if( bIsTopNode )
 		{
-			ContentWidget = SNew( SBorder )
-				.BorderImage(FEditorStyle::GetBrush("DetailsView.CategoryMiddle") )
-				.Padding(FMargin( 0.0f, 0.0f, SDetailTableRowBase::ScrollbarPaddingSize, 0.0f ) )
-				[
-					SNew( SImage )
-					.Image(FEditorStyle::GetBrush("DetailsView.AdvancedDropdownBorder.Open"))
-				];
+			ContentWidget = SNew(SSpacer);
 		}
 		else if( InArgs._ShouldShowAdvancedButton )
 		{
-			ContentWidget = 
-				SNew( SBorder )
-				.BorderImage( FEditorStyle::GetBrush("DetailsView.AdvancedDropdownBorder") )
-				.Padding( FMargin( 0.0f, 3.0f, SDetailTableRowBase::ScrollbarPaddingSize, 0.0f ) )
+			ContentWidget = 		
+				SNew( SVerticalBox )
+				+ SVerticalBox::Slot()
+				.HAlign( HAlign_Center )
+				.AutoHeight()
 				[
-					SNew( SVerticalBox )
-					+SVerticalBox::Slot()
-					.HAlign( HAlign_Center )
-					.AutoHeight()
-					[
-						SNew( STextBlock )
-						.Text( NSLOCTEXT("DetailsView", "NoSimpleProperties", "Click the arrow to display advanced properties") )
-						.Font( IDetailLayoutBuilder::GetDetailFont() )
-						.Visibility( this, &SAdvancedDropdownRow::OnGetHelpTextVisibility )
-						.ColorAndOpacity(FLinearColor(1,1,1,.5))
-					]
-					+ SVerticalBox::Slot()
-					.AutoHeight()
+					SNew( STextBlock )
+					.Text( NSLOCTEXT("DetailsView", "NoSimpleProperties", "Click the arrow to display advanced properties") )
+					.Font( IDetailLayoutBuilder::GetDetailFont() )
+					.Visibility( this, &SAdvancedDropdownRow::OnGetHelpTextVisibility )
+					.ColorAndOpacity(FLinearColor(1,1,1,.5))
+				]
+				+ SVerticalBox::Slot()
+				.AutoHeight()
+				[
+					SAssignNew(ExpanderButton, SButton)
+					.ButtonStyle(FEditorStyle::Get(), "NoBorder")
+					.HAlign(HAlign_Center)
+					.ContentPadding(2)
+					.OnClicked(InArgs._OnClicked)
+					.IsEnabled(InArgs._IsButtonEnabled)
+					.ToolTipText(this, &SAdvancedDropdownRow::GetAdvancedPulldownToolTipText )
 					[
-						SAssignNew(ExpanderButton, SButton)
-						.ButtonStyle(FEditorStyle::Get(), "NoBorder")
-						.HAlign(HAlign_Center)
-						.ContentPadding(2)
-						.OnClicked(InArgs._OnClicked)
-						.IsEnabled(InArgs._IsButtonEnabled)
-						.ToolTipText(this, &SAdvancedDropdownRow::GetAdvancedPulldownToolTipText )
-						[
-							SNew(SImage)
-							.Image(this, &SAdvancedDropdownRow::GetAdvancedPulldownImage)
-						]
+						SNew(SImage)
+						.Image(this, &SAdvancedDropdownRow::GetAdvancedPulldownImage)
 					]
 				];
 		}
 		else
 		{
-			ContentWidget =
-				SNew( SBorder )
-				.BorderImage( FEditorStyle::GetBrush("DetailsView.AdvancedDropdownBorder") )
-				.Padding( FMargin( 0.0f, 0.0f, SDetailTableRowBase::ScrollbarPaddingSize, 2.0f ) )
-				[
-					SNew(SSpacer)
-				];
+			ContentWidget = SNew(SSpacer);
 		}
 		
 		ChildSlot
 		[
-			SNew( SVerticalBox)
-			+ SVerticalBox::Slot()
-			[
-				ContentWidget.ToSharedRef()
-			]
-			+ SVerticalBox::Slot()
-			.AutoHeight()
+			SNew( SBorder )
+			.BorderImage(FAppStyle::Get().GetBrush( "DetailsView.GridLine"))
+			.Padding( FMargin(0, 0, SDetailTableRowBase::ScrollbarPaddingSize, 0) )
 			[
-        		SNew(SBox)
-        		.HeightOverride(2.f)
-        		[ 
-					SNew( SBorder )
-					.BorderImage( FEditorStyle::GetBrush("DetailsView.AdvancedDropdownBorder") )
+				SNew( SBorder )
+				.BorderImage( FEditorStyle::GetBrush("DetailsView.AdvancedDropdownBorder") )
+				.Padding( FMargin( 0.0f, 3.0f, 0.0f, 2.0f ) )
+				[
+					ContentWidget.ToSharedRef()
 				]
 			]
 		];
