[Instanced properties] ResetToDefault on a property value will also update instanced properties to use deep copy (same as SDetailSingleItemRow::OnPasteProperty)

diff --git a/Engine/Source/Editor/PropertyEditor/Private/PropertyHandleImpl.cpp b/Engine/Source/Editor/PropertyEditor/Private/PropertyHandleImpl.cpp
--- a/Engine/Source/Editor/PropertyEditor/Private/PropertyHandleImpl.cpp
+++ b/Engine/Source/Editor/PropertyEditor/Private/PropertyHandleImpl.cpp
@@ -193,6 +193,62 @@ void FPropertyValueImpl::GenerateArrayIndexMapToObjectNode( TMap<FString,int32>&
 	}
 }
 
+void FPropertyValueImpl::RebuildInstancedProperties(const TSharedPtr<IPropertyHandle>& RootHandle, FPropertyNode* PropertyNode)
+{
+	if (PropertyNode != nullptr)
+	{
+		// Cache expansion state and then rebuild child nodes, in case we're pasting an array of a different size. This ensures instanced properties can be rebuild properly
+		TSet<FString> ExpandedChildPropertyPaths;
+		PropertyNode->GetExpandedChildPropertyPaths(ExpandedChildPropertyPaths);
+		PropertyNode->RebuildChildren();
+		PropertyNode->SetExpandedChildPropertyNodes(ExpandedChildPropertyPaths);
+	}
+
+	TSharedPtr<IPropertyHandle> Handle = RootHandle;
+
+	TArray<TSharedPtr<IPropertyHandle>> CopiedHandles;
+	CopiedHandles.Add(Handle);
+
+	while (CopiedHandles.Num() > 0)
+	{
+		Handle = CopiedHandles.Pop();
+
+		// Add all child properties to the list so we can check them next
+		uint32 NumChildren;
+		Handle->GetNumChildren(NumChildren);
+		for (uint32 ChildIndex = 0; ChildIndex < NumChildren; ChildIndex++)
+		{
+			CopiedHandles.Add(Handle->GetChildHandle(ChildIndex));
+		}
+
+		UObject* NewValueAsObject = nullptr;
+		if (FPropertyAccess::Success != Handle->GetValue(NewValueAsObject))
+		{
+			continue;
+		}
+
+		// Skip properties that are not instanced
+		if (Handle->GetProperty() == nullptr
+			|| !Handle->GetProperty()->HasAnyPropertyFlags(CPF_InstancedReference | CPF_ContainsInstancedReference))
+		{
+			continue;
+		}
+
+		// The property is instanced so we need to do a deep copy.
+		// Update the duplicate's outer to point to this outer. 
+		// The source's outer may be some other object/asset but we want this to own the duplicate.
+		TArray<UObject*> Outers;
+		Handle->GetOuterObjects(Outers);
+		UObject* DuplicateOuter = (Outers.Num() > 0) ? Outers[0] : nullptr;
+
+		// This does a deep copy of NewValueAsObject. It's subobjects and property data will be copied.
+		UObject* DuplicateOfNewValue = DuplicateObject<UObject>(NewValueAsObject, DuplicateOuter);
+		TArray<FString> DuplicateValueAsString;
+		DuplicateValueAsString.Add(DuplicateOfNewValue->GetPathName());
+		Handle->SetPerObjectValues(DuplicateValueAsString);
+	}
+}
+
 FPropertyAccess::Result FPropertyValueImpl::ImportText( const FString& InValue, FPropertyNode* InPropertyNode, EPropertyValueSetFlags::Type Flags )
 {
 	TArray<FObjectBaseAddress> ObjectsToModify;
@@ -878,6 +934,9 @@ void FPropertyValueImpl::ResetToDefault()
 				.ImportText(KeyNode->GetDefaultValueAsString(bUseDisplayName), EPropertyValueSetFlags::DefaultFlags);
 		}
 		ImportText(PropertyNodePin->GetDefaultValueAsString(bUseDisplayName), EPropertyValueSetFlags::DefaultFlags);
+		
+		TSharedPtr<IPropertyHandle> PropertyHandle = PropertyEditorHelpers::GetPropertyHandle(PropertyNodePin.ToSharedRef(), GetNotifyHook(), GetPropertyUtilities());
+		RebuildInstancedProperties(PropertyHandle, KeyNode.Get());
 
 		PropertyNodePin->BroadcastPropertyResetToDefault();
 	}
diff --git a/Engine/Source/Editor/PropertyEditor/Private/PropertyHandleImpl.h b/Engine/Source/Editor/PropertyEditor/Private/PropertyHandleImpl.h
--- a/Engine/Source/Editor/PropertyEditor/Private/PropertyHandleImpl.h
+++ b/Engine/Source/Editor/PropertyEditor/Private/PropertyHandleImpl.h
@@ -56,6 +56,11 @@ public:
 	 * @param InNode - node to start adding array offsets for.  This function will move upward until it gets to an object node
 	 */
 	static void GenerateArrayIndexMapToObjectNode( TMap<FString,int32>& OutArrayIndexMap, FPropertyNode* PropertyNode );
+
+	/**
+	 * Recurse through a property node hierarchy to update all instanced properties that should use deep copy
+	 */
+	static void RebuildInstancedProperties( const TSharedPtr<IPropertyHandle>& Handle, FPropertyNode* PropertyNode );
 	
 	/**
 	 * Gets the value as a string formatted for multiple values in an array                 
diff --git a/Engine/Source/Editor/PropertyEditor/Private/SDetailSingleItemRow.cpp b/Engine/Source/Editor/PropertyEditor/Private/SDetailSingleItemRow.cpp
--- a/Engine/Source/Editor/PropertyEditor/Private/SDetailSingleItemRow.cpp
+++ b/Engine/Source/Editor/PropertyEditor/Private/SDetailSingleItemRow.cpp
@@ -771,59 +771,7 @@ void SDetailSingleItemRow::OnPasteProperty()
 
 			Handle->SetValueFromFormattedString(ClipboardContent);
 
-			// Cache expansion state and then rebuild child nodes, in case we're pasting an array of a different size. This ensures instanced properties can be rebuild properly
-			TSet<FString> ExpandedChildPropertyPaths;
-			PropertyNode->GetExpandedChildPropertyPaths(ExpandedChildPropertyPaths);
-			PropertyNode->RebuildChildren();
-			PropertyNode->SetExpandedChildPropertyNodes(ExpandedChildPropertyPaths);
-
-			TArray<TSharedPtr<IPropertyHandle>> CopiedHandles;
-
-			CopiedHandles.Add(Handle);
-
-			while (CopiedHandles.Num() > 0)
-			{
-
-				Handle = CopiedHandles.Pop();
-
-				// Add all child properties to the list so we can check them next
-				uint32 NumChildren;
-				Handle->GetNumChildren(NumChildren);
-				for (uint32 ChildIndex = 0; ChildIndex < NumChildren; ChildIndex++)
-				{
-					CopiedHandles.Add(Handle->GetChildHandle(ChildIndex));
-				}
-
-				UObject* NewValueAsObject = nullptr;
-				if (FPropertyAccess::Success == Handle->GetValue(NewValueAsObject))
-				{
-				
-					// if the object is instanced, then we need to do a deep copy.
-					if (Handle->GetProperty() != nullptr
-						&& (Handle->GetProperty()->PropertyFlags & (CPF_InstancedReference | CPF_ContainsInstancedReference)) != 0)
-					{
-						UObject* DuplicateOuter = nullptr;
-
-						TArray<UObject*> Outers;
-						Handle->GetOuterObjects(Outers);
-
-						// Update the duplicate's outer to point to this outer. The source's outer may be some other object/asset
-						// but we want this to own the duplicate.
-						if (Outers.Num() > 0)
-						{
-							DuplicateOuter = Outers[0];
-						}
-
-						// This does a deep copy of NewValueAsObject. It's subobjects and property data will be 
-						// copied.
-						UObject* DuplicateOfNewValue = DuplicateObject<UObject>(NewValueAsObject, DuplicateOuter);
-						TArray<FString> DuplicateValueAsString;
-						DuplicateValueAsString.Add(DuplicateOfNewValue->GetPathName());
-						Handle->SetPerObjectValues(DuplicateValueAsString);
-
-					}
-				}
-			}
+			FPropertyValueImpl::RebuildInstancedProperties(Handle, PropertyNode.Get());
 
 			// Need to refresh the details panel in case a property was pasted over another.
 			OwnerTreeNode.Pin()->GetDetailsView()->ForceRefresh();
