Possible temp fix for Keyframe button crash. Unblocking Frosty content creators.

diff --git a/Engine/Source/Editor/PropertyEditor/Private/SDetailSingleItemRow.cpp b/Engine/Source/Editor/PropertyEditor/Private/SDetailSingleItemRow.cpp
--- a/Engine/Source/Editor/PropertyEditor/Private/SDetailSingleItemRow.cpp
+++ b/Engine/Source/Editor/PropertyEditor/Private/SDetailSingleItemRow.cpp
@@ -927,17 +927,25 @@ void SDetailSingleItemRow::CreateGlobalExtensionWidgets(TArray<FPropertyRowExten
 
 bool SDetailSingleItemRow::IsKeyframeButtonVisible() const
 {
-	TSharedPtr<IDetailKeyframeHandler> KeyframeHandler = OwnerTreeNode.Pin()->GetDetailsView()->GetKeyframeHandler();
-	if (!Customization->HasPropertyNode() || !KeyframeHandler.IsValid())
+	if (OwnerTreeNode.IsValid())
 	{
-		return false;
-	}
+		if (IDetailsViewPrivate* DetailsView = OwnerTreeNode.Pin()->GetDetailsView())
+		{
+			TSharedPtr<IDetailKeyframeHandler> KeyframeHandler = DetailsView->GetKeyframeHandler();
+			if (!Customization->HasPropertyNode() || !KeyframeHandler.IsValid())
+			{
+				return false;
+			}
 
-	TSharedPtr<IPropertyHandle> Handle = PropertyEditorHelpers::GetPropertyHandle(Customization->GetPropertyNode().ToSharedRef(), nullptr, nullptr);
+			TSharedPtr<IPropertyHandle> Handle = PropertyEditorHelpers::GetPropertyHandle(Customization->GetPropertyNode().ToSharedRef(), nullptr, nullptr);
+
+			FObjectPropertyNode* ObjectItemParent = Customization->GetPropertyNode()->FindObjectItemParent();
+			UClass* ObjectClass = ObjectItemParent != nullptr ? ObjectItemParent->GetObjectBaseClass() : nullptr;
+			return ObjectClass != nullptr && KeyframeHandler->IsPropertyKeyable(ObjectClass, *Handle);
+		}
+	}
 
-	FObjectPropertyNode* ObjectItemParent = Customization->GetPropertyNode()->FindObjectItemParent();
-	UClass* ObjectClass = ObjectItemParent != nullptr ? ObjectItemParent->GetObjectBaseClass() : nullptr;
-	return ObjectClass != nullptr && KeyframeHandler->IsPropertyKeyable(ObjectClass, *Handle);
+	return false;
 }
 
 bool SDetailSingleItemRow::IsKeyframeButtonEnabled() const
