Correct fix for key button crash, checking validity of pin rather than the node.

diff --git a/Engine/Source/Editor/PropertyEditor/Private/SDetailSingleItemRow.cpp b/Engine/Source/Editor/PropertyEditor/Private/SDetailSingleItemRow.cpp
--- a/Engine/Source/Editor/PropertyEditor/Private/SDetailSingleItemRow.cpp
+++ b/Engine/Source/Editor/PropertyEditor/Private/SDetailSingleItemRow.cpp
@@ -927,7 +927,7 @@ void SDetailSingleItemRow::CreateGlobalExtensionWidgets(TArray<FPropertyRowExten
 
 bool SDetailSingleItemRow::IsKeyframeButtonVisible() const
 {
-	if (OwnerTreeNode.IsValid())
+	if (OwnerTreeNode.Pin().IsValid())
 	{
 		if (IDetailsViewPrivate* DetailsView = OwnerTreeNode.Pin()->GetDetailsView())
 		{
