FDetailCategoryImpl::GetDetailsView() now checks that the layout is valid before dereferencing it.

diff --git a/Engine/Source/Editor/PropertyEditor/Private/DetailCategoryBuilderImpl.cpp b/Engine/Source/Editor/PropertyEditor/Private/DetailCategoryBuilderImpl.cpp
--- a/Engine/Source/Editor/PropertyEditor/Private/DetailCategoryBuilderImpl.cpp
+++ b/Engine/Source/Editor/PropertyEditor/Private/DetailCategoryBuilderImpl.cpp
@@ -532,8 +532,7 @@ bool FDetailCategoryImpl::IsAdvancedDropdownEnabled() const
 
 void FDetailCategoryImpl::RequestItemExpanded(TSharedRef<FDetailTreeNode> TreeNode, bool bShouldBeExpanded)
 {
-	TSharedPtr<FDetailLayoutBuilderImpl> DetailLayoutBuilderPtr = DetailLayoutBuilder.Pin();
-	if (DetailLayoutBuilderPtr.IsValid() && GetDetailsView())
+	if (GetDetailsView())
 	{		
 		GetDetailsView()->RequestItemExpanded(TreeNode, bShouldBeExpanded);
 	}
@@ -662,6 +661,16 @@ void FDetailCategoryImpl::SetDisplayName(FName InCategoryName, const FText& Loca
 	}
 }
 
+IDetailsViewPrivate* FDetailCategoryImpl::GetDetailsView() const
+{
+	TSharedPtr<FDetailLayoutBuilderImpl> DetailLayoutBuilderPtr = DetailLayoutBuilder.Pin();
+	if (DetailLayoutBuilderPtr.IsValid())
+	{
+		return DetailLayoutBuilderPtr->GetDetailsView();
+	}
+
+	return nullptr;
+}
 
 TSharedRef<ITableRow> FDetailCategoryImpl::GenerateWidgetForTableView(const TSharedRef<STableViewBase>& OwnerTable, bool bAllowFavoriteSystem)
 {
@@ -1035,6 +1044,6 @@ void FDetailCategoryImpl::GenerateLayout()
 
 bool FDetailCategoryImpl::IsParentEnabled() const
 {
-	TSharedPtr<FDetailLayoutBuilderImpl> DetailLayoutBuilderPtr = DetailLayoutBuilder.Pin();
-	return !DetailLayoutBuilderPtr.IsValid() || !DetailLayoutBuilderPtr->GetDetailsView() || DetailLayoutBuilderPtr->GetDetailsView()->IsPropertyEditingEnabled();
+	IDetailsViewPrivate* DetailsView = GetDetailsView();
+	return !DetailsView || DetailsView->IsPropertyEditingEnabled();
 }
diff --git a/Engine/Source/Editor/PropertyEditor/Private/DetailCategoryBuilderImpl.h b/Engine/Source/Editor/PropertyEditor/Private/DetailCategoryBuilderImpl.h
--- a/Engine/Source/Editor/PropertyEditor/Private/DetailCategoryBuilderImpl.h
+++ b/Engine/Source/Editor/PropertyEditor/Private/DetailCategoryBuilderImpl.h
@@ -180,8 +180,8 @@ public:
 	virtual void SetSortOrder(int32 InSortOrder) override;
 
 	/** FDetailTreeNode interface */
-	virtual IDetailsView* GetNodeDetailsView() const override { return DetailLayoutBuilder.Pin()->GetDetailsView(); }
-	virtual IDetailsViewPrivate* GetDetailsView() const override { return DetailLayoutBuilder.Pin()->GetDetailsView(); }
+	virtual IDetailsView* GetNodeDetailsView() const override { return GetDetailsView(); }
+	virtual IDetailsViewPrivate* GetDetailsView() const override;
 	virtual TSharedRef< ITableRow > GenerateWidgetForTableView(const TSharedRef<STableViewBase>& OwnerTable, bool bAllowFavoriteSystem) override;
 	virtual bool GenerateStandaloneWidget(FDetailWidgetRow& OutRow) const override;
 
