Added GetOuterBaseClass() to IPropertyHandle, which fetches the common base class of all the outer objects of an IPropertyHandle.

diff --git a/Engine/Source/Editor/PropertyEditor/Private/PropertyHandleImpl.cpp b/Engine/Source/Editor/PropertyEditor/Private/PropertyHandleImpl.cpp
--- a/Engine/Source/Editor/PropertyEditor/Private/PropertyHandleImpl.cpp
+++ b/Engine/Source/Editor/PropertyEditor/Private/PropertyHandleImpl.cpp
@@ -2452,7 +2452,20 @@ void FPropertyHandleBase::GetOuterObjects( TArray<UObject*>& OuterObjects ) cons
 			}
 		}
 	}
+}
 
+const UClass* FPropertyHandleBase::GetOuterBaseClass() const
+{
+	if (Implementation->GetPropertyNode().IsValid())
+	{
+		FObjectPropertyNode* ObjectNode = Implementation->GetPropertyNode()->FindObjectItemParent();
+		if (ObjectNode)
+		{
+			return ObjectNode->GetObjectBaseClass();
+		}
+	}
+
+	return nullptr;
 }
 
 void FPropertyHandleBase::ReplaceOuterObjects(const TArray<UObject*>& OuterObjects)
diff --git a/Engine/Source/Editor/PropertyEditor/Private/PropertyHandleImpl.h b/Engine/Source/Editor/PropertyEditor/Private/PropertyHandleImpl.h
--- a/Engine/Source/Editor/PropertyEditor/Private/PropertyHandleImpl.h
+++ b/Engine/Source/Editor/PropertyEditor/Private/PropertyHandleImpl.h
@@ -451,6 +451,7 @@ public:
 	virtual void AccessRawData( TArray<const void*>& RawData ) const override;
 	virtual uint32 GetNumOuterObjects() const override;
 	virtual void GetOuterObjects(TArray<UObject*>& OuterObjects) const override;
+	virtual const UClass* GetOuterBaseClass() const override;
 	virtual void ReplaceOuterObjects(const TArray<UObject*>& OuterObjects) override;
 	virtual void GetOuterPackages(TArray<UPackage*>& OuterPackages) const override;
 	virtual FPropertyAccess::Result GetNumChildren( uint32& OutNumChildren ) const override;
diff --git a/Engine/Source/Editor/PropertyEditor/Public/PropertyHandle.h b/Engine/Source/Editor/PropertyEditor/Public/PropertyHandle.h
--- a/Engine/Source/Editor/PropertyEditor/Public/PropertyHandle.h
+++ b/Engine/Source/Editor/PropertyEditor/Public/PropertyHandle.h
@@ -415,6 +415,13 @@ public:
 	 */
 	virtual void GetOuterObjects( TArray<UObject*>& OuterObjects ) const = 0;
 
+	/**
+	 * Get the shared base class of the objects that contain this property.
+	 *
+	 * @return The shared base class of the outer objects, or null if none are selected.
+	 */
+	virtual const UClass* GetOuterBaseClass() const = 0;
+
 	/**
 	 * Set the outer objects for this property
 	 *
