Keyframe button handler has been moved out of SDetailSingleItemRow and into LevelEditorSequencerIntegration.

IDetailKeyframeHandler takes a const UClass* to match upstream interfaces.

diff --git a/Engine/Plugins/Experimental/ControlRig/Source/ControlRigEditor/Private/EditMode/SControlRigEditModeTools.cpp b/Engine/Plugins/Experimental/ControlRig/Source/ControlRigEditor/Private/EditMode/SControlRigEditModeTools.cpp
--- a/Engine/Plugins/Experimental/ControlRig/Source/ControlRigEditor/Private/EditMode/SControlRigEditModeTools.cpp
+++ b/Engine/Plugins/Experimental/ControlRig/Source/ControlRigEditor/Private/EditMode/SControlRigEditModeTools.cpp
@@ -270,7 +270,7 @@ void SControlRigEditModeTools::SetSequencer(TWeakPtr<ISequencer> InSequencer)
 	WeakSequencer = InSequencer.Pin();
 }
 
-bool SControlRigEditModeTools::IsPropertyKeyable(UClass* InObjectClass, const IPropertyHandle& InPropertyHandle) const
+bool SControlRigEditModeTools::IsPropertyKeyable(const UClass* InObjectClass, const IPropertyHandle& InPropertyHandle) const
 {
 	if (InObjectClass && InObjectClass->IsChildOf(UControlRigTransformNoScaleControlProxy::StaticClass()) && InObjectClass->IsChildOf(UControlRigEulerTransformControlProxy::StaticClass()) && InPropertyHandle.GetProperty()
 		&& InPropertyHandle.GetProperty()->GetFName() == GET_MEMBER_NAME_CHECKED(UControlRigTransformControlProxy, Transform)) 
diff --git a/Engine/Plugins/Experimental/ControlRig/Source/ControlRigEditor/Private/EditMode/SControlRigEditModeTools.h b/Engine/Plugins/Experimental/ControlRig/Source/ControlRigEditor/Private/EditMode/SControlRigEditModeTools.h
--- a/Engine/Plugins/Experimental/ControlRig/Source/ControlRigEditor/Private/EditMode/SControlRigEditModeTools.h
+++ b/Engine/Plugins/Experimental/ControlRig/Source/ControlRigEditor/Private/EditMode/SControlRigEditModeTools.h
@@ -35,7 +35,7 @@ public:
 	void SetControlRig(UControlRig* ControlRig);
 
 	// IDetailKeyframeHandler interface
-	virtual bool IsPropertyKeyable(UClass* InObjectClass, const class IPropertyHandle& PropertyHandle) const override;
+	virtual bool IsPropertyKeyable(const UClass* InObjectClass, const class IPropertyHandle& PropertyHandle) const override;
 	virtual bool IsPropertyKeyingEnabled() const override;
 	virtual void OnKeyPropertyClicked(const IPropertyHandle& KeyedPropertyHandle) override;
 	virtual bool IsPropertyAnimated(const class IPropertyHandle& PropertyHandle, UObject *ParentObject) const override;
diff --git a/Engine/Source/Editor/PropertyEditor/Private/SDetailSingleItemRow.cpp b/Engine/Source/Editor/PropertyEditor/Private/SDetailSingleItemRow.cpp
--- a/Engine/Source/Editor/PropertyEditor/Private/SDetailSingleItemRow.cpp
+++ b/Engine/Source/Editor/PropertyEditor/Private/SDetailSingleItemRow.cpp
@@ -181,7 +181,6 @@ TSharedPtr<FPropertyNode> SDetailSingleItemRow::GetPropertyNode() const
 	if (!PropertyNode.IsValid() && Customization->HasCustomBuilder())
 	{
 		TSharedPtr<IPropertyHandle> PropertyHandle = Customization->CustomBuilderRow->GetPropertyHandle();
-
 		if (PropertyHandle.IsValid())
 		{
 			PropertyNode = StaticCastSharedPtr<FPropertyHandleBase>(PropertyHandle)->GetPropertyNode();
@@ -197,8 +196,15 @@ TSharedPtr<IPropertyHandle> SDetailSingleItemRow::GetPropertyHandle() const
 	TSharedPtr<FPropertyNode> PropertyNode = GetPropertyNode();
 	if (PropertyNode.IsValid())
 	{
-		IDetailsViewPrivate* DetailsView = OwnerTreeNode.Pin()->GetDetailsView();
-		Handle = PropertyEditorHelpers::GetPropertyHandle(PropertyNode.ToSharedRef(), DetailsView->GetNotifyHook(), DetailsView->GetPropertyUtilities());
+		TSharedPtr<FDetailTreeNode> OwnerTreeNodePtr = OwnerTreeNode.Pin();
+		if (OwnerTreeNodePtr.IsValid())
+		{
+			IDetailsViewPrivate* DetailsView = OwnerTreeNodePtr->GetDetailsView();
+			if (DetailsView != nullptr)
+			{
+				Handle = PropertyEditorHelpers::GetPropertyHandle(PropertyNode.ToSharedRef(), DetailsView->GetNotifyHook(), DetailsView->GetPropertyUtilities());
+			}
+		}
 	}
 	else if (WidgetRow.PropertyHandles.Num() > 0)
 	{
@@ -490,16 +496,6 @@ void SDetailSingleItemRow::Construct( const FArguments& InArgs, FDetailLayoutCus
 					FText::GetEmpty();
 			});
 
-			FPropertyRowExtensionButton& CreateKey = ExtensionButtons.AddDefaulted_GetRef();
-			CreateKey.Icon = FSlateIcon(FAppStyle::Get().GetStyleSetName(), "Sequencer.AddKey.Details");
-			CreateKey.Label = NSLOCTEXT("PropertyEditor", "CreateKey", "Create Key");
-			CreateKey.ToolTip = NSLOCTEXT("PropertyEditor", "CreateKeyToolTip", "Add a keyframe for this property.");
-			CreateKey.UIAction = FUIAction(
-				FExecuteAction::CreateSP(this, &SDetailSingleItemRow::OnAddKeyframeClicked),
-				FCanExecuteAction::CreateSP(this, &SDetailSingleItemRow::IsKeyframeButtonEnabled),
-				FGetActionCheckState(),
-				FIsActionButtonVisible::CreateSP(this, &SDetailSingleItemRow::IsKeyframeButtonVisible));
-
 			CreateGlobalExtensionWidgets(ExtensionButtons);
 
 			FSlimHorizontalToolBarBuilder ToolbarBuilder(TSharedPtr<FUICommandList>(), FMultiBoxCustomization::None);
@@ -818,9 +814,7 @@ TSharedRef<SWidget> SDetailSingleItemRow::CreateExtensionWidget() const
 		if (Customization->HasPropertyNode() && ExtensionHandler.IsValid())
 		{
 			TSharedPtr<IPropertyHandle> Handle = PropertyEditorHelpers::GetPropertyHandle(Customization->GetPropertyNode().ToSharedRef(), nullptr, nullptr);
-
-			FObjectPropertyNode* ObjectItemParent = Customization->GetPropertyNode()->FindObjectItemParent();
-			UClass* ObjectClass = ObjectItemParent ? ObjectItemParent->GetObjectBaseClass() : nullptr;
+			const UClass* ObjectClass = Handle->GetOuterBaseClass();
 			if (Handle->IsValidHandle() && ObjectClass && ExtensionHandler->IsPropertyExtendable(ObjectClass, *Handle))
 			{
 				FDetailLayoutBuilderImpl& DetailLayout = OwnerTreeNodePinned->GetParentCategory()->GetParentLayoutImpl();
@@ -851,7 +845,7 @@ void SDetailSingleItemRow::OnFavoriteMenuToggle()
 	//////////////////////////////////////////////////////////////////////////
 	// Calculate properly the scrolling offset (by item) to make sure the mouse stay over the same property
 
-	//Get the node item number in case it is expand we have to recursively count all childrens
+	// Get the node item number in case it is expand we have to recursively count all childrens
 	int32 ExpandSize = 0;
 	if (OwnerTreeNodePinned->ShouldBeExpanded())
 	{
@@ -859,11 +853,11 @@ void SDetailSingleItemRow::OnFavoriteMenuToggle()
 	}
 	else
 	{
-		//if the item is not expand count is 1
+		// if the item is not expand count is 1
 		ExpandSize = 1;
 	}
 			
-	//Get the number of favorite child (simple and advance) to know if the favorite category will be create or remove
+	// Get the number of favorite child (simple and advance) to know if the favorite category will be create or remove
 	int32 SimplePropertiesNum = 0;
 	int32 NumAdvancedProperties = 0;
 
@@ -889,12 +883,12 @@ void SDetailSingleItemRow::OnFavoriteMenuToggle()
 
 		if (IsAdvancedProperty && NumAdvancedProperties == 1)
 		{
-			//Removing the advance button count as 1 item
+			// Removing the advance button count as 1 item
 			ScrollingOffsetRemove -= 1;
 		}
 		if (NumAdvancedProperties + SimplePropertiesNum == 1)
 		{
-			//Removing a full category count as 2 items
+			// Removing a full category count as 2 items
 			ScrollingOffsetRemove -= 2;
 		}
 	}
@@ -921,61 +915,21 @@ void SDetailSingleItemRow::CreateGlobalExtensionWidgets(TArray<FPropertyRowExten
 	// fetch global extension widgets 
 	FPropertyEditorModule& PropertyEditorModule = FModuleManager::Get().GetModuleChecked<FPropertyEditorModule>("PropertyEditor");
 
-	FOnGenerateGlobalRowExtensionArgs Args { GetPropertyHandle(), GetPropertyNode(), OwnerTreeNode };
-	PropertyEditorModule.GetGlobalRowExtensionDelegate().Broadcast(Args, OutExtensions);
-}
+	FOnGenerateGlobalRowExtensionArgs Args;
+	Args.OwnerTreeNode = OwnerTreeNode;
 
-bool SDetailSingleItemRow::IsKeyframeButtonVisible() const
-{
-	if (OwnerTreeNode.Pin().IsValid())
+	if (Customization->HasPropertyNode())
 	{
-		if (IDetailsViewPrivate* DetailsView = OwnerTreeNode.Pin()->GetDetailsView())
-		{
-			TSharedPtr<IDetailKeyframeHandler> KeyframeHandler = DetailsView->GetKeyframeHandler();
-			if (!Customization->HasPropertyNode() || !KeyframeHandler.IsValid())
-			{
-				return false;
-			}
-
-			TSharedPtr<IPropertyHandle> Handle = PropertyEditorHelpers::GetPropertyHandle(Customization->GetPropertyNode().ToSharedRef(), nullptr, nullptr);
-
-			FObjectPropertyNode* ObjectItemParent = Customization->GetPropertyNode()->FindObjectItemParent();
-			UClass* ObjectClass = ObjectItemParent != nullptr ? ObjectItemParent->GetObjectBaseClass() : nullptr;
-			return ObjectClass != nullptr && KeyframeHandler->IsPropertyKeyable(ObjectClass, *Handle);
-		}
+		Args.PropertyHandle = PropertyEditorHelpers::GetPropertyHandle(Customization->GetPropertyNode().ToSharedRef(), nullptr, nullptr);
 	}
 
-	return false;
-}
-
-bool SDetailSingleItemRow::IsKeyframeButtonEnabled() const
-{
-	TSharedPtr<FDetailTreeNode> OwnerTreeNodePinned = OwnerTreeNode.Pin();
-	if (OwnerTreeNodePinned.IsValid() == false)
-	{
-		return false;
-	}
-
-	TSharedPtr<IDetailKeyframeHandler> KeyframeHandler = OwnerTreeNodePinned->GetDetailsView()->GetKeyframeHandler();
-	return OwnerTreeNodePinned->IsPropertyEditingEnabled().Get() &&
-		KeyframeHandler.IsValid() &&
-		KeyframeHandler->IsPropertyKeyingEnabled();
-}
-
-void SDetailSingleItemRow::OnAddKeyframeClicked()
-{	
-	TSharedPtr<IDetailKeyframeHandler> KeyframeHandler = OwnerTreeNode.Pin()->GetDetailsView()->GetKeyframeHandler();
-	if (KeyframeHandler.IsValid())
-	{
-		TSharedPtr<IPropertyHandle> Handle = PropertyEditorHelpers::GetPropertyHandle(Customization->GetPropertyNode().ToSharedRef(), nullptr, nullptr);
-
-		KeyframeHandler->OnKeyPropertyClicked(*Handle);
-	}
+	PropertyEditorModule.GetGlobalRowExtensionDelegate().Broadcast(Args, OutExtensions);
 }
 
 bool SDetailSingleItemRow::IsHighlighted() const
 {
-	return OwnerTreeNode.Pin()->IsHighlighted();
+	TSharedPtr<FDetailTreeNode> OwnerTreeNodePtr = OwnerTreeNode.Pin();
+	return OwnerTreeNodePtr.IsValid() ? OwnerTreeNodePtr->IsHighlighted() : false;
 }
 
 void SDetailSingleItemRow::SetIsDragDrop(bool bInIsDragDrop)
diff --git a/Engine/Source/Editor/PropertyEditor/Private/SDetailSingleItemRow.h b/Engine/Source/Editor/PropertyEditor/Private/SDetailSingleItemRow.h
--- a/Engine/Source/Editor/PropertyEditor/Private/SDetailSingleItemRow.h
+++ b/Engine/Source/Editor/PropertyEditor/Private/SDetailSingleItemRow.h
@@ -75,13 +75,9 @@ private:
 
 	void CreateGlobalExtensionWidgets(TArray<FPropertyRowExtensionButton>& ExtensionButtons) const;
 	TSharedRef<SWidget> CreateExtensionWidget() const;
-	TSharedRef<SWidget> CreateKeyframeButton();
 	void OnResetToDefaultClicked() const;
 	bool IsResetToDefaultEnabled() const;
 
-	bool IsKeyframeButtonVisible() const;
-	bool IsKeyframeButtonEnabled() const;
-	void OnAddKeyframeClicked();
 	bool IsHighlighted() const;
 
 	void OnFavoriteMenuToggle();
diff --git a/Engine/Source/Editor/PropertyEditor/Public/IDetailKeyframeHandler.h b/Engine/Source/Editor/PropertyEditor/Public/IDetailKeyframeHandler.h
--- a/Engine/Source/Editor/PropertyEditor/Public/IDetailKeyframeHandler.h
+++ b/Engine/Source/Editor/PropertyEditor/Public/IDetailKeyframeHandler.h
@@ -11,12 +11,12 @@ class IDetailKeyframeHandler
 public:
 	virtual ~IDetailKeyframeHandler(){}
 
-	virtual bool IsPropertyKeyable(UClass* InObjectClass, const class IPropertyHandle& PropertyHandle) const = 0;
+	virtual bool IsPropertyKeyable(const UClass* InObjectClass, const IPropertyHandle& PropertyHandle) const = 0;
 
 	virtual bool IsPropertyKeyingEnabled() const = 0;
 
 	virtual void OnKeyPropertyClicked(const IPropertyHandle& KeyedPropertyHandle) = 0;
 
-	virtual bool IsPropertyAnimated(const class IPropertyHandle& PropertyHandle, UObject *ParentObject) const = 0;
+	virtual bool IsPropertyAnimated(const IPropertyHandle& PropertyHandle, UObject* ParentObject) const = 0;
 
 };
diff --git a/Engine/Source/Editor/PropertyEditor/Public/PropertyEditorDelegates.h b/Engine/Source/Editor/PropertyEditor/Public/PropertyEditorDelegates.h
--- a/Engine/Source/Editor/PropertyEditor/Public/PropertyEditorDelegates.h
+++ b/Engine/Source/Editor/PropertyEditor/Public/PropertyEditorDelegates.h
@@ -9,7 +9,12 @@
 
 class FPropertyPath;
 struct FPropertyChangedEvent;
+class IDetailCustomization;
+class IDetailTreeNode;
 class IPropertyHandle;
+class IPropertyTypeIdentifier;
+class IPropertyTypeCustomization;
+class SWidget;
 
 struct FPropertyAndParent
 {
@@ -41,10 +46,10 @@ DECLARE_DELEGATE_RetVal(bool, FIsCustomRowVisibilityFiltered);
 DECLARE_DELEGATE_RetVal_TwoParams(bool, FIsCustomRowVisible, FName /*InRowName*/, FName /*InParentName*/);
 
 /** Delegate called to get a detail layout for a specific object class */
-DECLARE_DELEGATE_RetVal( TSharedRef<class IDetailCustomization>, FOnGetDetailCustomizationInstance );
+DECLARE_DELEGATE_RetVal( TSharedRef<IDetailCustomization>, FOnGetDetailCustomizationInstance );
 
 /** Delegate called to get a property layout for a specific property type */
-DECLARE_DELEGATE_RetVal( TSharedRef<class IPropertyTypeCustomization>, FOnGetPropertyTypeCustomizationInstance );
+DECLARE_DELEGATE_RetVal( TSharedRef<IPropertyTypeCustomization>, FOnGetPropertyTypeCustomizationInstance );
 
 /** Notification for when a property view changes */
 DECLARE_DELEGATE_TwoParams( FOnObjectArrayChanged, const FString&, const TArray<UObject*>& );
@@ -59,12 +64,12 @@ DECLARE_DELEGATE_OneParam( FOnPropertySelectionChanged, FProperty* )
 DECLARE_DELEGATE_OneParam( FOnPropertyDoubleClicked, FProperty* )
 
 /** Notification for when a property is clicked by the user*/
-DECLARE_DELEGATE_OneParam( FOnPropertyClicked, const TSharedPtr< class FPropertyPath >& )
+DECLARE_DELEGATE_OneParam( FOnPropertyClicked, const TSharedPtr<FPropertyPath >& )
 
 /** */
 DECLARE_DELEGATE_OneParam( FConstructExternalColumnHeaders, const TSharedRef< class SHeaderRow >& )
 
-DECLARE_DELEGATE_RetVal_TwoParams( TSharedRef< class SWidget >, FConstructExternalColumnCell,  const FName& /*ColumnName*/, const TSharedRef< class IPropertyTreeRow >& /*RowWidget*/)
+DECLARE_DELEGATE_RetVal_TwoParams( TSharedRef< SWidget >, FConstructExternalColumnCell,  const FName& /*ColumnName*/, const TSharedRef< class IPropertyTreeRow >& /*RowWidget*/)
 
 /** Delegate called to see if a property editing is enabled */
 DECLARE_DELEGATE_RetVal(bool, FIsPropertyEditingEnabled );
@@ -80,10 +85,8 @@ struct FOnGenerateGlobalRowExtensionArgs
 {
 	/** The detail row's property handle. */
 	TSharedPtr<IPropertyHandle> PropertyHandle;
-	/** The detail row's property node. */
-	TSharedPtr<class FPropertyNode> PropertyNode;
 	/** The detail row's owner tree node. */
-	TWeakPtr<class IDetailTreeNode> OwnerTreeNode;
+	TWeakPtr<IDetailTreeNode> OwnerTreeNode;
 };
 
 /** 
@@ -125,11 +128,11 @@ struct FPropertyTypeLayoutCallback
 {
 	FOnGetPropertyTypeCustomizationInstance PropertyTypeLayoutDelegate;
 
-	TSharedPtr<class IPropertyTypeIdentifier> PropertyTypeIdentifier;
+	TSharedPtr<IPropertyTypeIdentifier> PropertyTypeIdentifier;
 
 	bool IsValid() const { return PropertyTypeLayoutDelegate.IsBound(); }
 
-	TSharedRef<class IPropertyTypeCustomization> GetCustomizationInstance() const;
+	TSharedRef<IPropertyTypeCustomization> GetCustomizationInstance() const;
 };
 
 
@@ -145,7 +148,7 @@ struct FPropertyTypeLayoutCallbackList
 
 	void Remove(const TSharedPtr<IPropertyTypeIdentifier>& InIdentifier);
 
-	const FPropertyTypeLayoutCallback& Find(const class IPropertyHandle& PropertyHandle) const;
+	const FPropertyTypeLayoutCallback& Find(const IPropertyHandle& PropertyHandle) const;
 };
 
 /** This is a multimap as there many be more than one customization per property type */
diff --git a/Engine/Source/Editor/Sequencer/Private/KeyPropertyParams.cpp b/Engine/Source/Editor/Sequencer/Private/KeyPropertyParams.cpp
--- a/Engine/Source/Editor/Sequencer/Private/KeyPropertyParams.cpp
+++ b/Engine/Source/Editor/Sequencer/Private/KeyPropertyParams.cpp
@@ -38,13 +38,13 @@ FPropertyPath PropertyHandleToPropertyPath(const UClass* OwnerClass, const IProp
 	return PropertyPath;
 }
 
-FCanKeyPropertyParams::FCanKeyPropertyParams(UClass* InObjectClass, const FPropertyPath& InPropertyPath)
+FCanKeyPropertyParams::FCanKeyPropertyParams(const UClass* InObjectClass, const FPropertyPath& InPropertyPath)
 	: ObjectClass(InObjectClass)
 	, PropertyPath(InPropertyPath)
 {
 }
 
-FCanKeyPropertyParams::FCanKeyPropertyParams(UClass* InObjectClass, const IPropertyHandle& InPropertyHandle)
+FCanKeyPropertyParams::FCanKeyPropertyParams(const UClass* InObjectClass, const IPropertyHandle& InPropertyHandle)
 	: ObjectClass(InObjectClass)
 	, PropertyPath(PropertyHandleToPropertyPath(InObjectClass, InPropertyHandle))
 {
diff --git a/Engine/Source/Editor/Sequencer/Private/LevelEditorSequencerIntegration.cpp b/Engine/Source/Editor/Sequencer/Private/LevelEditorSequencerIntegration.cpp
--- a/Engine/Source/Editor/Sequencer/Private/LevelEditorSequencerIntegration.cpp
+++ b/Engine/Source/Editor/Sequencer/Private/LevelEditorSequencerIntegration.cpp
@@ -5,6 +5,7 @@
 #include "Styling/SlateIconFinder.h"
 #include "PropertyHandle.h"
 #include "IDetailKeyframeHandler.h"
+#include "IDetailTreeNode.h"
 #include "GameDelegates.h"
 #include "Settings/LevelEditorPlaySettings.h"
 #include "Editor/PropertyEditor/Public/PropertyEditorModule.h"
@@ -63,7 +64,7 @@ public:
 		Sequencers.Remove(InSequencer);
 	}
 
-	virtual bool IsPropertyKeyable(UClass* InObjectClass, const IPropertyHandle& InPropertyHandle) const
+	virtual bool IsPropertyKeyable(const UClass* InObjectClass, const IPropertyHandle& InPropertyHandle) const
 	{
 		FCanKeyPropertyParams CanKeyPropertyParams(InObjectClass, InPropertyHandle);
 
@@ -222,21 +223,6 @@ void FLevelEditorSequencerIntegration::Initialize(const FLevelEditorSequencerInt
 	ActivateSequencerEditorMode();
 	BindLevelEditorCommands();
 
-	{
-		FPropertyEditorModule& EditModule = FModuleManager::Get().GetModuleChecked<FPropertyEditorModule>("PropertyEditor");
-
-		FDelegateHandle Handle = EditModule.OnPropertyEditorOpened().AddRaw(this, &FLevelEditorSequencerIntegration::OnPropertyEditorOpened);
-		AcquiredResources.Add(
-			[=]{
-				FPropertyEditorModule* EditModulePtr = FModuleManager::Get().GetModulePtr<FPropertyEditorModule>("PropertyEditor");
-				if (EditModulePtr)
-				{
-					EditModulePtr->OnPropertyEditorOpened().Remove(Handle);
-				}
-			}
-		);
-	}
-
 	{
 		FLevelEditorModule& LevelEditorModule = FModuleManager::Get().GetModuleChecked<FLevelEditorModule>("LevelEditor");
 
@@ -795,9 +781,91 @@ void FLevelEditorSequencerIntegration::MakeBrowseToSelectedActorSubMenu(FMenuBui
 	}
 }
 
+static TSharedPtr<IDetailKeyframeHandler> GetKeyframeHandler(TWeakPtr<IDetailTreeNode> OwnerTreeNode)
+{
+	TSharedPtr<IDetailTreeNode> OwnerTreeNodePtr = OwnerTreeNode.Pin();
+	if (!OwnerTreeNodePtr.IsValid())
+	{
+		return TSharedPtr<IDetailKeyframeHandler>();
+	}
+
+	IDetailsView* DetailsView = OwnerTreeNodePtr->GetNodeDetailsView();
+	if (DetailsView == nullptr)
+	{
+		return TSharedPtr<IDetailKeyframeHandler>();
+	}
+
+	return DetailsView->GetKeyframeHandler();
+}
+
+static bool IsKeyframeButtonVisible(TWeakPtr<IDetailTreeNode> OwnerTreeNode, TSharedPtr<IPropertyHandle> PropertyHandle)
+{
+	TSharedPtr<IDetailKeyframeHandler> KeyframeHandler = GetKeyframeHandler(OwnerTreeNode);
+	if (!KeyframeHandler.IsValid() || !PropertyHandle.IsValid())
+	{
+		return false;
+	}
+
+	const UClass* ObjectClass = PropertyHandle->GetOuterBaseClass();
+	if (ObjectClass == nullptr)
+	{
+		return false;
+	}
+
+	return KeyframeHandler->IsPropertyKeyable(ObjectClass, *PropertyHandle);
+}
+
+static bool IsKeyframeButtonEnabled(TWeakPtr<IDetailTreeNode> OwnerTreeNode)
+{
+	TSharedPtr<IDetailKeyframeHandler> KeyframeHandler = GetKeyframeHandler(OwnerTreeNode);
+	if (!KeyframeHandler.IsValid())
+	{
+		return false;
+	}
+
+	return KeyframeHandler->IsPropertyKeyingEnabled();
+}
+
+static void OnAddKeyframeClicked(TWeakPtr<IDetailTreeNode> OwnerTreeNode, TSharedPtr<IPropertyHandle> PropertyHandle)
+{
+	TSharedPtr<IDetailKeyframeHandler> KeyframeHandler = GetKeyframeHandler(OwnerTreeNode);
+	if (!KeyframeHandler.IsValid() || !PropertyHandle.IsValid())
+	{
+		return;
+	}
+
+	KeyframeHandler->OnKeyPropertyClicked(*PropertyHandle);
+}
+
+
+static void RegisterDetailsExtensionHandler(const FOnGenerateGlobalRowExtensionArgs& Args, TArray<FPropertyRowExtensionButton>& OutExtensionButtons)
+{
+	// local copy for capturing in handlers below
+	TSharedPtr<IPropertyHandle> PropertyHandle = Args.PropertyHandle;
+	if (!PropertyHandle.IsValid())
+	{
+		return;
+	}
+
+	static FSlateIcon CreateKeyIcon(FAppStyle::Get().GetStyleSetName(), "Sequencer.AddKey.Details");
+
+	TWeakPtr<IDetailTreeNode> OwnerTreeNode = Args.OwnerTreeNode;
+
+	FPropertyRowExtensionButton& CreateKey = OutExtensionButtons.AddDefaulted_GetRef();
+	CreateKey.Icon = CreateKeyIcon;
+	CreateKey.Label = NSLOCTEXT("PropertyEditor", "CreateKey", "Create Key");
+	CreateKey.ToolTip = NSLOCTEXT("PropertyEditor", "CreateKeyToolTip", "Add a keyframe for this property.");
+	CreateKey.UIAction = FUIAction(
+		FExecuteAction::CreateStatic(&OnAddKeyframeClicked, OwnerTreeNode, PropertyHandle),
+		FCanExecuteAction::CreateStatic(&IsKeyframeButtonEnabled, OwnerTreeNode),
+		FGetActionCheckState(),
+		FIsActionButtonVisible::CreateStatic(&IsKeyframeButtonVisible, OwnerTreeNode, PropertyHandle)
+	);
+}
+
 void FLevelEditorSequencerIntegration::ActivateDetailHandler(const FLevelEditorSequencerIntegrationOptions& Options)
 {
-	FName DetailHandlerName("DetailHandler");
+	static FName DetailHandlerName("DetailHandler");
 
 	AcquiredResources.Release(DetailHandlerName);
 
@@ -817,7 +885,7 @@ void FLevelEditorSequencerIntegration::ActivateDetailHandler(const FLevelEditorS
 	FDelegateHandle OnPropertyEditorOpenedHandle = EditModule.OnPropertyEditorOpened().AddRaw(this, &FLevelEditorSequencerIntegration::OnPropertyEditorOpened);
 
 	auto DeactivateDetailKeyframeHandler =
-		[this, OnPropertyEditorOpenedHandle]
+		[this, OnPropertyEditorOpenedHandle]()
 		{
 			FPropertyEditorModule* EditModulePtr = FModuleManager::Get().GetModulePtr<FPropertyEditorModule>("PropertyEditor");
 			if (!EditModulePtr)
@@ -842,9 +910,11 @@ void FLevelEditorSequencerIntegration::ActivateDetailHandler(const FLevelEditorS
 			}
 		};
 
+	EditModule.GetGlobalRowExtensionDelegate().AddStatic(&RegisterDetailsExtensionHandler);
+
 	FName DetailHandlerRefreshName("DetailHandlerRefresh");
 	auto RefreshDetailHandler =
-		[]
+		[]()
 		{
 			FPropertyEditorModule* EditModulePtr = FModuleManager::Get().GetModulePtr<FPropertyEditorModule>("PropertyEditor");
 			if (!EditModulePtr)
diff --git a/Engine/Source/Editor/Sequencer/Public/KeyPropertyParams.h b/Engine/Source/Editor/Sequencer/Public/KeyPropertyParams.h
--- a/Engine/Source/Editor/Sequencer/Public/KeyPropertyParams.h
+++ b/Engine/Source/Editor/Sequencer/Public/KeyPropertyParams.h
@@ -19,14 +19,14 @@ struct SEQUENCER_API FCanKeyPropertyParams
 	 * @param InObjectClass the class of the object which has the property to be keyed.
 	 * @param InPropertyPath path get from the root object to the property to be keyed.
 	 */
-	FCanKeyPropertyParams(UClass* InObjectClass, const FPropertyPath& InPropertyPath);
+	FCanKeyPropertyParams(const UClass* InObjectClass, const FPropertyPath& InPropertyPath);
 
 	/**
 	* Creates new can key property parameters.
 	* @param InObjectClass the class of the object which has the property to be keyed.
 	* @param InPropertyHandle a handle to the property to be keyed.
 	*/
-	FCanKeyPropertyParams(UClass* InObjectClass, const IPropertyHandle& InPropertyHandle);
+	FCanKeyPropertyParams(const UClass* InObjectClass, const IPropertyHandle& InPropertyHandle);
 
 	const UStruct* FindPropertyContainer(const FProperty* ForProperty) const;
 
diff --git a/Engine/Source/Editor/UMGEditor/Private/Animation/UMGDetailKeyframeHandler.cpp b/Engine/Source/Editor/UMGEditor/Private/Animation/UMGDetailKeyframeHandler.cpp
--- a/Engine/Source/Editor/UMGEditor/Private/Animation/UMGDetailKeyframeHandler.cpp
+++ b/Engine/Source/Editor/UMGEditor/Private/Animation/UMGDetailKeyframeHandler.cpp
@@ -10,7 +10,7 @@ FUMGDetailKeyframeHandler::FUMGDetailKeyframeHandler(TSharedPtr<FWidgetBlueprint
 	: BlueprintEditor( InBlueprintEditor )
 {}
 
-bool FUMGDetailKeyframeHandler::IsPropertyKeyable(UClass* InObjectClass, const IPropertyHandle& InPropertyHandle) const
+bool FUMGDetailKeyframeHandler::IsPropertyKeyable(const UClass* InObjectClass, const IPropertyHandle& InPropertyHandle) const
 {
 	return BlueprintEditor.Pin()->GetSequencer()->CanKeyProperty(FCanKeyPropertyParams(InObjectClass, InPropertyHandle));
 }
diff --git a/Engine/Source/Editor/UMGEditor/Private/Animation/UMGDetailKeyframeHandler.h b/Engine/Source/Editor/UMGEditor/Private/Animation/UMGDetailKeyframeHandler.h
--- a/Engine/Source/Editor/UMGEditor/Private/Animation/UMGDetailKeyframeHandler.h
+++ b/Engine/Source/Editor/UMGEditor/Private/Animation/UMGDetailKeyframeHandler.h
@@ -13,7 +13,7 @@ class FUMGDetailKeyframeHandler : public IDetailKeyframeHandler
 public:
 	FUMGDetailKeyframeHandler( TSharedPtr<class FWidgetBlueprintEditor> InBlueprintEditor );
 
-	virtual bool IsPropertyKeyable(UClass* InObjectClass, const class IPropertyHandle& PropertyHandle) const override;
+	virtual bool IsPropertyKeyable(const UClass* InObjectClass, const class IPropertyHandle& PropertyHandle) const override;
 
 	virtual bool IsPropertyKeyingEnabled() const override;
 
