Right column min width extended to always fit 2 buttons (reset to default + 1).

diff --git a/Engine/Source/Editor/PropertyEditor/Private/SDetailSingleItemRow.cpp b/Engine/Source/Editor/PropertyEditor/Private/SDetailSingleItemRow.cpp
--- a/Engine/Source/Editor/PropertyEditor/Private/SDetailSingleItemRow.cpp
+++ b/Engine/Source/Editor/PropertyEditor/Private/SDetailSingleItemRow.cpp
@@ -519,7 +519,7 @@ void SDetailSingleItemRow::Construct( const FArguments& InArgs, FDetailLayoutCus
 			OuterSplitter->AddSlot()
 				.Value(ColumnSizeData.RightColumnWidth)
 				.OnSlotResized(ColumnSizeData.OnRightColumnResized)
-				.MinSize(40)
+				.MinSize(50)
 			[
 				RightColumnBox
 			];
