Property names in the details panel are now clipped consistently, rather than allowing a bit of slack that causes the text to zig-zag when the column is resized to be small.

diff --git a/Engine/Source/Editor/PropertyEditor/Private/DetailPropertyRow.cpp b/Engine/Source/Editor/PropertyEditor/Private/DetailPropertyRow.cpp
--- a/Engine/Source/Editor/PropertyEditor/Private/DetailPropertyRow.cpp
+++ b/Engine/Source/Editor/PropertyEditor/Private/DetailPropertyRow.cpp
@@ -675,7 +675,8 @@ void FDetailPropertyRow::MakeNameOrKeyWidget( FDetailWidgetRow& Row, const TShar
 		IsEnabledAttrib.Bind( this, &FDetailPropertyRow::GetEnabledState );
 	}
 
-	TSharedRef<SHorizontalBox> NameHorizontalBox = SNew(SHorizontalBox);
+	TSharedRef<SHorizontalBox> NameHorizontalBox = SNew(SHorizontalBox)
+		.Clipping(EWidgetClipping::OnDemand);
 	
 	TSharedPtr<SWidget> NameWidget;
 
