Removed extra padding in the name of detail groups that was causing a tiny amount of zig-zagginess.

diff --git a/Engine/Source/Editor/PropertyEditor/Private/DetailGroup.cpp b/Engine/Source/Editor/PropertyEditor/Private/DetailGroup.cpp
--- a/Engine/Source/Editor/PropertyEditor/Private/DetailGroup.cpp
+++ b/Engine/Source/Editor/PropertyEditor/Private/DetailGroup.cpp
@@ -215,7 +215,7 @@ TSharedRef<SWidget> FDetailGroup::MakeNameWidget()
 	return
 		SNew( SButton )
 		.ButtonStyle( FEditorStyle::Get(), "NoBorder" )
-		.ContentPadding(FMargin(0,2,2,2))
+		.ContentPadding(FMargin(0,2,0,2))
 		.OnClicked( this, &FDetailGroup::OnNameClicked )				
 		.ForegroundColor( FSlateColor::UseForeground() )
 		.Content()
