Details view scrollbar well disappears when the scrollbar isn't visible.

Restored Advanced properties dropdown drop shadow.

diff --git a/Engine/Source/Editor/PropertyEditor/Private/DetailAdvancedDropdownNode.cpp b/Engine/Source/Editor/PropertyEditor/Private/DetailAdvancedDropdownNode.cpp
--- a/Engine/Source/Editor/PropertyEditor/Private/DetailAdvancedDropdownNode.cpp
+++ b/Engine/Source/Editor/PropertyEditor/Private/DetailAdvancedDropdownNode.cpp
@@ -34,7 +34,8 @@ public:
 		TSharedPtr<SWidget> ContentWidget;
 		if( bIsTopNode )
 		{
-			ContentWidget = SNew(SSpacer);
+			ContentWidget = SNew(SImage)
+				.Image(FEditorStyle::GetBrush("DetailsView.AdvancedDropdownBorder.Open"));
 		}
 		else if( InArgs._ShouldShowAdvancedButton )
 		{
@@ -70,16 +71,29 @@ public:
 		{
 			ContentWidget = SNew(SSpacer);
 		}
-		
+
+		TWeakPtr<STableViewBase> OwnerTableViewWeak = InOwnerTableView;
+		auto GetGridLinePadding = [this, bIsTopNode, OwnerTableViewWeak]()
+		{
+			FMargin Margin = GetRowScrollBarPadding(OwnerTableViewWeak);
+			if (bIsTopNode)
+			{
+				Margin.Bottom = 0;
+			}
+			return Margin;
+		};
+
+		FMargin BorderPadding = bIsTopNode ? FMargin(0, 0, 0, 2) : FMargin(0, 3, 0, 2);
+
 		ChildSlot
 		[
-			SNew( SBorder )
+			SNew(SBorder)
 			.BorderImage(FAppStyle::Get().GetBrush( "DetailsView.GridLine"))
-			.Padding( FMargin(0, 0, SDetailTableRowBase::ScrollbarPaddingSize, 0) )
+			.Padding_Lambda(GetGridLinePadding)
 			[
-				SNew( SBorder )
-				.BorderImage( FEditorStyle::GetBrush("DetailsView.AdvancedDropdownBorder") )
-				.Padding( FMargin( 0.0f, 3.0f, 0.0f, 2.0f ) )
+				SNew(SBorder)
+				.BorderImage(FEditorStyle::GetBrush("DetailsView.AdvancedDropdownBorder"))
+				.Padding(BorderPadding)
 				[
 					ContentWidget.ToSharedRef()
 				]
diff --git a/Engine/Source/Editor/PropertyEditor/Private/DetailCategoryGroupNode.cpp b/Engine/Source/Editor/PropertyEditor/Private/DetailCategoryGroupNode.cpp
--- a/Engine/Source/Editor/PropertyEditor/Private/DetailCategoryGroupNode.cpp
+++ b/Engine/Source/Editor/PropertyEditor/Private/DetailCategoryGroupNode.cpp
@@ -56,12 +56,13 @@ void SDetailCategoryTableRow::Construct( const FArguments& InArgs, TSharedRef<FD
 	
 	const float EditConditionWidgetWidth = 33; // this is the width of the edit condition widget displayed on the left in SDetailSingleItemRow
 
-	ChildSlot
-	.Padding(0)
+	TWeakPtr<STableViewBase> OwnerTableViewWeak = InOwnerTableView;
+	
+	this->ChildSlot
 	[
 		SNew( SBorder )
-		.BorderImage( FAppStyle::Get().GetBrush( "DetailsView.GridLine") )
-		.Padding( FMargin(0, 0, SDetailTableRowBase::ScrollbarPaddingSize, 1) )
+		.BorderImage(FAppStyle::Get().GetBrush("DetailsView.GridLine"))
+		.Padding(this, &SDetailTableRowBase::GetRowScrollBarPadding, OwnerTableViewWeak)
 		[
 			SNew(SHorizontalBox)
 			+ SHorizontalBox::Slot()
@@ -82,8 +83,8 @@ void SDetailCategoryTableRow::Construct( const FArguments& InArgs, TSharedRef<FD
 			[
 				SNew( SBorder )
 				.BorderImage( this, &SDetailCategoryTableRow::GetBackgroundImage )
-				.BorderBackgroundColor( this, &SDetailCategoryTableRow::GetInnerBackgroundColor )
-				.Padding(FMargin(0, 0, 0, 0))
+				.BorderBackgroundColor(this, &SDetailCategoryTableRow::GetInnerBackgroundColor)
+				.Padding(0)
 				[
 					HeaderBox
 				]
diff --git a/Engine/Source/Editor/PropertyEditor/Private/DetailMultiTopLevelObjectRootNode.cpp b/Engine/Source/Editor/PropertyEditor/Private/DetailMultiTopLevelObjectRootNode.cpp
--- a/Engine/Source/Editor/PropertyEditor/Private/DetailMultiTopLevelObjectRootNode.cpp
+++ b/Engine/Source/Editor/PropertyEditor/Private/DetailMultiTopLevelObjectRootNode.cpp
@@ -12,8 +12,6 @@ void SDetailMultiTopLevelObjectTableRow::Construct( const FArguments& InArgs, TS
 	OwnerTreeNode = InOwnerTreeNode;
 	ExpansionArrowUsage = InArgs._ExpansionArrowUsage;
 
-
-
 	STableRow< TSharedPtr< FDetailTreeNode > >::ConstructInternal(
 		STableRow::FArguments()
 			.Style(FEditorStyle::Get(), "DetailsView.TreeView.TableRow")
@@ -23,14 +21,16 @@ void SDetailMultiTopLevelObjectTableRow::Construct( const FArguments& InArgs, TS
 }
 
 
-void SDetailMultiTopLevelObjectTableRow::SetContent(TSharedRef<SWidget> InContent)
+void SDetailMultiTopLevelObjectTableRow::SetContent(TSharedRef<SWidget> InContent, TSharedRef<STableViewBase> OwnerTableView)
 {
 	if (ExpansionArrowUsage == EExpansionArrowUsage::Default)
 	{
+		TWeakPtr<STableViewBase> OwnerTableViewWeak = OwnerTableView;
+
 		ChildSlot
 		[
 			SNew(SBox)
-			.Padding(FMargin(0.0f, 0.0f, SDetailTableRowBase::ScrollbarPaddingSize, 0.0f))
+			.Padding(this, &SDetailTableRowBase::GetRowScrollBarPadding, OwnerTableViewWeak)
 			[
 				SNew(SHorizontalBox)
 				+ SHorizontalBox::Slot()
@@ -56,7 +56,6 @@ void SDetailMultiTopLevelObjectTableRow::SetContent(TSharedRef<SWidget> InConten
 			InContent
 		];
 	}
-	
 }
 
 const FSlateBrush* SDetailMultiTopLevelObjectTableRow::GetBackgroundImage() const
@@ -162,7 +161,7 @@ TSharedRef< ITableRow > FDetailMultiTopLevelObjectRootNode::GenerateWidgetForTab
 	FDetailWidgetRow Row;
 	GenerateWidget_Internal(Row, TableRowWidget);
 
-	TableRowWidget->SetContent(Row.NameWidget.Widget);
+	TableRowWidget->SetContent(Row.NameWidget.Widget, OwnerTable);
 
 	return TableRowWidget;
 }
diff --git a/Engine/Source/Editor/PropertyEditor/Private/SDetailSingleItemRow.cpp b/Engine/Source/Editor/PropertyEditor/Private/SDetailSingleItemRow.cpp
--- a/Engine/Source/Editor/PropertyEditor/Private/SDetailSingleItemRow.cpp
+++ b/Engine/Source/Editor/PropertyEditor/Private/SDetailSingleItemRow.cpp
@@ -532,11 +532,13 @@ void SDetailSingleItemRow::Construct( const FArguments& InArgs, FDetailLayoutCus
 		Widget = SNew(SSpacer);
 	}
 
+	TWeakPtr<STableViewBase> OwnerTableViewWeak = InOwnerTableView;
+	
 	this->ChildSlot
 	[
 		SNew( SBorder )
 		.BorderImage(FAppStyle::Get().GetBrush( "DetailsView.GridLine"))
-		.Padding( FMargin(0, 0, SDetailTableRowBase::ScrollbarPaddingSize, 1) )
+		.Padding(this, &SDetailTableRowBase::GetRowScrollBarPadding, OwnerTableViewWeak)
 		[
 			SNew( SBorder )
 			.BorderImage(FAppStyle::Get().GetBrush("DetailsView.CategoryMiddle"))
@@ -878,19 +880,19 @@ void SDetailSingleItemRow::OnFavoriteMenuToggle()
 	int32 ScrollingOffsetRemove = -ExpandSize;
 	if (HasFavoritesCategory)
 	{
-		// Adding the advance button in a category add 1 item
-		ScrollingOffsetAdd += (IsAdvancedProperty && NumAdvancedProperties == 0) ? 1 : 0;
+	// Adding the advance button in a category add 1 item
+	ScrollingOffsetAdd += (IsAdvancedProperty && NumAdvancedProperties == 0) ? 1 : 0;
 
-		if (IsAdvancedProperty && NumAdvancedProperties == 1)
-		{
-			// Removing the advance button count as 1 item
-			ScrollingOffsetRemove -= 1;
-		}
-		if (NumAdvancedProperties + SimplePropertiesNum == 1)
-		{
-			// Removing a full category count as 2 items
-			ScrollingOffsetRemove -= 2;
-		}
+	if (IsAdvancedProperty && NumAdvancedProperties == 1)
+	{
+		// Removing the advance button count as 1 item
+		ScrollingOffsetRemove -= 1;
+	}
+	if (NumAdvancedProperties + SimplePropertiesNum == 1)
+	{
+		// Removing a full category count as 2 items
+		ScrollingOffsetRemove -= 2;
+	}
 	}
 	else
 	{
diff --git a/Engine/Source/Editor/PropertyEditor/Private/SDetailTableRowBase.cpp b/Engine/Source/Editor/PropertyEditor/Private/SDetailTableRowBase.cpp
--- a/Engine/Source/Editor/PropertyEditor/Private/SDetailTableRowBase.cpp
+++ b/Engine/Source/Editor/PropertyEditor/Private/SDetailTableRowBase.cpp
@@ -2,8 +2,6 @@
 
 #include "SDetailTableRowBase.h"
 
-float SDetailTableRowBase::ScrollbarPaddingSize = 16.0f;
-
 int32 SDetailTableRowBase::GetIndentLevelForBackgroundColor() const
 {
 	int32 IndentLevel = 0; 
@@ -23,4 +21,18 @@ int32 SDetailTableRowBase::GetIndentLevelForBackgroundColor() const
 	}
 
 	return FMath::Max(0, IndentLevel);
-}
\ No newline at end of file
+}
+
+FMargin SDetailTableRowBase::GetRowScrollBarPadding(TWeakPtr<STableViewBase> OwnerTableViewWeak) const
+{
+	TSharedPtr<STableViewBase> OwnerTableView = OwnerTableViewWeak.Pin();
+	if (OwnerTableView.IsValid())
+	{
+		if (OwnerTableView->GetScrollbarVisibility() == EVisibility::Visible)
+		{
+			const float ScrollbarPaddingSize = 16.0f;
+			return FMargin(0, 0, ScrollbarPaddingSize, 1);
+		}
+	}
+	return FMargin(0, 0, 0, 1);
+}
diff --git a/Engine/Source/Editor/PropertyEditor/Private/SDetailTableRowBase.h b/Engine/Source/Editor/PropertyEditor/Private/SDetailTableRowBase.h
--- a/Engine/Source/Editor/PropertyEditor/Private/SDetailTableRowBase.h
+++ b/Engine/Source/Editor/PropertyEditor/Private/SDetailTableRowBase.h
@@ -61,6 +61,7 @@ public:
 	}
 
 	int32 GetIndentLevelForBackgroundColor() const;
+	FMargin GetRowScrollBarPadding(TWeakPtr<STableViewBase> OwnerTableView) const;
 
 protected:
 	/**
@@ -94,6 +95,5 @@ private:
 		}
 	}
 protected:
-	static float ScrollbarPaddingSize;
 	TWeakPtr<FDetailTreeNode> OwnerTreeNode;
 };
diff --git a/Engine/Source/Runtime/Slate/Private/Widgets/Views/STableViewBase.cpp b/Engine/Source/Runtime/Slate/Private/Widgets/Views/STableViewBase.cpp
--- a/Engine/Source/Runtime/Slate/Private/Widgets/Views/STableViewBase.cpp
+++ b/Engine/Source/Runtime/Slate/Private/Widgets/Views/STableViewBase.cpp
@@ -837,6 +837,11 @@ void STableViewBase::SetScrollbarVisibility(const EVisibility InVisibility)
 	}
 }
 
+EVisibility STableViewBase::GetScrollbarVisibility() const
+{
+	return ScrollBar ? ScrollBar->ShouldBeVisible() : EVisibility::Collapsed;
+}
+
 void STableViewBase::SetFixedLineScrollOffset(TOptional<double> InFixedLineScrollOffset)
 {
 	if (FixedLineScrollOffset != InFixedLineScrollOffset)
diff --git a/Engine/Source/Runtime/Slate/Public/Widgets/Views/STableViewBase.h b/Engine/Source/Runtime/Slate/Public/Widgets/Views/STableViewBase.h
--- a/Engine/Source/Runtime/Slate/Public/Widgets/Views/STableViewBase.h
+++ b/Engine/Source/Runtime/Slate/Public/Widgets/Views/STableViewBase.h
@@ -153,6 +153,8 @@ public:
 	/** Add the scroll offset of this view (in items) */
 	void AddScrollOffset(const float InScrollOffsetDelta, bool RefreshList = false);
 
+	EVisibility GetScrollbarVisibility() const;
+
 	void SetScrollbarVisibility(const EVisibility InVisibility);
 
 	/** Sets the fixed offset in items to always apply to the top/left (depending on orientation) of the list. */
