Missed file in previous submit.

diff --git a/Engine/Source/Editor/PropertyEditor/Private/DetailMultiTopLevelObjectRootNode.h b/Engine/Source/Editor/PropertyEditor/Private/DetailMultiTopLevelObjectRootNode.h
--- a/Engine/Source/Editor/PropertyEditor/Private/DetailMultiTopLevelObjectRootNode.h
+++ b/Engine/Source/Editor/PropertyEditor/Private/DetailMultiTopLevelObjectRootNode.h
@@ -29,19 +29,18 @@ public:
 	SLATE_END_ARGS()
 
 	void Construct(const FArguments& InArgs, TSharedRef<FDetailTreeNode> InOwnerTreeNode, const TSharedRef<STableViewBase>& InOwnerTableView);
-	void SetContent(TSharedRef<SWidget> InContent);
-private:
-	const FSlateBrush* GetBackgroundImage() const;
+	void SetContent(TSharedRef<SWidget> InContent, TSharedRef<STableViewBase> InOwnerTableView);
 
 private:
+	const FSlateBrush* GetBackgroundImage() const;
 	virtual FReply OnMouseButtonDown(const FGeometry& MyGeometry, const FPointerEvent& MouseEvent) override;
 	virtual FReply OnMouseButtonDoubleClick(const FGeometry& InMyGeometry, const FPointerEvent& InMouseEvent) override;
+
 private:
 	EExpansionArrowUsage ExpansionArrowUsage;
 	SHorizontalBox::FSlot* ContentSlot = nullptr;
 };
 
-
 class FDetailMultiTopLevelObjectRootNode : public FDetailTreeNode, public TSharedFromThis<FDetailMultiTopLevelObjectRootNode>
 {
 public:
@@ -61,8 +60,8 @@ private:
 	virtual FName GetNodeName() const override { return NodeName; }
 	virtual EDetailNodeType GetNodeType() const override { return EDetailNodeType::Object; }
 	virtual TSharedPtr<IPropertyHandle> CreatePropertyHandle() const override { return nullptr; }
-private:
 	void GenerateWidget_Internal(FDetailWidgetRow& Row, TSharedPtr<SDetailMultiTopLevelObjectTableRow> TableRow) const;
+
 private:
 	FDetailNodeList ChildNodes;
 	IDetailsViewPrivate* DetailsView;
