SPropertyEditorButton now wraps its buttons in a 20x20 SBox so that they aren't inadvertently stretched.

diff --git a/Engine/Source/Editor/PropertyEditor/Private/PropertyCustomizationHelpers.cpp b/Engine/Source/Editor/PropertyEditor/Private/PropertyCustomizationHelpers.cpp
--- a/Engine/Source/Editor/PropertyEditor/Private/PropertyCustomizationHelpers.cpp
+++ b/Engine/Source/Editor/PropertyEditor/Private/PropertyCustomizationHelpers.cpp
@@ -41,7 +41,7 @@
 
 namespace PropertyCustomizationHelpers
 {
-	class SPropertyEditorButton : public SButton
+	class SPropertyEditorButton : public SCompoundWidget
 	{
 	public:
 
@@ -62,20 +62,28 @@ namespace PropertyCustomizationHelpers
 		{
 			OnClickAction = InArgs._OnClickAction;
 
-			SButton::FArguments ButtonArgs = SButton::FArguments()
-				.ButtonStyle( FAppStyle::Get(), "SimpleButton" )
-				.OnClicked( this, &SPropertyEditorButton::OnClick )
-				.ToolTipText( InArgs._Text )
-				.ForegroundColor( FSlateColor::UseForeground() )
-				.ContentPadding(0)
-				.IsFocusable(InArgs._IsFocusable)
-				[ 
-					SNew( SImage )
-					.Image( InArgs._Image )
-					.ColorAndOpacity( FSlateColor::UseForeground() )
-				]; 
-
-			SButton::Construct( ButtonArgs );
+			ChildSlot
+			[
+				SNew(SBox)
+				.HAlign(HAlign_Center)
+				.VAlign(VAlign_Center)
+				.WidthOverride(20)
+				.HeightOverride(20)
+				[
+					SNew(SButton)
+					.ButtonStyle( FAppStyle::Get(), "SimpleButton" )
+					.OnClicked( this, &SPropertyEditorButton::OnClick )
+					.ToolTipText( InArgs._Text )
+					.ForegroundColor( FSlateColor::UseForeground() )
+					.ContentPadding(0)
+					.IsFocusable(InArgs._IsFocusable)
+					[ 
+						SNew( SImage )
+						.Image( InArgs._Image )
+						.ColorAndOpacity( FSlateColor::UseForeground() )
+					]
+				]
+			]; 
 		}
 
 
