Fixing Clang build error due to SetContent's signature being changed.

diff --git a/Engine/Source/Editor/PropertyEditor/Private/DetailMultiTopLevelObjectRootNode.cpp b/Engine/Source/Editor/PropertyEditor/Private/DetailMultiTopLevelObjectRootNode.cpp
--- a/Engine/Source/Editor/PropertyEditor/Private/DetailMultiTopLevelObjectRootNode.cpp
+++ b/Engine/Source/Editor/PropertyEditor/Private/DetailMultiTopLevelObjectRootNode.cpp
@@ -11,6 +11,7 @@ void SDetailMultiTopLevelObjectTableRow::Construct( const FArguments& InArgs, TS
 {
 	OwnerTreeNode = InOwnerTreeNode;
 	ExpansionArrowUsage = InArgs._ExpansionArrowUsage;
+	OwnerTableViewWeak = InOwnerTableView;
 
 	STableRow< TSharedPtr< FDetailTreeNode > >::ConstructInternal(
 		STableRow::FArguments()
@@ -20,13 +21,10 @@ void SDetailMultiTopLevelObjectTableRow::Construct( const FArguments& InArgs, TS
 	);
 }
 
-
-void SDetailMultiTopLevelObjectTableRow::SetContent(TSharedRef<SWidget> InContent, TSharedRef<STableViewBase> OwnerTableView)
+void SDetailMultiTopLevelObjectTableRow::SetContent(TSharedRef<SWidget> InContent)
 {
 	if (ExpansionArrowUsage == EExpansionArrowUsage::Default)
 	{
-		TWeakPtr<STableViewBase> OwnerTableViewWeak = OwnerTableView;
-
 		ChildSlot
 		[
 			SNew(SBox)
@@ -161,7 +159,7 @@ TSharedRef< ITableRow > FDetailMultiTopLevelObjectRootNode::GenerateWidgetForTab
 	FDetailWidgetRow Row;
 	GenerateWidget_Internal(Row, TableRowWidget);
 
-	TableRowWidget->SetContent(Row.NameWidget.Widget, OwnerTable);
+	TableRowWidget->SetContent(Row.NameWidget.Widget);
 
 	return TableRowWidget;
 }
diff --git a/Engine/Source/Editor/PropertyEditor/Private/DetailMultiTopLevelObjectRootNode.h b/Engine/Source/Editor/PropertyEditor/Private/DetailMultiTopLevelObjectRootNode.h
--- a/Engine/Source/Editor/PropertyEditor/Private/DetailMultiTopLevelObjectRootNode.h
+++ b/Engine/Source/Editor/PropertyEditor/Private/DetailMultiTopLevelObjectRootNode.h
@@ -29,7 +29,7 @@ public:
 	SLATE_END_ARGS()
 
 	void Construct(const FArguments& InArgs, TSharedRef<FDetailTreeNode> InOwnerTreeNode, const TSharedRef<STableViewBase>& InOwnerTableView);
-	void SetContent(TSharedRef<SWidget> InContent, TSharedRef<STableViewBase> InOwnerTableView);
+	void SetContent(TSharedRef<SWidget> InContent) override;
 
 private:
 	const FSlateBrush* GetBackgroundImage() const;
@@ -39,6 +39,7 @@ private:
 private:
 	EExpansionArrowUsage ExpansionArrowUsage;
 	SHorizontalBox::FSlot* ContentSlot = nullptr;
+	TWeakPtr<STableViewBase> OwnerTableViewWeak;
 };
 
 class FDetailMultiTopLevelObjectRootNode : public FDetailTreeNode, public TSharedFromThis<FDetailMultiTopLevelObjectRootNode>
