Fixes for CIS errors.

diff --git a/Engine/Source/Editor/EditorConfig/Private/Tests/EditorConfigTestHelpers.h b/Engine/Source/Editor/EditorConfig/Private/Tests/EditorConfigTestHelpers.h
--- a/Engine/Source/Editor/EditorConfig/Private/Tests/EditorConfigTestHelpers.h
+++ b/Engine/Source/Editor/EditorConfig/Private/Tests/EditorConfigTestHelpers.h
@@ -5,5 +5,5 @@
 namespace FEditorConfigTestHelpers
 {
 	// check that two serialized JSON objects are equivalent
-	static bool AreJsonStringsEquivalent(const FString& Actual, const FString& Expected);
+	bool AreJsonStringsEquivalent(const FString& Actual, const FString& Expected);
 }
\ No newline at end of file
diff --git a/Engine/Source/Editor/PropertyEditor/Private/ItemPropertyNode.cpp b/Engine/Source/Editor/PropertyEditor/Private/ItemPropertyNode.cpp
--- a/Engine/Source/Editor/PropertyEditor/Private/ItemPropertyNode.cpp
+++ b/Engine/Source/Editor/PropertyEditor/Private/ItemPropertyNode.cpp
@@ -1,7 +1,6 @@
 // Copyright Epic Games, Inc. All Rights Reserved.
 
 #include "ItemPropertyNode.h"
-#include "DetailsViewConfig.h"
 #include "Editor.h"
 #include "EditorMetadataOverrides.h"
 #include "ObjectPropertyNode.h"
diff --git a/Engine/Source/Editor/PropertyEditor/Private/SDetailSingleItemRow.cpp b/Engine/Source/Editor/PropertyEditor/Private/SDetailSingleItemRow.cpp
--- a/Engine/Source/Editor/PropertyEditor/Private/SDetailSingleItemRow.cpp
+++ b/Engine/Source/Editor/PropertyEditor/Private/SDetailSingleItemRow.cpp
@@ -858,55 +858,10 @@ void SDetailSingleItemRow::OnFavoriteMenuToggle()
 		// if the item is not expand count is 1
 		ExpandSize = 1;
 	}
-			
-	// Get the number of favorite child (simple and advance) to know if the favorite category will be create or remove
-	int32 SimplePropertiesNum = 0;
-	int32 NumAdvancedProperties = 0;
-
-	FDetailLayoutBuilderImpl& DetailLayout = OwnerTreeNodePinned->GetParentCategory()->GetParentLayoutImpl();
-
-	const FName FavoritesCategoryName = TEXT("Favorites");
-	bool HasFavoritesCategory = DetailLayout.HasCategory(FavoritesCategoryName);
-	if (HasFavoritesCategory)
-	{
-		DetailLayout.DefaultCategory(FavoritesCategoryName).GetCategoryInformation(SimplePropertiesNum, NumAdvancedProperties);
-	}
-
-	// Check if the property we toggle is an advanced property
-	bool IsAdvancedProperty = Customization->GetPropertyNode()->HasNodeFlags(EPropertyNodeFlags::IsAdvanced) != 0;
-
-	// Compute the scrolling offset by item
-	int32 ScrollingOffsetAdd = ExpandSize;
-	int32 ScrollingOffsetRemove = -ExpandSize;
-	if (HasFavoritesCategory)
-	{
-	// Adding the advance button in a category add 1 item
-	ScrollingOffsetAdd += (IsAdvancedProperty && NumAdvancedProperties == 0) ? 1 : 0;
-
-	if (IsAdvancedProperty && NumAdvancedProperties == 1)
-	{
-		// Removing the advance button count as 1 item
-		ScrollingOffsetRemove -= 1;
-	}
-	if (NumAdvancedProperties + SimplePropertiesNum == 1)
-	{
-		// Removing a full category count as 2 items
-		ScrollingOffsetRemove -= 2;
-	}
-	}
-	else
-	{
-		// Adding new category (2 items) adding advance button (1 item)
-		ScrollingOffsetAdd += IsAdvancedProperty ? 3 : 2;
-				
-		// We should never remove an item from favorite if there is no favorite category
-		// Set the remove offset to 0
-		ScrollingOffsetRemove = 0;
-	}
 
 	// Apply the calculated offset
 	IDetailsViewPrivate* DetailsView = OwnerTreeNodePinned->GetDetailsView();
-	DetailsView->MoveScrollOffset(bToggled ? ScrollingOffsetAdd : ScrollingOffsetRemove);
+	DetailsView->MoveScrollOffset(bToggled ? ExpandSize : -ExpandSize);
 
 	// Refresh the tree
 	DetailsView->ForceRefresh();
