Temp fix for slider properties in blueprints not working. Comment out some code thats not working yet that came over from UE4.

diff --git a/Engine/Source/Editor/PropertyEditor/Private/SDetailsView.cpp b/Engine/Source/Editor/PropertyEditor/Private/SDetailsView.cpp
--- a/Engine/Source/Editor/PropertyEditor/Private/SDetailsView.cpp
+++ b/Engine/Source/Editor/PropertyEditor/Private/SDetailsView.cpp
@@ -713,10 +713,11 @@ void SDetailsView::ReplaceObjects(const TMap<UObject*, UObject*>& OldToNewObject
 		}
 	}
 
+/*
 	if (!bNeedRefresh)
 	{
 		bNeedRefresh |= AreExternalObjectsReplaced(OldToNewObjectMap);
-	}
+	}*/
 
 	if (bNeedRefresh)
 	{
