Fix for crash bug if GEditor is null when checking if an FItemPropertyNode is favorited.

diff --git a/Engine/Source/Editor/PropertyEditor/Private/ItemPropertyNode.cpp b/Engine/Source/Editor/PropertyEditor/Private/ItemPropertyNode.cpp
--- a/Engine/Source/Editor/PropertyEditor/Private/ItemPropertyNode.cpp
+++ b/Engine/Source/Editor/PropertyEditor/Private/ItemPropertyNode.cpp
@@ -392,25 +392,30 @@ void FItemPropertyNode::InitChildNodes()
 
 void FItemPropertyNode::SetFavorite(bool FavoriteValue)
 {
-	UEditorMetadataOverrides* MetadataOverrides = GEditor->GetEditorSubsystem<UEditorMetadataOverrides>();
-	if (MetadataOverrides != nullptr)
+	if (GEditor != nullptr)
 	{
-		MetadataOverrides->SetBoolMetadata(GetProperty(), "IsFavorite", FavoriteValue);
+		UEditorMetadataOverrides* MetadataOverrides = GEditor->GetEditorSubsystem<UEditorMetadataOverrides>();
+		if (MetadataOverrides != nullptr)
+		{
+			MetadataOverrides->SetBoolMetadata(GetProperty(), "IsFavorite", FavoriteValue);
+		}
 	}
 }
 
 bool FItemPropertyNode::IsFavorite() const
 {
-	UEditorMetadataOverrides* MetadataOverrides = GEditor->GetEditorSubsystem<UEditorMetadataOverrides>();
-	if (MetadataOverrides != nullptr)
+	if (GEditor != nullptr)
 	{
-		bool IsFavorite = false;
-		if (MetadataOverrides->GetBoolMetadata(GetProperty(), "IsFavorite", IsFavorite))
+		UEditorMetadataOverrides* MetadataOverrides = GEditor->GetEditorSubsystem<UEditorMetadataOverrides>();
+		if (MetadataOverrides != nullptr)
 		{
-			return IsFavorite;
+			bool IsFavorite = false;
+			if (MetadataOverrides->GetBoolMetadata(GetProperty(), "IsFavorite", IsFavorite))
+			{
+				return IsFavorite;
+			}
 		}
 	}
-
 	return false;
 }
 
