Fixed up locations making one off thumbnail pools and converted them to use the shared thumbnail pool.

diff --git a/Engine/Plugins/Editor/AssetSearch/Source/Private/Widgets/SSearchBrowser.cpp b/Engine/Plugins/Editor/AssetSearch/Source/Private/Widgets/SSearchBrowser.cpp
--- a/Engine/Plugins/Editor/AssetSearch/Source/Private/Widgets/SSearchBrowser.cpp
+++ b/Engine/Plugins/Editor/AssetSearch/Source/Private/Widgets/SSearchBrowser.cpp
@@ -42,13 +42,6 @@
 #define LOCTEXT_NAMESPACE "SObjectBrowser"
 
 DEFINE_LOG_CATEGORY_STATIC(LogObjectBrowser, Log, All)
-
-namespace AssetSearchConstants
-{
-	/** The size of the thumbnail pool */
-	const int32 ThumbnailPoolSize = 64;
-}
-
 //;
 //TODO Expose TSharedRef<SWidget> SAssetViewItem::CreateToolTipWidget() const via IContentBrowserSingleton.
 
@@ -75,7 +68,6 @@ void SSearchBrowser::Construct( const FArguments& InArgs )
 	SortByColumn = SSearchTreeRow::NAME_ColumnName;
 	SortMode = EColumnSortMode::Ascending;
 
-	ThumbnailPool = MakeShared<FAssetThumbnailPool>(AssetSearchConstants::ThumbnailPoolSize);
 
 	AssetRegistry = &FModuleManager::LoadModuleChecked<FAssetRegistryModule>("AssetRegistry").Get();
 
@@ -479,7 +471,7 @@ void SSearchBrowser::TryRefreshingSearch(const FText& InText)
 
 TSharedRef<ITableRow> SSearchBrowser::HandleListGenerateRow(TSharedPtr<FSearchNode> ObjectPtr, const TSharedRef<STableViewBase>& OwnerTable)
 {
-	return SNew(SSearchTreeRow, OwnerTable, AssetRegistry, ThumbnailPool)
+	return SNew(SSearchTreeRow, OwnerTable, AssetRegistry, UThumbnailManager::Get().GetSharedThumbnailPool())
 		.Object(ObjectPtr)
 		.HighlightText(FilterText);
 }
diff --git a/Engine/Plugins/Editor/AssetSearch/Source/Private/Widgets/SSearchBrowser.h b/Engine/Plugins/Editor/AssetSearch/Source/Private/Widgets/SSearchBrowser.h
--- a/Engine/Plugins/Editor/AssetSearch/Source/Private/Widgets/SSearchBrowser.h
+++ b/Engine/Plugins/Editor/AssetSearch/Source/Private/Widgets/SSearchBrowser.h
@@ -88,8 +88,6 @@ private:
 
 	IAssetRegistry* AssetRegistry = nullptr;
 
-	TSharedPtr<FAssetThumbnailPool> ThumbnailPool;
-
 	TSharedPtr<SHeaderRow> HeaderColumns;
 
 	FName SortByColumn;
diff --git a/Engine/Plugins/Enterprise/VariantManager/Source/VariantManager/Private/DisplayNodes/VariantManagerVariantNode.cpp b/Engine/Plugins/Enterprise/VariantManager/Source/VariantManager/Private/DisplayNodes/VariantManagerVariantNode.cpp
--- a/Engine/Plugins/Enterprise/VariantManager/Source/VariantManager/Private/DisplayNodes/VariantManagerVariantNode.cpp
+++ b/Engine/Plugins/Enterprise/VariantManager/Source/VariantManager/Private/DisplayNodes/VariantManagerVariantNode.cpp
@@ -24,6 +24,7 @@
 #include "ScopedTransaction.h"
 #include "Textures/SlateIcon.h"
 #include "Widgets/Notifications/SNotificationList.h"
+#include "ThumbnailRendering/ThumbnailManager.h"
 
 #define LOCTEXT_NAMESPACE "VariantManagerVariantNode"
 
@@ -507,9 +508,7 @@ TSharedRef<SWidget> FVariantManagerVariantNode::GetThumbnailWidget()
 	// Use default thumbnail graphic for Variants as a fallback
 	if(!ThumbnailWidget.IsValid())
 	{
-		ThumbnailPool = MakeShared<FAssetThumbnailPool>(1, false);
-
-		TSharedRef<FAssetThumbnail> AssetThumbnail = MakeShared<FAssetThumbnail>(&Variant, VARIANT_MANAGER_THUMBNAIL_SIZE, VARIANT_MANAGER_THUMBNAIL_SIZE, ThumbnailPool);
+		TSharedRef<FAssetThumbnail> AssetThumbnail = MakeShared<FAssetThumbnail>(&Variant, VARIANT_MANAGER_THUMBNAIL_SIZE, VARIANT_MANAGER_THUMBNAIL_SIZE, UThumbnailManager::Get().GetSharedThumbnailPool());
 		ThumbnailWidget = AssetThumbnail->MakeThumbnailWidget();
 	}
 
diff --git a/Engine/Plugins/Enterprise/VariantManager/Source/VariantManager/Private/DisplayNodes/VariantManagerVariantNode.h b/Engine/Plugins/Enterprise/VariantManager/Source/VariantManager/Private/DisplayNodes/VariantManagerVariantNode.h
--- a/Engine/Plugins/Enterprise/VariantManager/Source/VariantManager/Private/DisplayNodes/VariantManagerVariantNode.h
+++ b/Engine/Plugins/Enterprise/VariantManager/Source/VariantManager/Private/DisplayNodes/VariantManagerVariantNode.h
@@ -6,7 +6,6 @@
 
 #include "CoreMinimal.h"
 
-class FAssetThumbnailPool;
 class FDragDropEvent;
 class FMenuBuilder;
 class SCheckBox;
@@ -49,8 +48,6 @@ private:
 
 	TSharedPtr<SCheckBox> RadioButton;
 
-	TSharedPtr<FAssetThumbnailPool> ThumbnailPool;
-
 	// We need to keep the brush alive while the thumbnail is rendered
 	TSharedPtr<FSlateImageBrush> ImageBrush;
 };
diff --git a/Engine/Plugins/Experimental/ControlRig/Source/ControlRigEditor/Private/EditMode/SControlRigControlViews.cpp b/Engine/Plugins/Experimental/ControlRig/Source/ControlRigEditor/Private/EditMode/SControlRigControlViews.cpp
--- a/Engine/Plugins/Experimental/ControlRig/Source/ControlRigEditor/Private/EditMode/SControlRigControlViews.cpp
+++ b/Engine/Plugins/Experimental/ControlRig/Source/ControlRigEditor/Private/EditMode/SControlRigControlViews.cpp
@@ -586,12 +586,8 @@ FReply SControlRigPoseView::OnCaptureThumbnail()
 
 TSharedRef<SWidget> SControlRigPoseView::GetThumbnailWidget()
 {
-	FLevelEditorModule& LevelEditorModule = FModuleManager::LoadModuleChecked<FLevelEditorModule>("LevelEditor");
-	TSharedPtr<FAssetThumbnailPool> ThumbnailPool = LevelEditorModule.GetFirstLevelEditor()->GetThumbnailPool();
-
-
 	const int32 ThumbnailSize = 128;
-	Thumbnail = MakeShareable(new FAssetThumbnail(PoseAsset.Get(), ThumbnailSize, ThumbnailSize, ThumbnailPool));
+	Thumbnail = MakeShareable(new FAssetThumbnail(PoseAsset.Get(), ThumbnailSize, ThumbnailSize, UThumbnailManager::Get().GetSharedThumbnailPool()));
 	FAssetThumbnailConfig ThumbnailConfig;
 	ThumbnailConfig.bAllowFadeIn = false;
 	ThumbnailConfig.bAllowHintText = false;
diff --git a/Engine/Plugins/Experimental/VirtualHeightfieldMesh/Source/VirtualHeightfieldMeshEditor/Private/VirtualHeightfieldMeshDetailsCustomization.cpp b/Engine/Plugins/Experimental/VirtualHeightfieldMesh/Source/VirtualHeightfieldMeshEditor/Private/VirtualHeightfieldMeshDetailsCustomization.cpp
--- a/Engine/Plugins/Experimental/VirtualHeightfieldMesh/Source/VirtualHeightfieldMeshEditor/Private/VirtualHeightfieldMeshDetailsCustomization.cpp
+++ b/Engine/Plugins/Experimental/VirtualHeightfieldMesh/Source/VirtualHeightfieldMeshEditor/Private/VirtualHeightfieldMeshDetailsCustomization.cpp
@@ -23,8 +23,7 @@
 #define LOCTEXT_NAMESPACE "VirtualHeightfieldMeshEditorModule"
 
 FVirtualHeightfieldMeshComponentDetailsCustomization::FVirtualHeightfieldMeshComponentDetailsCustomization()
-	: AssetThumbnailPool(MakeShareable(new FAssetThumbnailPool(4)))
-	, AssetThumbnail(MakeShared<FAssetThumbnail>(nullptr, 64, 64, AssetThumbnailPool))
+	: AssetThumbnail(MakeShared<FAssetThumbnail>(nullptr, 64, 64, UThumbnailManager::Get().GetSharedThumbnailPool()))
 {
 }
 
diff --git a/Engine/Plugins/Experimental/VirtualHeightfieldMesh/Source/VirtualHeightfieldMeshEditor/Private/VirtualHeightfieldMeshDetailsCustomization.h b/Engine/Plugins/Experimental/VirtualHeightfieldMesh/Source/VirtualHeightfieldMeshEditor/Private/VirtualHeightfieldMeshDetailsCustomization.h
--- a/Engine/Plugins/Experimental/VirtualHeightfieldMesh/Source/VirtualHeightfieldMeshEditor/Private/VirtualHeightfieldMeshDetailsCustomization.h
+++ b/Engine/Plugins/Experimental/VirtualHeightfieldMesh/Source/VirtualHeightfieldMeshEditor/Private/VirtualHeightfieldMeshDetailsCustomization.h
@@ -32,7 +32,5 @@ protected:
 
 private:
 	class UVirtualHeightfieldMeshComponent* VirtualHeightfieldMeshComponent;
-	
-	TSharedPtr<class FAssetThumbnailPool> AssetThumbnailPool;
 	TSharedRef<class FAssetThumbnail> AssetThumbnail;
 };
diff --git a/Engine/Plugins/FX/Niagara/Source/NiagaraEditor/Private/SNiagaraAssetPickerList.cpp b/Engine/Plugins/FX/Niagara/Source/NiagaraEditor/Private/SNiagaraAssetPickerList.cpp
--- a/Engine/Plugins/FX/Niagara/Source/NiagaraEditor/Private/SNiagaraAssetPickerList.cpp
+++ b/Engine/Plugins/FX/Niagara/Source/NiagaraEditor/Private/SNiagaraAssetPickerList.cpp
@@ -12,6 +12,7 @@
 #include "Widgets/Input/SButton.h"
 #include "Widgets/Layout/SUniformGridPanel.h"
 #include "EditorStyleSet.h"
+#include "ThumbnailRendering/ThumbnailManager.h"
 
 #define LOCTEXT_NAMESPACE "SNiagaraAssetSelector"
 
@@ -28,8 +29,6 @@ void SNiagaraAssetPickerList::Construct(const FArguments& InArgs, UClass* AssetC
 	TArray<FAssetData> EmitterAssets;
 	AssetRegistryModule.Get().GetAssetsByClass(AssetClass->GetFName(), EmitterAssets);
 
-	AssetThumbnailPool = MakeShareable(new FAssetThumbnailPool(24));
-
 	OnTemplateAssetActivated = InArgs._OnTemplateAssetActivated;
 	ViewOptions = InArgs._ViewOptions;
 
@@ -72,7 +71,6 @@ void SNiagaraAssetPickerList::Construct(const FArguments& InArgs, UClass* AssetC
 
 void SNiagaraAssetPickerList::Tick(const FGeometry& AllottedGeometry, const double InCurrentTime, const float InDeltaTime)
 {
-	AssetThumbnailPool->Tick(InDeltaTime);
 }
 
 TArray<FAssetData> SNiagaraAssetPickerList::GetSelectedAssets() const
@@ -311,7 +309,7 @@ const int32 ThumbnailSize = 72;
 
 TSharedRef<SWidget> SNiagaraAssetPickerList::OnGenerateWidgetForItem(const FAssetData& Item)
 {
-	TSharedRef<FAssetThumbnail> AssetThumbnail = MakeShared<FAssetThumbnail>(Item, ThumbnailSize, ThumbnailSize, AssetThumbnailPool);
+	TSharedRef<FAssetThumbnail> AssetThumbnail = MakeShared<FAssetThumbnail>(Item, ThumbnailSize, ThumbnailSize, UThumbnailManager::Get().GetSharedThumbnailPool());
 	FAssetThumbnailConfig ThumbnailConfig;
 	ThumbnailConfig.bAllowFadeIn = false;
 	
diff --git a/Engine/Plugins/FX/Niagara/Source/NiagaraEditor/Private/SNiagaraAssetPickerList.h b/Engine/Plugins/FX/Niagara/Source/NiagaraEditor/Private/SNiagaraAssetPickerList.h
--- a/Engine/Plugins/FX/Niagara/Source/NiagaraEditor/Private/SNiagaraAssetPickerList.h
+++ b/Engine/Plugins/FX/Niagara/Source/NiagaraEditor/Private/SNiagaraAssetPickerList.h
@@ -121,7 +121,6 @@ private:
 	static FText LibraryCategory;
 	static FText NonLibraryCategory;
 	static FText UncategorizedCategory;
-	TSharedPtr<FAssetThumbnailPool> AssetThumbnailPool;
 	FOnTemplateAssetActivated OnTemplateAssetActivated;
 	FOnDoesAssetPassCustomFilter OnDoesAssetPassCustomFilter;
 	FNiagaraAssetPickerListViewOptions ViewOptions;
diff --git a/Engine/Plugins/FX/Niagara/Source/NiagaraEditorWidgets/Private/SNiagaraOverviewStackNode.cpp b/Engine/Plugins/FX/Niagara/Source/NiagaraEditorWidgets/Private/SNiagaraOverviewStackNode.cpp
--- a/Engine/Plugins/FX/Niagara/Source/NiagaraEditorWidgets/Private/SNiagaraOverviewStackNode.cpp
+++ b/Engine/Plugins/FX/Niagara/Source/NiagaraEditorWidgets/Private/SNiagaraOverviewStackNode.cpp
@@ -42,7 +42,7 @@ void SNiagaraOverviewStackNode::Construct(const FArguments& InArgs, UNiagaraOver
 	CurrentIssueIndex = -1;
 
 	EmitterHandleViewModelWeak.Reset();
-	ThumbnailPool = MakeShared<FAssetThumbnailPool>(10);
+
 	ThumbnailBar = SNew(SHorizontalBox);
 	if (OverviewStackNode->GetOwningSystem() != nullptr)
 	{
@@ -403,9 +403,9 @@ void SNiagaraOverviewStackNode::FillThumbnailBar(UObject* ChangedObject, const b
 				if (UNiagaraStackRendererItem* RendererItem = Cast<UNiagaraStackRendererItem>(Entry))
 				{
 					TArray<TSharedPtr<SWidget>> Widgets;
-					RendererItem->GetRendererProperties()->GetRendererWidgets(InInstance, Widgets, ThumbnailPool);
+					RendererItem->GetRendererProperties()->GetRendererWidgets(InInstance, Widgets, UThumbnailManager::Get().GetSharedThumbnailPool());
 					TArray<TSharedPtr<SWidget>> TooltipWidgets;
-					RendererItem->GetRendererProperties()->GetRendererTooltipWidgets(InInstance, TooltipWidgets, ThumbnailPool);
+					RendererItem->GetRendererProperties()->GetRendererTooltipWidgets(InInstance, TooltipWidgets, UThumbnailManager::Get().GetSharedThumbnailPool());
 					for (int32 WidgetIndex = 0; WidgetIndex < Widgets.Num(); WidgetIndex++)
 					{
 						ThumbnailBar->AddSlot()
diff --git a/Engine/Plugins/FX/Niagara/Source/NiagaraEditorWidgets/Private/SNiagaraOverviewStackNode.h b/Engine/Plugins/FX/Niagara/Source/NiagaraEditorWidgets/Private/SNiagaraOverviewStackNode.h
--- a/Engine/Plugins/FX/Niagara/Source/NiagaraEditorWidgets/Private/SNiagaraOverviewStackNode.h
+++ b/Engine/Plugins/FX/Niagara/Source/NiagaraEditorWidgets/Private/SNiagaraOverviewStackNode.h
@@ -10,7 +10,6 @@ class UNiagaraOverviewNode;
 class UNiagaraStackViewModel;
 class UNiagaraSystemSelectionViewModel;
 class FNiagaraEmitterHandleViewModel;
-class FAssetThumbnailPool;
 class FAssetThumbnail;
 class UNiagaraStackEntry;
 
@@ -64,7 +63,6 @@ private:
 	UNiagaraStackViewModel* StackViewModel;
 	UNiagaraSystemSelectionViewModel* OverviewSelectionViewModel;
 	TWeakPtr<FNiagaraEmitterHandleViewModel> EmitterHandleViewModelWeak;
-	TSharedPtr<FAssetThumbnailPool> ThumbnailPool;
 	/** Thumbnail widget containers */
 	TSharedPtr<SHorizontalBox> ThumbnailBar;
 	TArray<UNiagaraStackEntry*> PreviewStackEntries;
diff --git a/Engine/Plugins/Media/MediaFrameworkUtilities/Source/MediaFrameworkUtilitiesEditor/Private/MediaFrameworkUtilitiesPlacement.cpp b/Engine/Plugins/Media/MediaFrameworkUtilities/Source/MediaFrameworkUtilitiesEditor/Private/MediaFrameworkUtilitiesPlacement.cpp
--- a/Engine/Plugins/Media/MediaFrameworkUtilities/Source/MediaFrameworkUtilitiesEditor/Private/MediaFrameworkUtilitiesPlacement.cpp
+++ b/Engine/Plugins/Media/MediaFrameworkUtilities/Source/MediaFrameworkUtilitiesEditor/Private/MediaFrameworkUtilitiesPlacement.cpp
@@ -23,6 +23,7 @@
 
 #include "MediaBundle.h"
 #include "MediaSource.h"
+#include "ThumbnailRendering/ThumbnailManager.h"
 
 
 #define LOCTEXT_NAMESPACE "MediaFrameworkEditor"
@@ -98,9 +99,6 @@ public:
 			.OnDragDetected(this, &SMediaPlacementPalette::OnDraggingListViewWidget);
 
 		// Get the MediaSource thumbnail or the MediaBundle is not loaded
-		FLevelEditorModule& LevelEditorModule = FModuleManager::LoadModuleChecked<FLevelEditorModule>("LevelEditor");
-		TSharedPtr<FAssetThumbnailPool> ThumbnailPool = LevelEditorModule.GetFirstLevelEditor()->GetThumbnailPool();
-
 		FAssetData ThumbnailAssetData = MediaPlacement->MediaBundle;
 		if (MediaPlacement->MediaBundle.IsAssetLoaded())
 		{
@@ -110,7 +108,7 @@ public:
 				ThumbnailAssetData = FAssetData(MediaBundle->MediaSource);
 			}
 		}
-		TSharedPtr< FAssetThumbnail > Thumbnail = MakeShareable(new FAssetThumbnail(ThumbnailAssetData, 32, 32, ThumbnailPool));
+		TSharedPtr<FAssetThumbnail> Thumbnail = MakeShareable(new FAssetThumbnail(ThumbnailAssetData, 32, 32, UThumbnailManager::Get().GetSharedThumbnailPool()));
 
 		// Create the TableRow content
 		TSharedRef<SWidget> Content =
diff --git a/Engine/Source/Editor/ContentBrowserData/Private/ContentBrowserDataDragDropOp.cpp b/Engine/Source/Editor/ContentBrowserData/Private/ContentBrowserDataDragDropOp.cpp
--- a/Engine/Source/Editor/ContentBrowserData/Private/ContentBrowserDataDragDropOp.cpp
+++ b/Engine/Source/Editor/ContentBrowserData/Private/ContentBrowserDataDragDropOp.cpp
@@ -5,6 +5,7 @@
 #include "AssetThumbnail.h"
 #include "Editor.h"
 #include "Editor/EditorEngine.h"
+#include "ThumbnailRendering/ThumbnailManager.h"
 
 TSharedRef<FContentBrowserDataDragDropOp> FContentBrowserDataDragDropOp::New(TArrayView<const FContentBrowserItem> InDraggedItems)
 {
@@ -110,16 +111,12 @@ void FContentBrowserDataDragDropOp::InitThumbnail()
 {
 	if (DraggedFiles.Num() > 0 && ThumbnailSize > 0)
 	{
-		// Create a thumbnail pool to hold the single thumbnail rendered
-		ThumbnailPool = MakeShared<FAssetThumbnailPool>(1, /*InAreRealTileThumbnailsAllowed=*/false);
-
 		// Create the thumbnail handle
-		AssetThumbnail = MakeShared<FAssetThumbnail>(FAssetData(), ThumbnailSize, ThumbnailSize, ThumbnailPool);
+		AssetThumbnail = MakeShared<FAssetThumbnail>(FAssetData(), ThumbnailSize, ThumbnailSize, UThumbnailManager::Get().GetSharedThumbnailPool());
 		if (DraggedFiles[0].UpdateThumbnail(*AssetThumbnail))
 		{
 			// Request the texture then tick the pool once to render the thumbnail
 			AssetThumbnail->GetViewportRenderTargetTexture();
-			ThumbnailPool->Tick(0);
 		}
 		else
 		{
diff --git a/Engine/Source/Editor/LevelEditor/Private/LevelEditorCreateActorMenu.cpp b/Engine/Source/Editor/LevelEditor/Private/LevelEditorCreateActorMenu.cpp
--- a/Engine/Source/Editor/LevelEditor/Private/LevelEditorCreateActorMenu.cpp
+++ b/Engine/Source/Editor/LevelEditor/Private/LevelEditorCreateActorMenu.cpp
@@ -43,6 +43,7 @@
 #include "IPlacementModeModule.h"
 #include "Engine/TriggerBase.h"
 #include "LevelEditorMenuContext.h"
+#include "ThumbnailRendering/ThumbnailManager.h"
 
 class SMenuThumbnail : public SCompoundWidget
 {
@@ -59,10 +60,7 @@ public:
 	{
 		Asset = InAsset;
 
-		FLevelEditorModule& LevelEditorModule = FModuleManager::LoadModuleChecked<FLevelEditorModule>("LevelEditor");
-		TSharedPtr<FAssetThumbnailPool> ThumbnailPool = LevelEditorModule.GetFirstLevelEditor()->GetThumbnailPool();
-
-		Thumbnail = MakeShareable( new FAssetThumbnail( Asset, InArgs._Width, InArgs._Height, ThumbnailPool ) );
+		Thumbnail = MakeShareable( new FAssetThumbnail(Asset, InArgs._Width, InArgs._Height, UThumbnailManager::Get().GetSharedThumbnailPool()));
 
 		ChildSlot
 		[
diff --git a/Engine/Source/Editor/LevelEditor/Private/SLevelEditor.cpp b/Engine/Source/Editor/LevelEditor/Private/SLevelEditor.cpp
--- a/Engine/Source/Editor/LevelEditor/Private/SLevelEditor.cpp
+++ b/Engine/Source/Editor/LevelEditor/Private/SLevelEditor.cpp
@@ -63,6 +63,7 @@
 #include "Elements/Framework/TypedElementSelectionSet.h"
 #include "Elements/Actor/ActorElementLevelEditorSelectionCustomization.h"
 #include "Elements/Component/ComponentElementLevelEditorSelectionCustomization.h"
+#include "ThumbnailRendering/ThumbnailManager.h"
 
 #define LOCTEXT_NAMESPACE "SLevelEditor"
 
@@ -74,17 +75,10 @@ static const FName WorldBrowserCompositionTab("WorldBrowserComposition");
 static const FName WorldPartitionEditorTab("WorldBrowserPartitionEditor");
 
 
-namespace LevelEditorConstants
-{
-	/** The size of the thumbnail pool */
-	const int32 ThumbnailPoolSize = 32;
-}
-
 SLevelEditor::SLevelEditor()
 	: World(nullptr)
 	, bNeedsRefresh(false)
 {
-	ThumbnailPool = MakeShared<FAssetThumbnailPool>(LevelEditorConstants::ThumbnailPoolSize);
 }
 
 void SLevelEditor::BindCommands()
@@ -1753,9 +1747,9 @@ TSharedPtr<IAssetViewport> SLevelEditor::GetActiveViewportInterface()
 	return GetActiveViewport();
 }
 
-TSharedPtr< class FAssetThumbnailPool > SLevelEditor::GetThumbnailPool() const
+TSharedPtr<FAssetThumbnailPool> SLevelEditor::GetThumbnailPool() const
 {
-	return ThumbnailPool;
+	return UThumbnailManager::Get().GetSharedThumbnailPool();
 }
 
 void SLevelEditor::AppendCommands( const TSharedRef<FUICommandList>& InCommandsToAppend )
diff --git a/Engine/Source/Editor/LevelEditor/Private/SLevelEditor.h b/Engine/Source/Editor/LevelEditor/Private/SLevelEditor.h
--- a/Engine/Source/Editor/LevelEditor/Private/SLevelEditor.h
+++ b/Engine/Source/Editor/LevelEditor/Private/SLevelEditor.h
@@ -276,9 +276,6 @@ private:
 	// Holds the world settings details view.
 	TSharedPtr<IDetailsView> WorldSettingsView;
 
-	// The thumbnail pool used to display asset thumbnails
-	TSharedPtr<FAssetThumbnailPool> ThumbnailPool;
-
 	/** Transient editor viewport states - one for each view type. Key is "LayoutId[ELevelViewportType]", eg) "Viewport 1[0]" */
 	TMap<FString, FLevelViewportInfo> TransientEditorViews;
 
diff --git a/Engine/Source/Editor/LevelEditor/Public/ILevelEditor.h b/Engine/Source/Editor/LevelEditor/Public/ILevelEditor.h
--- a/Engine/Source/Editor/LevelEditor/Public/ILevelEditor.h
+++ b/Engine/Source/Editor/LevelEditor/Public/ILevelEditor.h
@@ -47,6 +47,7 @@ public:
 	virtual TSharedPtr<IAssetViewport> GetActiveViewportInterface() = 0;
 
 	/** Get the thumbnail pool used by this level editor */
+	UE_DEPRECATED(5.0, "GetThumbnailPool has been replaced by UThumbnailManager::Get().GetSharedThumbnailPool().")
 	virtual TSharedPtr< class FAssetThumbnailPool > GetThumbnailPool() const = 0;
 
 	/** Access the level editor's action command list */
diff --git a/Engine/Source/Editor/MaterialEditor/Private/SMaterialLayersFunctionsTree.cpp b/Engine/Source/Editor/MaterialEditor/Private/SMaterialLayersFunctionsTree.cpp
--- a/Engine/Source/Editor/MaterialEditor/Private/SMaterialLayersFunctionsTree.cpp
+++ b/Engine/Source/Editor/MaterialEditor/Private/SMaterialLayersFunctionsTree.cpp
@@ -41,6 +41,7 @@
 #include "Editor.h"
 #include "EditorFontGlyphs.h"
 #include "SResetToDefaultPropertyEditor.h"
+#include "ThumbnailRendering/ThumbnailManager.h"
 
 #define LOCTEXT_NAMESPACE "MaterialLayerCustomization"
 
@@ -1328,7 +1329,7 @@ FReply SMaterialLayersFunctionsInstanceTree::ToggleLayerVisibility(int32 Index)
 
 TSharedPtr<class FAssetThumbnailPool> SMaterialLayersFunctionsInstanceTree::GetTreeThumbnailPool()
 {
-	return Generator->GetGeneratedThumbnailPool();
+	return UThumbnailManager::Get().GetSharedThumbnailPool();
 }
 
 void SMaterialLayersFunctionsInstanceTree::CreateGroupsWidget()
@@ -2683,7 +2684,7 @@ void SMaterialLayersFunctionsMaterialTree::SetParentsExpansionState()
 
 TSharedPtr<class FAssetThumbnailPool> SMaterialLayersFunctionsMaterialTree::GetTreeThumbnailPool()
 {
-	return Wrapper->GetGenerator()->GetGeneratedThumbnailPool();
+	return UThumbnailManager::Get().GetSharedThumbnailPool();
 }
 
 void SMaterialLayersFunctionsMaterialTree::CreateGroupsWidget()
diff --git a/Engine/Source/Editor/MaterialEditor/Private/SMaterialParametersOverviewWidget.cpp b/Engine/Source/Editor/MaterialEditor/Private/SMaterialParametersOverviewWidget.cpp
--- a/Engine/Source/Editor/MaterialEditor/Private/SMaterialParametersOverviewWidget.cpp
+++ b/Engine/Source/Editor/MaterialEditor/Private/SMaterialParametersOverviewWidget.cpp
@@ -40,6 +40,7 @@
 #include "Widgets/Views/SExpanderArrow.h"
 #include "Widgets/Layout/SWidgetSwitcher.h"
 #include "MaterialEditor/DEditorRuntimeVirtualTextureParameterValue.h"
+#include "ThumbnailRendering/ThumbnailManager.h"
 
 
 
@@ -596,7 +597,7 @@ void SMaterialParametersOverviewTree::SetParentsExpansionState()
 
 TSharedPtr<class FAssetThumbnailPool> SMaterialParametersOverviewTree::GetTreeThumbnailPool()
 {
-	return GetOwner().Pin()->GetGenerator()->GetGeneratedThumbnailPool();
+	return UThumbnailManager::Get().GetSharedThumbnailPool();
 }
 
 void SMaterialParametersOverviewTree::CreateGroupsWidget()
diff --git a/Engine/Source/Editor/PlacementMode/Private/SPlacementModeTools.cpp b/Engine/Source/Editor/PlacementMode/Private/SPlacementModeTools.cpp
--- a/Engine/Source/Editor/PlacementMode/Private/SPlacementModeTools.cpp
+++ b/Engine/Source/Editor/PlacementMode/Private/SPlacementModeTools.cpp
@@ -19,6 +19,7 @@
 #include "Widgets/Input/SSearchBox.h"
 #include "ClassIconFinder.h"
 #include "Widgets/Docking/SDockTab.h"
+#include "ThumbnailRendering/ThumbnailManager.h"
 
 
 struct FSortPlaceableItems
@@ -100,8 +101,7 @@ public:
 	{
 		Asset = InAsset;
 
-		FLevelEditorModule& LevelEditorModule = FModuleManager::LoadModuleChecked<FLevelEditorModule>( "LevelEditor" );
-		TSharedPtr<FAssetThumbnailPool> ThumbnailPool = LevelEditorModule.GetFirstLevelEditor()->GetThumbnailPool();
+		TSharedPtr<FAssetThumbnailPool> ThumbnailPool = UThumbnailManager::Get().GetSharedThumbnailPool();
 
 		Thumbnail = MakeShareable(new FAssetThumbnail(Asset, InArgs._Width, InArgs._Height, ThumbnailPool));
 
@@ -279,7 +279,7 @@ FReply SPlacementAssetEntry::OnDragDetected(const FGeometry& MyGeometry, const F
 
 	if( MouseEvent.IsMouseButtonDown( EKeys::LeftMouseButton ) )
 	{
-		return FReply::Handled().BeginDragDrop( FContentBrowserDataDragDropOp::Legacy_New( MakeArrayView(&Item->AssetData, 1), TArrayView<const FString>(), Item->Factory ) );
+		return FReply::Handled().BeginDragDrop(FAssetDragDropOp::New(Item->AssetData, Item->Factory));
 	}
 	else
 	{
diff --git a/Engine/Source/Editor/PropertyEditor/Private/PropertyEditorModule.cpp b/Engine/Source/Editor/PropertyEditor/Private/PropertyEditorModule.cpp
--- a/Engine/Source/Editor/PropertyEditor/Private/PropertyEditorModule.cpp
+++ b/Engine/Source/Editor/PropertyEditor/Private/PropertyEditorModule.cpp
@@ -800,5 +800,5 @@ TSharedRef<class IStructureDetailsView> FPropertyEditorModule::CreateStructureDe
 
 TSharedRef<class IPropertyRowGenerator> FPropertyEditorModule::CreatePropertyRowGenerator(const struct FPropertyRowGeneratorArgs& InArgs)
 {
-	return MakeShared<FPropertyRowGenerator>(InArgs, GetThumbnailPool());
+	return MakeShared<FPropertyRowGenerator>(InArgs);
 }
diff --git a/Engine/Source/Editor/PropertyEditor/Private/PropertyRowGenerator.cpp b/Engine/Source/Editor/PropertyEditor/Private/PropertyRowGenerator.cpp
--- a/Engine/Source/Editor/PropertyEditor/Private/PropertyRowGenerator.cpp
+++ b/Engine/Source/Editor/PropertyEditor/Private/PropertyRowGenerator.cpp
@@ -15,6 +15,7 @@
 #include "IPropertyGenerationUtilities.h"
 #include "EditConditionParser.h"
 #include "UObject/StructOnScope.h"
+#include "ThumbnailRendering/ThumbnailManager.h"
 
 class FPropertyRowGeneratorUtilities : public IPropertyUtilities
 {
@@ -127,12 +128,19 @@ private:
 	FPropertyRowGenerator* Generator;
 };
 
+FPropertyRowGenerator::FPropertyRowGenerator(const FPropertyRowGeneratorArgs& InArgs)
+	: Args(InArgs)
+	, PropertyUtilities(new FPropertyRowGeneratorUtilities(*this))
+	, PropertyGenerationUtilities(new FPropertyRowGeneratorGenerationUtilities(*this))
+{
+}
+
 FPropertyRowGenerator::FPropertyRowGenerator(const FPropertyRowGeneratorArgs& InArgs, TSharedPtr<FAssetThumbnailPool> InThumbnailPool)
 	: Args(InArgs)
-	, ThumbnailPool(InThumbnailPool)
 	, PropertyUtilities(new FPropertyRowGeneratorUtilities(*this))
 	, PropertyGenerationUtilities(new FPropertyRowGeneratorGenerationUtilities(*this))
 {
+
 }
 
 FPropertyRowGenerator::~FPropertyRowGenerator()
@@ -415,7 +423,7 @@ void FPropertyRowGenerator::ForceRefresh()
 
 TSharedPtr<class FAssetThumbnailPool> FPropertyRowGenerator::GetThumbnailPool() const
 {
-	return ThumbnailPool;
+	return UThumbnailManager::Get().GetSharedThumbnailPool();
 }
 
 void FPropertyRowGenerator::PreSetObject(int32 NumNewObjects, bool bHasStructRoots)
diff --git a/Engine/Source/Editor/PropertyEditor/Private/PropertyRowGenerator.h b/Engine/Source/Editor/PropertyEditor/Private/PropertyRowGenerator.h
--- a/Engine/Source/Editor/PropertyEditor/Private/PropertyRowGenerator.h
+++ b/Engine/Source/Editor/PropertyEditor/Private/PropertyRowGenerator.h
@@ -74,7 +74,11 @@ typedef TArray<TSharedPtr<FComplexPropertyNode>> FRootPropertyNodeList;
 class FPropertyRowGenerator : public IPropertyRowGenerator, public FTickableEditorObject, public TSharedFromThis<FPropertyRowGenerator>
 {
 public:
+	FPropertyRowGenerator(const FPropertyRowGeneratorArgs& InArgs);
+
+	UE_DEPRECATED(5.0, "FPropertyRowGenerator which takes in a thumbnail pool parameter is no longer necessary.")
 	FPropertyRowGenerator(const FPropertyRowGeneratorArgs& InArgs, TSharedPtr<FAssetThumbnailPool> InThumbnailPool);
+
 	~FPropertyRowGenerator();
 
 	DECLARE_DERIVED_EVENT(FPropertyRowGenerator, IPropertyRowGenerator::FOnRowsRefreshed, FOnRowsRefreshed);
@@ -91,11 +95,6 @@ public:
 	virtual void RegisterInstancedCustomPropertyTypeLayout(FName PropertyTypeName, FOnGetPropertyTypeCustomizationInstance PropertyTypeLayoutDelegate, TSharedPtr<IPropertyTypeIdentifier> Identifier = nullptr) override;
 	virtual void UnregisterInstancedCustomPropertyLayout(UStruct* Class) override;
 	virtual void UnregisterInstancedCustomPropertyTypeLayout(FName PropertyTypeName, TSharedPtr<IPropertyTypeIdentifier> Identifier = nullptr) override;
-	virtual TSharedPtr<FAssetThumbnailPool> GetGeneratedThumbnailPool() override
-	{ 
-		return GetThumbnailPool();
-	};
-
 
 	/** FTickableEditorObject interface */
 	virtual ETickableTickType GetTickableTickType() const override { return ETickableTickType::Always; }
@@ -145,8 +144,6 @@ private:
 	FCustomPropertyTypeLayoutMap InstancedTypeToLayoutMap;
 	/** A mapping of classes to detail layout delegates, called when querying for custom detail layouts in this instance of the details view only*/
 	FCustomDetailLayoutMap InstancedClassToDetailLayoutMap;
-	/** Asset pool for rendering and managing asset thumbnails visible in this view */
-	TSharedPtr<FAssetThumbnailPool> ThumbnailPool;
 	/** Utility class for accessing commonly used helper methods from customizations */
 	TSharedRef<IPropertyUtilities> PropertyUtilities;
 	/** Utility class for accessing internal helper methods */
diff --git a/Engine/Source/Editor/PropertyEditor/Private/SDetailsViewBase.cpp b/Engine/Source/Editor/PropertyEditor/Private/SDetailsViewBase.cpp
--- a/Engine/Source/Editor/PropertyEditor/Private/SDetailsViewBase.cpp
+++ b/Engine/Source/Editor/PropertyEditor/Private/SDetailsViewBase.cpp
@@ -25,10 +25,10 @@
 #include "DetailLayoutHelpers.h"
 #include "EditConditionParser.h"
 #include "Editor.h"
+#include "ThumbnailRendering/ThumbnailManager.h"
 
 SDetailsViewBase::~SDetailsViewBase()
 {
-	ThumbnailPool.Reset();
 }
 
 void SDetailsViewBase::OnGetChildrenForDetailTree(TSharedRef<FDetailTreeNode> InTreeNode, TArray< TSharedRef<FDetailTreeNode> >& OutChildren)
@@ -514,13 +514,7 @@ void SDetailsViewBase::RefreshRootObjectVisibility()
 
 TSharedPtr<FAssetThumbnailPool> SDetailsViewBase::GetThumbnailPool() const
 {
-	if (!ThumbnailPool.IsValid())
-	{
-		// Create a thumbnail pool for the view if it doesnt exist.  This does not use resources of no thumbnails are used
-		ThumbnailPool = MakeShared<FAssetThumbnailPool>(50);
-	}
-
-	return ThumbnailPool;
+	return UThumbnailManager::Get().GetSharedThumbnailPool();
 }
 
 void SDetailsViewBase::NotifyFinishedChangingProperties(const FPropertyChangedEvent& PropertyChangedEvent)
diff --git a/Engine/Source/Editor/PropertyEditor/Private/SDetailsViewBase.h b/Engine/Source/Editor/PropertyEditor/Private/SDetailsViewBase.h
--- a/Engine/Source/Editor/PropertyEditor/Private/SDetailsViewBase.h
+++ b/Engine/Source/Editor/PropertyEditor/Private/SDetailsViewBase.h
@@ -410,8 +410,6 @@ protected:
 	TSharedPtr<IPropertyGenerationUtilities> PropertyGenerationUtilities;
 	/** The name area which is not recreated when selection changes */
 	TSharedPtr<class SDetailNameArea> NameArea;
-	/** Asset pool for rendering and managing asset thumbnails visible in this view*/
-	mutable TSharedPtr<class FAssetThumbnailPool> ThumbnailPool;
 	/** The current filter */
 	FDetailFilter CurrentFilter;
 	/** Delegate called to get generic details not specific to an object being viewed */
diff --git a/Engine/Source/Editor/PropertyEditor/Public/IPropertyRowGenerator.h b/Engine/Source/Editor/PropertyEditor/Public/IPropertyRowGenerator.h
--- a/Engine/Source/Editor/PropertyEditor/Public/IPropertyRowGenerator.h
+++ b/Engine/Source/Editor/PropertyEditor/Public/IPropertyRowGenerator.h
@@ -100,6 +100,5 @@ public:
 	virtual void UnregisterInstancedCustomPropertyLayout(UStruct* Class) = 0;
 	virtual void UnregisterInstancedCustomPropertyTypeLayout(FName PropertyTypeName, TSharedPtr<IPropertyTypeIdentifier> Identifier = nullptr) = 0;
 
-	virtual TSharedPtr<class FAssetThumbnailPool> GetGeneratedThumbnailPool() = 0;
 	virtual FOnFinishedChangingProperties& OnFinishedChangingProperties() = 0;
 };
\ No newline at end of file
diff --git a/Engine/Source/Editor/StaticMeshEditor/Private/StaticMeshEditorTools.h b/Engine/Source/Editor/StaticMeshEditor/Private/StaticMeshEditorTools.h
--- a/Engine/Source/Editor/StaticMeshEditor/Private/StaticMeshEditorTools.h
+++ b/Engine/Source/Editor/StaticMeshEditor/Private/StaticMeshEditorTools.h
@@ -604,9 +604,6 @@ private:
 	/** The Static Mesh Editor this tool is associated with. */
 	FStaticMeshEditor& StaticMeshEditor;
 
-	/** Pool for material thumbnails. */
-	TSharedPtr<FAssetThumbnailPool> ThumbnailPool;
-
 	/** LOD group options. */
 	TArray<FName> LODGroupNames;
 	TArray<TSharedPtr<FString> > LODGroupOptions;
diff --git a/Engine/Source/Editor/UnrealEd/Private/Dialogs/SDeleteAssetsDialog.cpp b/Engine/Source/Editor/UnrealEd/Private/Dialogs/SDeleteAssetsDialog.cpp
--- a/Engine/Source/Editor/UnrealEd/Private/Dialogs/SDeleteAssetsDialog.cpp
+++ b/Engine/Source/Editor/UnrealEd/Private/Dialogs/SDeleteAssetsDialog.cpp
@@ -108,8 +108,6 @@ private:
 SDeleteAssetsDialog::~SDeleteAssetsDialog()
 {
 	DeleteModel->OnStateChanged().RemoveAll( this );
-	// Release all rendering resources being held onto
-	AssetThumbnailPool.Reset();
 }
 
 void SDeleteAssetsDialog::Construct( const FArguments& InArgs, TSharedRef<FAssetDeleteModel> InDeleteModel )
@@ -123,8 +121,6 @@ void SDeleteAssetsDialog::Construct( const FArguments& InArgs, TSharedRef<FAsset
 	// Save off the attributes
 	ParentWindow = InArgs._ParentWindow;
 
-	AssetThumbnailPool = MakeShareable( new FAssetThumbnailPool( 1, false ) );
-
 	ReferencerCommands = TSharedPtr< FUICommandList >(new FUICommandList);
 
 	ReferencerCommands->MapAction(FGenericCommands::Get().Delete, FUIAction(
@@ -638,7 +634,7 @@ TOptional< float > SDeleteAssetsDialog::ScanningProgressFraction() const
 
 TSharedRef<SWidget> SDeleteAssetsDialog::CreateThumbnailWidget()
 {
-	ConsolidationAssetThumbnail = MakeShareable( new FAssetThumbnail( NULL, 40, 40, AssetThumbnailPool ) );
+	ConsolidationAssetThumbnail = MakeShareable( new FAssetThumbnail( NULL, 40, 40, UThumbnailManager::Get().GetSharedThumbnailPool()) );
 
 	return SNew( SBox )
 		.WidthOverride( 40.0f )
diff --git a/Engine/Source/Editor/UnrealEd/Private/DragAndDrop/AssetDragDropOp.cpp b/Engine/Source/Editor/UnrealEd/Private/DragAndDrop/AssetDragDropOp.cpp
--- a/Engine/Source/Editor/UnrealEd/Private/DragAndDrop/AssetDragDropOp.cpp
+++ b/Engine/Source/Editor/UnrealEd/Private/DragAndDrop/AssetDragDropOp.cpp
@@ -8,6 +8,7 @@
 #include "Widgets/Layout/SBox.h"
 #include "AssetThumbnail.h"
 #include "ClassIconFinder.h"
+#include "ThumbnailRendering/ThumbnailManager.h"
 
 TSharedRef<FAssetDragDropOp> FAssetDragDropOp::New(const FAssetData& InAssetData, UActorFactory* ActorFactory)
 {
@@ -45,7 +46,6 @@ TSharedRef<FAssetDragDropOp> FAssetDragDropOp::New(TArray<FAssetData> InAssetDat
 
 FAssetDragDropOp::~FAssetDragDropOp()
 {
-	ThumbnailPool.Reset();
 }
 
 TSharedPtr<SWidget> FAssetDragDropOp::GetDefaultDecorator() const
@@ -221,14 +221,13 @@ void FAssetDragDropOp::InitThumbnail()
 	if (AssetData.Num() > 0 && ThumbnailSize > 0)
 	{
 		// Create a thumbnail pool to hold the single thumbnail rendered
-		ThumbnailPool = MakeShared<FAssetThumbnailPool>(1, /*InAreRealTileThumbnailsAllowed=*/false);
+		//ThumbnailPool = MakeShared<FAssetThumbnailPool>(1, /*InAreRealTileThumbnailsAllowed=*/false);
 
 		// Create the thumbnail handle
-		AssetThumbnail = MakeShared<FAssetThumbnail>(AssetData[0], ThumbnailSize, ThumbnailSize, ThumbnailPool);
+		AssetThumbnail = MakeShared<FAssetThumbnail>(AssetData[0], ThumbnailSize, ThumbnailSize, UThumbnailManager::Get().GetSharedThumbnailPool());
 
 		// Request the texture then tick the pool once to render the thumbnail
 		AssetThumbnail->GetViewportRenderTargetTexture();
-		ThumbnailPool->Tick(0);
 	}
 }
 
diff --git a/Engine/Source/Editor/UnrealEd/Private/Toolkits/SGlobalTabSwitchingDialog.cpp b/Engine/Source/Editor/UnrealEd/Private/Toolkits/SGlobalTabSwitchingDialog.cpp
--- a/Engine/Source/Editor/UnrealEd/Private/Toolkits/SGlobalTabSwitchingDialog.cpp
+++ b/Engine/Source/Editor/UnrealEd/Private/Toolkits/SGlobalTabSwitchingDialog.cpp
@@ -25,6 +25,7 @@
 #include "Mac/MacApplication.h"
 #endif
 #include "Subsystems/AssetEditorSubsystem.h"
+#include "ThumbnailRendering/ThumbnailManager.h"
 
 #define LOCTEXT_NAMESPACE "SGlobalTabSwitchingDialog"
 
@@ -261,7 +262,7 @@ void SGlobalTabSwitchingDialog::OnMainTabListSelectionChanged(FTabListItemPtr In
 	{
 		FTabListItemPtr SelectedItem = SelectedItems[0];
 
-		NewTopContents = SelectedItem->CreateWidget(AssetThumbnailPool);
+		NewTopContents = SelectedItem->CreateWidget(UThumbnailManager::Get().GetSharedThumbnailPool());
 
 		NewBottomContents =
 			SNew(SHorizontalBox)
@@ -387,8 +388,6 @@ void SGlobalTabSwitchingDialog::Construct(const FArguments& InArgs, FVector2D In
 
 	TriggerChord = InTriggerChord;
 
-	AssetThumbnailPool = MakeShareable(new FAssetThumbnailPool(128));
-
 	// Populate the list with open asset editors
 	TArray<UObject*> OpenAssetList = GEditor->GetEditorSubsystem<UAssetEditorSubsystem>()->GetAllEditedAssets();
 	for (UObject* OpenAsset : OpenAssetList)
@@ -551,7 +550,7 @@ FReply SGlobalTabSwitchingDialog::OnPreviewKeyDown(const FGeometry& MyGeometry,
 
 TSharedRef<ITableRow> SGlobalTabSwitchingDialog::OnGenerateTabSwitchListItemWidget(FTabListItemPtr InItem, const TSharedRef<STableViewBase>& OwnerTable)
 {
-	return SNew(STableRow<FTabListItemPtr>, OwnerTable)[InItem->CreateWidget(AssetThumbnailPool)];
+	return SNew(STableRow<FTabListItemPtr>, OwnerTable)[InItem->CreateWidget(UThumbnailManager::Get().GetSharedThumbnailPool())];
 }
 
 //////////////////////////////////////////////////////////////////////////
diff --git a/Engine/Source/Editor/UnrealEd/Private/Toolkits/SGlobalTabSwitchingDialog.h b/Engine/Source/Editor/UnrealEd/Private/Toolkits/SGlobalTabSwitchingDialog.h
--- a/Engine/Source/Editor/UnrealEd/Private/Toolkits/SGlobalTabSwitchingDialog.h
+++ b/Engine/Source/Editor/UnrealEd/Private/Toolkits/SGlobalTabSwitchingDialog.h
@@ -60,9 +60,6 @@ private:
 	// The chord that triggered the dialog (so we can handle the correct Tab/`/etc... key repeat, and dismiss on the correct control/command modifier release
 	FInputChord TriggerChord;
 
-	// Pool for maintaining and rendering thumbnails
-	TSharedPtr<class FAssetThumbnailPool> AssetThumbnailPool;
-
 	// The array of 'document' items
 	TArray<FTabListItemPtr> MainTabsListDataSource;
 
diff --git a/Engine/Source/Editor/UnrealEd/Public/Dialogs/SDeleteAssetsDialog.h b/Engine/Source/Editor/UnrealEd/Public/Dialogs/SDeleteAssetsDialog.h
--- a/Engine/Source/Editor/UnrealEd/Public/Dialogs/SDeleteAssetsDialog.h
+++ b/Engine/Source/Editor/UnrealEd/Public/Dialogs/SDeleteAssetsDialog.h
@@ -160,9 +160,6 @@ private:
 	/** The selected asset we're going to consolidate the would be deleted assets into. */
 	FAssetData ConsolidationAsset;
 
-	/** The thumbnail pool used by the replace references asset picker. */
-	TSharedPtr<FAssetThumbnailPool> AssetThumbnailPool;
-
 	/** The consolidation asset thumbnail */
 	TSharedPtr<class FAssetThumbnail> ConsolidationAssetThumbnail;
 
diff --git a/Engine/Source/Editor/UnrealEd/Public/DragAndDrop/AssetDragDropOp.h b/Engine/Source/Editor/UnrealEd/Public/DragAndDrop/AssetDragDropOp.h
--- a/Engine/Source/Editor/UnrealEd/Public/DragAndDrop/AssetDragDropOp.h
+++ b/Engine/Source/Editor/UnrealEd/Public/DragAndDrop/AssetDragDropOp.h
@@ -9,7 +9,6 @@
 #include "DragAndDrop/DecoratedDragDropOp.h"
 
 class FAssetThumbnail;
-class FAssetThumbnailPool;
 class UActorFactory;
 
 class UNREALED_API FAssetDragDropOp : public FDecoratedDragDropOp
@@ -85,9 +84,6 @@ private:
 	TWeakObjectPtr<UActorFactory> ActorFactory;
 
 protected:
-	/** Pool for maintaining and rendering thumbnails */
-	TSharedPtr<FAssetThumbnailPool> ThumbnailPool;
-
 	/** Handle to the thumbnail resource */
 	TSharedPtr<FAssetThumbnail> AssetThumbnail;
 
