Details rows now preserve space on the right for the scrollbar.

CustomRow() widgets can now be disabled from customizations.

Favoriting command is now hidden when favoriting is not enabled in a details panel.

Removed unnecessary code for calculating a scroll offset when the favorites category was added or removed since the Favorites category is never removed.

diff --git a/Engine/Source/Editor/PropertyEditor/Private/DetailAdvancedDropdownNode.cpp b/Engine/Source/Editor/PropertyEditor/Private/DetailAdvancedDropdownNode.cpp
--- a/Engine/Source/Editor/PropertyEditor/Private/DetailAdvancedDropdownNode.cpp
+++ b/Engine/Source/Editor/PropertyEditor/Private/DetailAdvancedDropdownNode.cpp
@@ -72,30 +72,50 @@ public:
 			ContentWidget = SNew(SSpacer);
 		}
 
+		const FSlateBrush* ScrollBarWellBrush = FAppStyle::Get().GetBrush("DetailsView.GridLine");
+		const FSlateBrush* BackgroundBrush = FAppStyle::Get().GetBrush("DetailsView.AdvancedDropdownBorder");
+
 		TWeakPtr<STableViewBase> OwnerTableViewWeak = InOwnerTableView;
-		auto GetGridLinePadding = [this, bIsTopNode, OwnerTableViewWeak]()
+		auto GetScrollbarWellBrush = [OwnerTableViewWeak, ScrollBarWellBrush, BackgroundBrush]()
 		{
-			FMargin Margin = GetRowScrollBarPadding(OwnerTableViewWeak);
-			if (bIsTopNode)
-			{
-				Margin.Bottom = 0;
-			}
-			return Margin;
+			return SDetailTableRowBase::IsScrollBarVisible(OwnerTableViewWeak) ? ScrollBarWellBrush : BackgroundBrush;
 		};
 
-		FMargin BorderPadding = bIsTopNode ? FMargin(0, 0, 0, 2) : FMargin(0, 3, 0, 2);
-
 		ChildSlot
 		[
 			SNew(SBorder)
-			.BorderImage(FAppStyle::Get().GetBrush( "DetailsView.GridLine"))
-			.Padding_Lambda(GetGridLinePadding)
+			.BorderImage(FAppStyle::Get().GetBrush("DetailsView.GridLine"))
+			.Padding(bIsTopNode ? FMargin(0) : FMargin(0, 0, 0, 1))
 			[
-				SNew(SBorder)
-				.BorderImage(FEditorStyle::GetBrush("DetailsView.AdvancedDropdownBorder"))
-				.Padding(BorderPadding)
+				SNew(SHorizontalBox)
+				+ SHorizontalBox::Slot()
+				.VAlign(VAlign_Fill)
+				.HAlign(HAlign_Fill)
+				[
+					SNew(SBorder)
+					.BorderImage(FAppStyle::Get().GetBrush("DetailsView.AdvancedDropdownBorder"))
+					.Padding(bIsTopNode ? FMargin(0, 0, 0, 2) : FMargin(0, 3, 0, 2))
+					[
+						ContentWidget.ToSharedRef()
+					]
+				]
+				+ SHorizontalBox::Slot()
+				.VAlign(VAlign_Fill)
+				.HAlign(HAlign_Right)
+				.AutoWidth()
 				[
-					ContentWidget.ToSharedRef()
+					SNew(SBorder)
+					.BorderImage_Lambda(GetScrollbarWellBrush)
+					.Padding(FMargin(0, 0, SDetailTableRowBase::ScrollBarPadding, 0))
+					[
+						SNew(SBox)
+						.VAlign(VAlign_Fill)
+						[
+							SNew(SImage)
+							.Image(FEditorStyle::GetBrush("DetailsView.AdvancedDropdownBorder.Open"))
+							.Visibility(bIsTopNode ? EVisibility::Visible : EVisibility::Collapsed)
+						]
+					]
 				]
 			]
 		];
diff --git a/Engine/Source/Editor/PropertyEditor/Private/DetailCategoryGroupNode.cpp b/Engine/Source/Editor/PropertyEditor/Private/DetailCategoryGroupNode.cpp
--- a/Engine/Source/Editor/PropertyEditor/Private/DetailCategoryGroupNode.cpp
+++ b/Engine/Source/Editor/PropertyEditor/Private/DetailCategoryGroupNode.cpp
@@ -5,6 +5,7 @@
 #include "PropertyEditorConstants.h"
 #include "SDetailExpanderArrow.h"
 #include "SDetailRowIndent.h"
+#include "Styling/StyleColors.h"
 
 void SDetailCategoryTableRow::Construct( const FArguments& InArgs, TSharedRef<FDetailTreeNode> InOwnerTreeNode, const TSharedRef<STableViewBase>& InOwnerTableView )
 {
@@ -57,12 +58,25 @@ void SDetailCategoryTableRow::Construct( const FArguments& InArgs, TSharedRef<FD
 	const float EditConditionWidgetWidth = 33; // this is the width of the edit condition widget displayed on the left in SDetailSingleItemRow
 
 	TWeakPtr<STableViewBase> OwnerTableViewWeak = InOwnerTableView;
-	
+	auto GetScrollbarWellBrush = [this, OwnerTableViewWeak]()
+	{
+		return SDetailTableRowBase::IsScrollBarVisible(OwnerTableViewWeak) ?
+			FAppStyle::Get().GetBrush("DetailsView.GridLine") : 
+			this->GetBackgroundImage();
+	};
+
+	auto GetScrollbarWellTint = [this, OwnerTableViewWeak]()
+	{
+		return SDetailTableRowBase::IsScrollBarVisible(OwnerTableViewWeak) ?
+			FSlateColor(EStyleColor::White) : 
+			this->GetInnerBackgroundColor();
+	};
+
 	this->ChildSlot
 	[
 		SNew( SBorder )
 		.BorderImage(FAppStyle::Get().GetBrush("DetailsView.GridLine"))
-		.Padding(this, &SDetailTableRowBase::GetRowScrollBarPadding, OwnerTableViewWeak)
+		.Padding(FMargin(0, 0, 0, 1))
 		[
 			SNew(SHorizontalBox)
 			+ SHorizontalBox::Slot()
@@ -75,20 +89,30 @@ void SDetailCategoryTableRow::Construct( const FArguments& InArgs, TSharedRef<FD
 				.Padding(0)
 				[
 					SNew(SSpacer)
-					.Size(bIsInnerCategory ? FVector2D(EditConditionWidgetWidth,0) : FVector2D(0,0))
+					.Size(bIsInnerCategory ? FVector2D(EditConditionWidgetWidth, 0) : FVector2D(0,0))
 				]
 			]
 			+ SHorizontalBox::Slot()
 			.HAlign(HAlign_Fill)
 			[
-				SNew( SBorder )
-				.BorderImage( this, &SDetailCategoryTableRow::GetBackgroundImage )
+				SNew(SBorder)
+				.BorderImage(this, &SDetailCategoryTableRow::GetBackgroundImage)
 				.BorderBackgroundColor(this, &SDetailCategoryTableRow::GetInnerBackgroundColor)
 				.Padding(0)
 				[
 					HeaderBox
 				]
 			]
+			+ SHorizontalBox::Slot()
+			.HAlign(HAlign_Right)
+			.VAlign(VAlign_Fill)
+			.AutoWidth()
+			[
+				SNew(SBorder)
+				.BorderImage_Lambda(GetScrollbarWellBrush)
+				.BorderBackgroundColor_Lambda(GetScrollbarWellTint)
+				.Padding(FMargin(0, 0, SDetailTableRowBase::ScrollBarPadding, 0))
+			]
 		]
 	];
 
diff --git a/Engine/Source/Editor/PropertyEditor/Private/DetailMultiTopLevelObjectRootNode.cpp b/Engine/Source/Editor/PropertyEditor/Private/DetailMultiTopLevelObjectRootNode.cpp
--- a/Engine/Source/Editor/PropertyEditor/Private/DetailMultiTopLevelObjectRootNode.cpp
+++ b/Engine/Source/Editor/PropertyEditor/Private/DetailMultiTopLevelObjectRootNode.cpp
@@ -27,23 +27,20 @@ void SDetailMultiTopLevelObjectTableRow::SetContent(TSharedRef<SWidget> InConten
 	{
 		ChildSlot
 		[
-			SNew(SBox)
-			.Padding(this, &SDetailTableRowBase::GetRowScrollBarPadding, OwnerTableViewWeak)
+			SNew(SHorizontalBox)
+			+ SHorizontalBox::Slot()
+			.VAlign(VAlign_Center)
+			.Padding(2.0f, 2.0f, 2.0f, 2.0f)
+			.AutoWidth()
 			[
-				SNew(SHorizontalBox)
-				+ SHorizontalBox::Slot()
-				.VAlign(VAlign_Center)
-				.Padding(2.0f, 2.0f, 2.0f, 2.0f)
-				.AutoWidth()
-				[
-					SNew(SExpanderArrow, SharedThis(this))
-					.Visibility(ExpansionArrowUsage == EExpansionArrowUsage::Default ? EVisibility::Visible : EVisibility::Collapsed)
-				]
-				+ SHorizontalBox::Slot()
-				.Expose(ContentSlot)
-				[
-					InContent
-				]
+				SNew(SExpanderArrow, SharedThis(this))
+				.Visibility(ExpansionArrowUsage == EExpansionArrowUsage::Default ? EVisibility::Visible : EVisibility::Collapsed)
+			]
+			+ SHorizontalBox::Slot()
+			.Expose(ContentSlot)
+			.Padding(FMargin(0, 0, 0, 16))
+			[
+				InContent
 			]
 		];
 	}
@@ -56,18 +53,6 @@ void SDetailMultiTopLevelObjectTableRow::SetContent(TSharedRef<SWidget> InConten
 	}
 }
 
-const FSlateBrush* SDetailMultiTopLevelObjectTableRow::GetBackgroundImage() const
-{
-	if (IsHovered())
-	{
-		return IsItemExpanded() ? FEditorStyle::GetBrush("DetailsView.CategoryTop_Hovered") : FEditorStyle::GetBrush("DetailsView.CollapsedCategory_Hovered");
-	}
-	else
-	{
-		return IsItemExpanded() ? FEditorStyle::GetBrush("DetailsView.CategoryTop") : FEditorStyle::GetBrush("DetailsView.CollapsedCategory");
-	}
-}
-
 FReply SDetailMultiTopLevelObjectTableRow::OnMouseButtonDown( const FGeometry& MyGeometry, const FPointerEvent& MouseEvent )
 {
 	if (ExpansionArrowUsage != EExpansionArrowUsage::None && MouseEvent.GetEffectingButton() == EKeys::LeftMouseButton)
diff --git a/Engine/Source/Editor/PropertyEditor/Private/SDetailSingleItemRow.cpp b/Engine/Source/Editor/PropertyEditor/Private/SDetailSingleItemRow.cpp
--- a/Engine/Source/Editor/PropertyEditor/Private/SDetailSingleItemRow.cpp
+++ b/Engine/Source/Editor/PropertyEditor/Private/SDetailSingleItemRow.cpp
@@ -21,6 +21,7 @@
 #include "Modules/ModuleInterface.h"
 #include "Modules/ModuleManager.h"
 #include "Settings/EditorExperimentalSettings.h"
+#include "Styling/StyleColors.h"
 #include "Widgets/Images/SImage.h"
 #include "Widgets/Input/SButton.h"
 #include "Widgets/Input/SComboButton.h"
@@ -272,10 +273,10 @@ void SDetailSingleItemRow::Construct( const FArguments& InArgs, FDetailLayoutCus
 				];
 
 			TAttribute<bool> IsEnabledAttribute;
-			if (WidgetRow.IsEnabledAttr.IsBound())
+			if (WidgetRow.IsEnabledAttr.IsSet() || WidgetRow.IsEnabledAttr.IsBound())
 			{
-				TAttribute<bool> PropertyEnabledAttr = InOwnerTreeNode->IsPropertyEditingEnabled();
 				TAttribute<bool> RowEnabledAttr = WidgetRow.IsEnabledAttr;
+				TAttribute<bool> PropertyEnabledAttr = InOwnerTreeNode->IsPropertyEditingEnabled();
 				IsEnabledAttribute = TAttribute<bool>::Create(
 					[RowEnabledAttr, PropertyEnabledAttr]()
 					{
@@ -533,20 +534,48 @@ void SDetailSingleItemRow::Construct( const FArguments& InArgs, FDetailLayoutCus
 	}
 
 	TWeakPtr<STableViewBase> OwnerTableViewWeak = InOwnerTableView;
-	
+	auto GetScrollbarWellBrush = [this, OwnerTableViewWeak]()
+	{
+		return SDetailTableRowBase::IsScrollBarVisible(OwnerTableViewWeak) ?
+			FAppStyle::Get().GetBrush("DetailsView.GridLine") : 
+			FAppStyle::Get().GetBrush("DetailsView.CategoryMiddle");
+	};
+
+	auto GetScrollbarWellTint = [this, OwnerTableViewWeak]()
+	{
+		return SDetailTableRowBase::IsScrollBarVisible(OwnerTableViewWeak) ?
+			FSlateColor(EStyleColor::White) : 
+			this->GetInnerBackgroundColor();
+	};
+
 	this->ChildSlot
 	[
 		SNew( SBorder )
 		.BorderImage(FAppStyle::Get().GetBrush( "DetailsView.GridLine"))
-		.Padding(this, &SDetailTableRowBase::GetRowScrollBarPadding, OwnerTableViewWeak)
+		.Padding(FMargin(0,0,0,1))
 		[
-			SNew( SBorder )
-			.BorderImage(FAppStyle::Get().GetBrush("DetailsView.CategoryMiddle"))
-			.BorderBackgroundColor(this, &SDetailSingleItemRow::GetOuterBackgroundColor)
-			.Padding(0)
-			.Clipping(EWidgetClipping::ClipToBounds)
+			SNew( SHorizontalBox )
+			+ SHorizontalBox::Slot()
+			.HAlign(HAlign_Fill)
+			[
+				SNew( SBorder )
+				.BorderImage(FAppStyle::Get().GetBrush("DetailsView.CategoryMiddle"))
+				.BorderBackgroundColor(this, &SDetailSingleItemRow::GetOuterBackgroundColor)
+				.Padding(0)
+				.Clipping(EWidgetClipping::ClipToBounds)
+				[
+					Widget
+				]
+			]
+			+ SHorizontalBox::Slot()
+			.HAlign(HAlign_Right)
+			.VAlign(VAlign_Fill)
+			.AutoWidth()
 			[
-				Widget
+				SNew(SBorder)
+				.BorderImage_Lambda(GetScrollbarWellBrush)
+				.BorderBackgroundColor_Lambda(GetScrollbarWellTint)
+				.Padding(FMargin(0, 0, SDetailTableRowBase::ScrollBarPadding, 0))
 			]
 		]
 	];
@@ -682,30 +711,33 @@ bool SDetailSingleItemRow::OnContextMenuOpening(FMenuBuilder& MenuBuilder)
 			PasteAction);
 	}
 
-	FUIAction FavoriteAction;
-	FavoriteAction.ExecuteAction = FExecuteAction::CreateSP(this, &SDetailSingleItemRow::OnFavoriteMenuToggle);
-	FavoriteAction.CanExecuteAction = FCanExecuteAction::CreateLambda([this]()
-		{
-			return Customization->GetPropertyNode().IsValid() && Customization->GetPropertyNode()->CanDisplayFavorite();
-		});
+	if (OwnerTreeNode.Pin()->GetDetailsView()->IsFavoritingEnabled())
+	{
+		FUIAction FavoriteAction;
+		FavoriteAction.ExecuteAction = FExecuteAction::CreateSP(this, &SDetailSingleItemRow::OnFavoriteMenuToggle);
+		FavoriteAction.CanExecuteAction = FCanExecuteAction::CreateLambda([this]()
+			{
+				return Customization->GetPropertyNode().IsValid() && Customization->GetPropertyNode()->CanDisplayFavorite();
+			});
 
-	FText FavoriteText = NSLOCTEXT("PropertyView", "FavoriteProperty", "Add to Favorites");
-	FText FavoriteTooltipText = NSLOCTEXT("PropertyView", "FavoriteProperty_ToolTip", "Add this property to your favorites.");
-	FName FavoriteIcon = "DetailsView.PropertyIsFavorite";
+		FText FavoriteText = NSLOCTEXT("PropertyView", "FavoriteProperty", "Add to Favorites");
+		FText FavoriteTooltipText = NSLOCTEXT("PropertyView", "FavoriteProperty_ToolTip", "Add this property to your favorites.");
+		FName FavoriteIcon = "DetailsView.PropertyIsFavorite";
 
-	bool IsFavorite = Customization->GetPropertyNode().IsValid() && Customization->GetPropertyNode()->IsFavorite();
-	if (IsFavorite)
-	{
-		FavoriteText = NSLOCTEXT("PropertyView", "RemoveFavoriteProperty", "Remove from Favorites");
-		FavoriteTooltipText = NSLOCTEXT("PropertyView", "RemoveFavoriteProperty_ToolTip", "Remove this property from your favorites.");
-		FavoriteIcon = "DetailsView.PropertyIsNotFavorite";
-	}
+		bool IsFavorite = Customization->GetPropertyNode().IsValid() && Customization->GetPropertyNode()->IsFavorite();
+		if (IsFavorite)
+		{
+			FavoriteText = NSLOCTEXT("PropertyView", "RemoveFavoriteProperty", "Remove from Favorites");
+			FavoriteTooltipText = NSLOCTEXT("PropertyView", "RemoveFavoriteProperty_ToolTip", "Remove this property from your favorites.");
+			FavoriteIcon = "DetailsView.PropertyIsNotFavorite";
+		}
 
-	MenuBuilder.AddMenuEntry(
-		FavoriteText,
-		FavoriteTooltipText,
-		FSlateIcon(FEditorStyle::Get().GetStyleSetName(), FavoriteIcon),
-		FavoriteAction);
+		MenuBuilder.AddMenuEntry(
+			FavoriteText,
+			FavoriteTooltipText,
+			FSlateIcon(FEditorStyle::Get().GetStyleSetName(), FavoriteIcon),
+			FavoriteAction);
+	}
 
 	if (WidgetRow.CustomMenuItems.Num() > 0)
 	{
@@ -844,10 +876,7 @@ void SDetailSingleItemRow::OnFavoriteMenuToggle()
 		return;
 	}
 
-	//////////////////////////////////////////////////////////////////////////
-	// Calculate properly the scrolling offset (by item) to make sure the mouse stay over the same property
-
-	// Get the node item number in case it is expand we have to recursively count all childrens
+	// Calculate the scrolling offset (by item) to make sure the mouse stay over the same property
 	int32 ExpandSize = 0;
 	if (OwnerTreeNodePinned->ShouldBeExpanded())
 	{
diff --git a/Engine/Source/Editor/PropertyEditor/Private/SDetailTableRowBase.cpp b/Engine/Source/Editor/PropertyEditor/Private/SDetailTableRowBase.cpp
--- a/Engine/Source/Editor/PropertyEditor/Private/SDetailTableRowBase.cpp
+++ b/Engine/Source/Editor/PropertyEditor/Private/SDetailTableRowBase.cpp
@@ -2,6 +2,8 @@
 
 #include "SDetailTableRowBase.h"
 
+const float SDetailTableRowBase::ScrollBarPadding = 16.0f;
+
 int32 SDetailTableRowBase::GetIndentLevelForBackgroundColor() const
 {
 	int32 IndentLevel = 0; 
@@ -23,16 +25,12 @@ int32 SDetailTableRowBase::GetIndentLevelForBackgroundColor() const
 	return FMath::Max(0, IndentLevel);
 }
 
-FMargin SDetailTableRowBase::GetRowScrollBarPadding(TWeakPtr<STableViewBase> OwnerTableViewWeak) const
+bool SDetailTableRowBase::IsScrollBarVisible(TWeakPtr<STableViewBase> OwnerTableViewWeak)
 {
 	TSharedPtr<STableViewBase> OwnerTableView = OwnerTableViewWeak.Pin();
 	if (OwnerTableView.IsValid())
 	{
-		if (OwnerTableView->GetScrollbarVisibility() == EVisibility::Visible)
-		{
-			const float ScrollbarPaddingSize = 16.0f;
-			return FMargin(0, 0, ScrollbarPaddingSize, 1);
-		}
+		return OwnerTableView->GetScrollbarVisibility() == EVisibility::Visible;
 	}
-	return FMargin(0, 0, 0, 1);
+	return false;
 }
diff --git a/Engine/Source/Editor/PropertyEditor/Private/SDetailTableRowBase.h b/Engine/Source/Editor/PropertyEditor/Private/SDetailTableRowBase.h
--- a/Engine/Source/Editor/PropertyEditor/Private/SDetailTableRowBase.h
+++ b/Engine/Source/Editor/PropertyEditor/Private/SDetailTableRowBase.h
@@ -61,7 +61,9 @@ public:
 	}
 
 	int32 GetIndentLevelForBackgroundColor() const;
-	FMargin GetRowScrollBarPadding(TWeakPtr<STableViewBase> OwnerTableView) const;
+
+	static bool IsScrollBarVisible(TWeakPtr<STableViewBase> OwnerTableViewWeak);
+	static const float ScrollBarPadding;
 
 protected:
 	/**
