Remove reset to default from SDefaultPropertyButtonWidgets

diff --git a/Engine/Source/Editor/PropertyEditor/Private/PropertyHandleImpl.cpp b/Engine/Source/Editor/PropertyEditor/Private/PropertyHandleImpl.cpp
--- a/Engine/Source/Editor/PropertyEditor/Private/PropertyHandleImpl.cpp
+++ b/Engine/Source/Editor/PropertyEditor/Private/PropertyHandleImpl.cpp
@@ -2315,15 +2315,6 @@ class SDefaultPropertyButtonWidgets : public SCompoundWidget
 			];
 		}
 
-		ButtonBox->AddSlot()
-		.AutoWidth()
-		.HAlign(HAlign_Left)
-		.VAlign(VAlign_Center)
-		.Padding(4.0f, 0.0f)
-		[
-			SNew(SResetToDefaultPropertyEditor, PropertyEditor->GetPropertyHandle())
-		];
-
 		ChildSlot
 		[
 			ButtonBox
