Added a EditorScriptExecutionGuard to FPropertyHandleBase::GeneratePossibleValues so GetOptions property functions no longer require CallInEditor specifiers when used with Actors

diff --git a/Engine/Source/Editor/PropertyEditor/Private/PropertyHandleImpl.cpp b/Engine/Source/Editor/PropertyEditor/Private/PropertyHandleImpl.cpp
--- a/Engine/Source/Editor/PropertyEditor/Private/PropertyHandleImpl.cpp
+++ b/Engine/Source/Editor/PropertyEditor/Private/PropertyHandleImpl.cpp
@@ -2986,16 +2986,20 @@ bool FPropertyHandleBase::GeneratePossibleValues(TArray< TSharedPtr<FString> >&
 				for (UObject* Target : OutObjects)
 				{
 					TArray<FString> StringOptions;
-					if (PropertyPathHelpers::GetPropertyValue(Target, Path, StringOptions))
 					{
-						// No-Op
-					}
-					else
-					{
-						TArray<FName> NameOptions;
-						if (PropertyPathHelpers::GetPropertyValue(Target, Path, NameOptions))
+						FEditorScriptExecutionGuard ScriptExecutionGuard;
+
+						if (PropertyPathHelpers::GetPropertyValue(Target, Path, StringOptions))
 						{
-							Algo::Transform(NameOptions, StringOptions, [](const FName& InName) { return InName.ToString(); });
+							// No-Op
+						}
+						else
+						{
+							TArray<FName> NameOptions;
+							if (PropertyPathHelpers::GetPropertyValue(Target, Path, NameOptions))
+							{
+								Algo::Transform(NameOptions, StringOptions, [](const FName& InName) { return InName.ToString(); });
+							}
 						}
 					}
 
