Removed unnecessary image that was causing a dark stripe in the drop shadow of advanced nodes in the details view.

diff --git a/Engine/Source/Editor/PropertyEditor/Private/DetailAdvancedDropdownNode.cpp b/Engine/Source/Editor/PropertyEditor/Private/DetailAdvancedDropdownNode.cpp
--- a/Engine/Source/Editor/PropertyEditor/Private/DetailAdvancedDropdownNode.cpp
+++ b/Engine/Source/Editor/PropertyEditor/Private/DetailAdvancedDropdownNode.cpp
@@ -107,15 +107,6 @@ public:
 					SNew(SBorder)
 					.BorderImage_Lambda(GetScrollbarWellBrush)
 					.Padding(FMargin(0, 0, SDetailTableRowBase::ScrollBarPadding, 0))
-					[
-						SNew(SBox)
-						.VAlign(VAlign_Fill)
-						[
-							SNew(SImage)
-							.Image(FEditorStyle::GetBrush("DetailsView.AdvancedDropdownBorder.Open"))
-							.Visibility(bIsTopNode ? EVisibility::Visible : EVisibility::Collapsed)
-						]
-					]
 				]
 			]
 		];
