Array properties now display Index[#] as their display name.

Simplified logic for how SPropertyEditorTitle creates the name textblock.

diff --git a/Engine/Source/Editor/PropertyEditor/Private/PropertyEditorHelpers.cpp b/Engine/Source/Editor/PropertyEditor/Private/PropertyEditorHelpers.cpp
--- a/Engine/Source/Editor/PropertyEditor/Private/PropertyEditorHelpers.cpp
+++ b/Engine/Source/Editor/PropertyEditor/Private/PropertyEditorHelpers.cpp
@@ -34,10 +34,6 @@ void SPropertyNameWidget::Construct( const FArguments& InArgs, TSharedPtr<FPrope
 {
 	PropertyEditor = InPropertyEditor;
 
-	static const FName NAME_TitleProperty = FName(TEXT("TitleProperty"));
-	// If our property has title support we pass in empty below so it retrieves a live value
-	const bool bHasTitleProperty = InPropertyEditor->GetProperty() && InPropertyEditor->GetProperty()->HasMetaData(NAME_TitleProperty);
-
 	TSharedPtr<SHorizontalBox> HorizontalBox;
 	ChildSlot
 	[
@@ -52,7 +48,6 @@ void SPropertyNameWidget::Construct( const FArguments& InArgs, TSharedPtr<FPrope
 			.VAlign(VAlign_Center)
 			[
 				SNew( SPropertyEditorTitle, PropertyEditor.ToSharedRef() )
-				.StaticDisplayName(bHasTitleProperty ? FText::GetEmpty() : PropertyEditor->GetDisplayName())
 				.OnDoubleClicked( InArgs._OnDoubleClicked )
                 .ToolTip( IDocumentation::Get()->CreateToolTip( PropertyEditor->GetToolTipText(), NULL, PropertyEditor->GetDocumentationLink(), PropertyEditor->GetDocumentationExcerptName() ) )
 			]
diff --git a/Engine/Source/Editor/PropertyEditor/Private/UserInterface/PropertyEditor/SPropertyEditorTitle.h b/Engine/Source/Editor/PropertyEditor/Private/UserInterface/PropertyEditor/SPropertyEditorTitle.h
--- a/Engine/Source/Editor/PropertyEditor/Private/UserInterface/PropertyEditor/SPropertyEditorTitle.h
+++ b/Engine/Source/Editor/PropertyEditor/Private/UserInterface/PropertyEditor/SPropertyEditorTitle.h
@@ -13,20 +13,15 @@
 #include "UserInterface/PropertyEditor/PropertyEditorConstants.h"
 #include "Widgets/Text/STextBlock.h"
 
-class FCategoryPropertyNode;
-class FItemPropertyNode;
-
 class SPropertyEditorTitle : public SCompoundWidget
 {
 
 public:
 
 	SLATE_BEGIN_ARGS( SPropertyEditorTitle )
-		: _StaticDisplayName()
-		, _PropertyFont( FEditorStyle::GetFontStyle( PropertyEditorConstants::PropertyFontStyle ) )
+		: _PropertyFont( FEditorStyle::GetFontStyle( PropertyEditorConstants::PropertyFontStyle ) )
 		, _CategoryFont( FEditorStyle::GetFontStyle( PropertyEditorConstants::CategoryFontStyle ) )
 		{}
-		SLATE_ARGUMENT( FText, StaticDisplayName )
 		SLATE_ATTRIBUTE( FSlateFontInfo, PropertyFont )
 		SLATE_ATTRIBUTE( FSlateFontInfo, CategoryFont )
 		SLATE_EVENT( FOnClicked, OnDoubleClicked )
@@ -38,42 +33,79 @@ public:
 		OnDoubleClicked = InArgs._OnDoubleClicked;
 
 		const TSharedRef< FPropertyNode > PropertyNode = InPropertyEditor->GetPropertyNode();
-		FCategoryPropertyNode* CategoryNode = PropertyNode->AsCategoryNode();
-		FItemPropertyNode* ItemPropertyNode = PropertyNode->AsItemPropertyNode();
+		const bool bIsCategory = PropertyNode->AsCategoryNode() != nullptr;
+		const TAttribute<FSlateFontInfo> NameFont = bIsCategory ? InArgs._CategoryFont : InArgs._PropertyFont;
+
+		TSharedPtr<STextBlock> NameTextBlock;
 
-		TSharedPtr<SWidget> NameWidget;
-		const bool bHasStaticName = !InArgs._StaticDisplayName.IsEmpty();
-		if ( CategoryNode != NULL )
+		// If our property has title support we want to fetch the value every tick, otherwise we can just use a static value
+		static const FName NAME_TitleProperty = FName(TEXT("TitleProperty"));
+		const bool bHasTitleProperty = InPropertyEditor->GetProperty() && InPropertyEditor->GetProperty()->HasMetaData(NAME_TitleProperty);
+		if (bHasTitleProperty)
+		{
+			NameTextBlock =
+				SNew( STextBlock )
+				.Text( InPropertyEditor, &FPropertyEditor::GetDisplayName )
+				.Font( NameFont );
+		}		
+		else
 		{
-			NameWidget =
+			NameTextBlock =
 				SNew( STextBlock )
-				.Text( bHasStaticName ? InArgs._StaticDisplayName : InPropertyEditor->GetDisplayName() )
-				.Font( InArgs._CategoryFont );
+				.Text( InPropertyEditor->GetDisplayName() )
+				.Font( NameFont );
 		}
-		else if ( ItemPropertyNode != NULL )
+
+		TSharedPtr<SWidget> NameWidget = NameTextBlock;
+		const bool bInArray = PropertyNode->GetProperty() != nullptr && PropertyNode->GetArrayIndex() != INDEX_NONE;
+		if (bInArray && !bHasTitleProperty)
 		{
-			if( bHasStaticName )
+			const bool bNameIsIndex = NameTextBlock->GetText().EqualTo(FText::AsNumber(PropertyNode->GetArrayIndex()));
+			if (bNameIsIndex)
 			{
-				NameWidget =
-					SNew( STextBlock )
-					.Text( InArgs._StaticDisplayName )
-					.Font( InArgs._PropertyFont );
-			}		
-			else
-			{
-				NameWidget =
-					SNew( STextBlock )
-					.Text( InPropertyEditor, &FPropertyEditor::GetDisplayName )
-					.Font( InArgs._PropertyFont );
+				NameWidget = SNew(SHorizontalBox)
+				+ SHorizontalBox::Slot()
+				.HAlign(HAlign_Left)
+				.VAlign(VAlign_Center)
+				.Padding(0, 0, 3, 0)
+				.AutoWidth()
+				[
+					SNew(STextBlock)
+					.Text(NSLOCTEXT("PropertyEditor", "Index", "Index"))
+					.Font(NameFont)
+					.ColorAndOpacity(FSlateColor::UseSubduedForeground())
+				]
+				+ SHorizontalBox::Slot()
+				.HAlign(HAlign_Left)
+				.VAlign(VAlign_Center)
+				.Padding(0, 0, 3, 0)
+				.AutoWidth()
+				[
+					SNew(STextBlock)
+					.Text(NSLOCTEXT("PropertyEditor", "OpenBracket", "["))
+					.Font(NameFont)
+					.ColorAndOpacity(FSlateColor::UseSubduedForeground())
+				]
+				+ SHorizontalBox::Slot()
+				.HAlign(HAlign_Left)
+				.VAlign(VAlign_Center)
+				.Padding(0, 0, 3, 0)
+				.AutoWidth()
+				[
+					NameTextBlock.ToSharedRef()
+				]
+				+ SHorizontalBox::Slot()
+				.HAlign(HAlign_Left)
+				.VAlign(VAlign_Center)
+				.AutoWidth()
+				[
+					SNew(STextBlock)
+					.Text(NSLOCTEXT("PropertyEditor", "CloseBracket", "]"))
+					.Font(NameFont)
+					.ColorAndOpacity(FSlateColor::UseSubduedForeground())
+				];
 			}
 		}
-		else
-		{
-			NameWidget = 
-				SNew( STextBlock )
-				.Text( bHasStaticName ? InArgs._StaticDisplayName : InPropertyEditor->GetDisplayName() )
-				.Font( InArgs._PropertyFont );
-		}
 
 		ChildSlot
 		[
@@ -81,7 +113,6 @@ public:
 		];
 	}
 
-
 private:
 
 	FReply OnMouseButtonDoubleClick( const FGeometry& InMyGeometry, const FPointerEvent& InMouseEvent )
