Details view labels are now, once again, left aligned.

diff --git a/Engine/Source/Editor/PropertyEditor/Private/DetailPropertyRow.cpp b/Engine/Source/Editor/PropertyEditor/Private/DetailPropertyRow.cpp
--- a/Engine/Source/Editor/PropertyEditor/Private/DetailPropertyRow.cpp
+++ b/Engine/Source/Editor/PropertyEditor/Private/DetailPropertyRow.cpp
@@ -657,7 +657,7 @@ bool FDetailPropertyRow::GetForceAutoExpansion() const
 void FDetailPropertyRow::MakeNameOrKeyWidget( FDetailWidgetRow& Row, const TSharedPtr<FDetailWidgetRow> InCustomRow ) const
 {
 	EVerticalAlignment VerticalAlignment = VAlign_Center;
-	EHorizontalAlignment HorizontalAlignment = HAlign_Right;
+	EHorizontalAlignment HorizontalAlignment = HAlign_Left;
 
 	// We will only use key widgets for non-struct keys
 	const bool bHasKeyNode = PropertyKeyEditor.IsValid() && !PropertyHandle->HasMetaData(TEXT("ReadOnlyKeys"));
@@ -720,7 +720,7 @@ void FDetailPropertyRow::MakeNameOrKeyWidget( FDetailWidgetRow& Row, const TShar
 	}
 
 	SHorizontalBox::FSlot& Slot = NameHorizontalBox->AddSlot()
-	.HAlign(HAlign_Right)
+	.HAlign(HAlign_Left)
 	[
 		NameWidget.ToSharedRef()
 	];
diff --git a/Engine/Source/Editor/PropertyEditor/Public/DetailWidgetRow.h b/Engine/Source/Editor/PropertyEditor/Public/DetailWidgetRow.h
--- a/Engine/Source/Editor/PropertyEditor/Public/DetailWidgetRow.h
+++ b/Engine/Source/Editor/PropertyEditor/Public/DetailWidgetRow.h
@@ -106,7 +106,7 @@ public:
 	PROPERTYEDITOR_API const static float DefaultValueMaxWidth;
 
 	FDetailWidgetRow()
-		: NameWidget( *this, 0.0f, 0.0f, HAlign_Right, VAlign_Center )
+		: NameWidget( *this, 0.0f, 0.0f, HAlign_Left, VAlign_Center )
 		, ValueWidget( *this, DefaultValueMinWidth, DefaultValueMaxWidth, HAlign_Left, VAlign_Fill )
 		, WholeRowWidget( *this, 0.0f, 0.0f, HAlign_Fill, VAlign_Fill )
 		, VisibilityAttr( EVisibility::Visible )
