Array reorder handle now hides itself if the row isn't hovered.

Indentation is now consistent regardless of whether an item has children or not.

Expander arrow no longer contains padding, instead it's been moved to the places it's used.

diff --git a/Engine/Source/Editor/PropertyEditor/Private/DetailCategoryGroupNode.cpp b/Engine/Source/Editor/PropertyEditor/Private/DetailCategoryGroupNode.cpp
--- a/Engine/Source/Editor/PropertyEditor/Private/DetailCategoryGroupNode.cpp
+++ b/Engine/Source/Editor/PropertyEditor/Private/DetailCategoryGroupNode.cpp
@@ -29,6 +29,7 @@ void SDetailCategoryTableRow::Construct( const FArguments& InArgs, TSharedRef<FD
 		+ SHorizontalBox::Slot()
 		.HAlign(HAlign_Left)
 		.VAlign(VAlign_Center)
+		.Padding(5,0,0,0)
 		.AutoWidth()
 		[
 			SNew(SDetailExpanderArrow, SharedThis(this))
diff --git a/Engine/Source/Editor/PropertyEditor/Private/SDetailExpanderArrow.cpp b/Engine/Source/Editor/PropertyEditor/Private/SDetailExpanderArrow.cpp
--- a/Engine/Source/Editor/PropertyEditor/Private/SDetailExpanderArrow.cpp
+++ b/Engine/Source/Editor/PropertyEditor/Private/SDetailExpanderArrow.cpp
@@ -13,7 +13,7 @@ void SDetailExpanderArrow::Construct(const FArguments& InArgs, TSharedRef<SDetai
 	ChildSlot
 	[
 		SNew(SConstrainedBox)
-		.MinWidth(20)
+		.MinWidth(16)
 		.Visibility(this, &SDetailExpanderArrow::GetExpanderVisibility)
 		[
 			SAssignNew(ExpanderArrow, SButton)
@@ -22,7 +22,7 @@ void SDetailExpanderArrow::Construct(const FArguments& InArgs, TSharedRef<SDetai
 			.HAlign(HAlign_Center)
 			.ClickMethod(EButtonClickMethod::MouseDown)
 			.OnClicked(this, &SDetailExpanderArrow::OnExpanderClicked)
-			.ContentPadding(FMargin(5,0,0,0))
+			.ContentPadding(FMargin(0))
 			.IsFocusable(false)
 			[
 				SNew(SImage)
@@ -35,7 +35,7 @@ void SDetailExpanderArrow::Construct(const FArguments& InArgs, TSharedRef<SDetai
 
 EVisibility SDetailExpanderArrow::GetExpanderVisibility() const
 {
-	return Row->DoesItemHaveChildren() ? EVisibility::Visible : EVisibility::Collapsed;
+	return Row->DoesItemHaveChildren() ? EVisibility::Visible : EVisibility::Hidden;
 }
 
 const FSlateBrush* SDetailExpanderArrow::GetExpanderImage() const
diff --git a/Engine/Source/Editor/PropertyEditor/Private/SDetailSingleItemRow.cpp b/Engine/Source/Editor/PropertyEditor/Private/SDetailSingleItemRow.cpp
--- a/Engine/Source/Editor/PropertyEditor/Private/SDetailSingleItemRow.cpp
+++ b/Engine/Source/Editor/PropertyEditor/Private/SDetailSingleItemRow.cpp
@@ -28,7 +28,7 @@
 
 namespace DetailWidgetConstants
 {
-	const FMargin LeftRowPadding( 12.0f, 2.5f, 12.0f, 2.5f );
+	const FMargin LeftRowPadding( 20.0f, 2.5f, 10.0f, 2.5f );
 	const FMargin RightRowPadding( 12.0f, 2.5f, 2.0f, 2.5f );
 }
 
@@ -360,7 +360,7 @@ void SDetailSingleItemRow::Construct( const FArguments& InArgs, FDetailLayoutCus
 			NameColumnBox->AddSlot()
 				.HAlign(HAlign_Left)
 				.VAlign(VAlign_Fill)
-				.Padding(0,0,0,0)
+				.Padding(0)
 				.AutoWidth()
 				[
 					SNew(SDetailRowIndent, SharedThis(this))
@@ -369,7 +369,7 @@ void SDetailSingleItemRow::Construct( const FArguments& InArgs, FDetailLayoutCus
 			NameColumnBox->AddSlot()
 				.HAlign(HAlign_Left)
 				.VAlign(VAlign_Center)
-				.Padding(0)
+				.Padding(5,0,0,0)
 				.AutoWidth()
 				[
 					SNew(SDetailExpanderArrow, SharedThis(this))
@@ -383,11 +383,15 @@ void SDetailSingleItemRow::Construct( const FArguments& InArgs, FDetailLayoutCus
 					TSharedPtr<SDetailSingleItemRow> InRow = SharedThis(this);
 					TSharedRef<SWidget> ArrayHandle = PropertyEditorHelpers::MakePropertyReorderHandle(PropertyNode.ToSharedRef(), InRow);
 					ArrayHandle->SetEnabled(IsEnabledAttribute);
+					ArrayHandle->SetVisibility(TAttribute<EVisibility>::Create([this]() 
+					{
+						return this->IsHovered() ? EVisibility::Visible : EVisibility::Hidden;
+					}));
 
 					NameColumnBox->AddSlot()
-						.Padding(5,0,0,0)
 						.HAlign(HAlign_Left)
 						.VAlign(VAlign_Center)
+						.Padding(0)
 						.AutoWidth()
 						[
 							ArrayHandle
@@ -407,12 +411,20 @@ void SDetailSingleItemRow::Construct( const FArguments& InArgs, FDetailLayoutCus
 				}
 			}
 
+			const bool bIsReorderable = PropertyNode.IsValid() && PropertyNode->IsReorderable();
+			auto GetLeftRowPadding = [this, bIsReorderable]()
+			{
+				FMargin Padding = DetailWidgetConstants::LeftRowPadding;
+				Padding.Left -= bIsReorderable ? 16 : 0;
+				return Padding;
+			};
+
 			if (bHasMultipleColumns)
 			{
 				NameColumnBox->AddSlot()
 					.HAlign(WidgetRow.NameWidget.HorizontalAlignment)
 					.VAlign(WidgetRow.NameWidget.VerticalAlignment)
-					.Padding(DetailWidgetConstants::LeftRowPadding)
+					.Padding(TAttribute<FMargin>::Create(GetLeftRowPadding))
 					[
 						NameWidget.ToSharedRef()
 					];
@@ -459,7 +471,7 @@ void SDetailSingleItemRow::Construct( const FArguments& InArgs, FDetailLayoutCus
 				NameColumnBox->AddSlot()
 					.HAlign(WidgetRow.WholeRowWidget.HorizontalAlignment)
 					.VAlign(WidgetRow.WholeRowWidget.VerticalAlignment)
-					.Padding(DetailWidgetConstants::LeftRowPadding)
+					.Padding(TAttribute<FMargin>::Create(GetLeftRowPadding))
 					[
 						WidgetRow.WholeRowWidget.Widget
 					];
@@ -950,7 +962,6 @@ FReply SArrayRowHandle::OnDragDetected(const FGeometry& MyGeometry, const FPoint
 	}
 
 	return FReply::Unhandled();
-
 }
 
 TSharedPtr<FArrayRowDragDropOp> SArrayRowHandle::CreateDragDropOperation(TSharedPtr<SDetailSingleItemRow> InRow)
