Add NoCreate metadata tag to prevent uproperties from diplaying factories

diff --git a/Engine/Source/Editor/PropertyEditor/Private/UserInterface/PropertyEditor/SPropertyEditorAsset.cpp b/Engine/Source/Editor/PropertyEditor/Private/UserInterface/PropertyEditor/SPropertyEditorAsset.cpp
--- a/Engine/Source/Editor/PropertyEditor/Private/UserInterface/PropertyEditor/SPropertyEditorAsset.cpp
+++ b/Engine/Source/Editor/PropertyEditor/Private/UserInterface/PropertyEditor/SPropertyEditorAsset.cpp
@@ -346,14 +346,17 @@ void SPropertyEditorAsset::Construct(const FArguments& InArgs, const TSharedPtr<
 
 	bIsActor = ObjectClass->IsChildOf(AActor::StaticClass()) || ObjectClass->IsChildOf(UInterface::StaticClass());
 
-	if (InArgs._NewAssetFactories.IsSet())
+	if (Property && !Property->HasMetaData("NoCreate"))
 	{
-		NewAssetFactories = InArgs._NewAssetFactories.GetValue();
-	}
-	// If there are more allowed classes than just UObject 
-	else if (AllowedClassFilters.Num() > 1 || !AllowedClassFilters.Contains(UObject::StaticClass()))
-	{
-		NewAssetFactories = PropertyCustomizationHelpers::GetNewAssetFactoriesForClasses(AllowedClassFilters, DisallowedClassFilters);
+		if (InArgs._NewAssetFactories.IsSet())
+		{
+			NewAssetFactories = InArgs._NewAssetFactories.GetValue();
+		}
+		// If there are more allowed classes than just UObject 
+		else if (AllowedClassFilters.Num() > 1 || !AllowedClassFilters.Contains(UObject::StaticClass()))
+		{
+			NewAssetFactories = PropertyCustomizationHelpers::GetNewAssetFactoriesForClasses(AllowedClassFilters, DisallowedClassFilters);
+		}
 	}
 	
 	TSharedPtr<SHorizontalBox> ValueContentBox = nullptr;
