Padding fixes for left alignment in the details panel.

Splitter line now always visible, and splitter column width ratio adjusted to be closer to 50/50.

Favorites hint text now renders with a more subdued style.

diff --git a/Engine/Source/Editor/EditorStyle/Private/StarshipStyle.cpp b/Engine/Source/Editor/EditorStyle/Private/StarshipStyle.cpp
--- a/Engine/Source/Editor/EditorStyle/Private/StarshipStyle.cpp
+++ b/Engine/Source/Editor/EditorStyle/Private/StarshipStyle.cpp
@@ -3267,11 +3267,6 @@ void FStarshipEditorStyle::FStyle::SetupPropertyEditorStyles()
 		Set( "DetailsView.NoFavoritesSystem", new IMAGE_BRUSH("PropertyView/NoFavoritesSystem", Icon12x12));
 
 		Set( "DetailsView.Splitter", FSplitterStyle()
-			.SetHandleNormalBrush(FSlateColorBrush(FStyleColors::Transparent))                   
-			.SetHandleHighlightBrush(FSlateColorBrush(FStyleColors::Recessed))
-		);
-
-		Set( "DetailsView.Splitter.Outer", FSplitterStyle()
 			.SetHandleNormalBrush(FSlateColorBrush(FStyleColors::Recessed))                   
 			.SetHandleHighlightBrush(FSlateColorBrush(FStyleColors::Recessed))
 		);
diff --git a/Engine/Source/Editor/PropertyEditor/Private/DetailCategoryGroupNode.cpp b/Engine/Source/Editor/PropertyEditor/Private/DetailCategoryGroupNode.cpp
--- a/Engine/Source/Editor/PropertyEditor/Private/DetailCategoryGroupNode.cpp
+++ b/Engine/Source/Editor/PropertyEditor/Private/DetailCategoryGroupNode.cpp
@@ -56,8 +56,6 @@ void SDetailCategoryTableRow::Construct( const FArguments& InArgs, TSharedRef<FD
 		];
 	}
 	
-	const float EditConditionWidgetWidth = 33; // this is the width of the edit condition widget displayed on the left in SDetailSingleItemRow
-
 	TWeakPtr<STableViewBase> OwnerTableViewWeak = InOwnerTableView;
 	auto GetScrollbarWellBrush = [this, OwnerTableViewWeak]()
 	{
@@ -81,19 +79,6 @@ void SDetailCategoryTableRow::Construct( const FArguments& InArgs, TSharedRef<FD
 		[
 			SNew(SHorizontalBox)
 			+ SHorizontalBox::Slot()
-			.HAlign(HAlign_Left)
-			.AutoWidth()
-			[
-				SNew(SBorder)
-				.BorderImage(FAppStyle::Get().GetBrush("DetailsView.CategoryMiddle"))
-				.BorderBackgroundColor( this, &SDetailCategoryTableRow::GetOuterBackgroundColor )
-				.Padding(0)
-				[
-					SNew(SSpacer)
-					.Size(bIsInnerCategory ? FVector2D(EditConditionWidgetWidth, 0) : FVector2D(0,0))
-				]
-			]
-			+ SHorizontalBox::Slot()
 			.HAlign(HAlign_Fill)
 			[
 				SNew(SBorder)
diff --git a/Engine/Source/Editor/PropertyEditor/Private/SDetailExpanderArrow.cpp b/Engine/Source/Editor/PropertyEditor/Private/SDetailExpanderArrow.cpp
--- a/Engine/Source/Editor/PropertyEditor/Private/SDetailExpanderArrow.cpp
+++ b/Engine/Source/Editor/PropertyEditor/Private/SDetailExpanderArrow.cpp
@@ -35,7 +35,7 @@ void SDetailExpanderArrow::Construct(const FArguments& InArgs, TSharedRef<SDetai
 
 EVisibility SDetailExpanderArrow::GetExpanderVisibility() const
 {
-	return Row->DoesItemHaveChildren() ? EVisibility::Visible : EVisibility::Hidden;
+	return Row->DoesItemHaveChildren() ? EVisibility::Visible : EVisibility::Collapsed;
 }
 
 const FSlateBrush* SDetailExpanderArrow::GetExpanderImage() const
diff --git a/Engine/Source/Editor/PropertyEditor/Private/SDetailNameArea.cpp b/Engine/Source/Editor/PropertyEditor/Private/SDetailNameArea.cpp
--- a/Engine/Source/Editor/PropertyEditor/Private/SDetailNameArea.cpp
+++ b/Engine/Source/Editor/PropertyEditor/Private/SDetailNameArea.cpp
@@ -190,7 +190,7 @@ TSharedRef< SWidget > SDetailNameArea::BuildObjectNameArea( const TArray< TWeakO
 			[
 				SNew( STextBlock )
 				.Text( LOCTEXT("NoObjectsSelected", "Select an object to view details.") )
-				.ShadowOffset( FVector2D(1,1) )
+				.TextStyle( FAppStyle::Get(), "HintText")
 			];
 		}
 		else
diff --git a/Engine/Source/Editor/PropertyEditor/Private/SDetailSingleItemRow.cpp b/Engine/Source/Editor/PropertyEditor/Private/SDetailSingleItemRow.cpp
--- a/Engine/Source/Editor/PropertyEditor/Private/SDetailSingleItemRow.cpp
+++ b/Engine/Source/Editor/PropertyEditor/Private/SDetailSingleItemRow.cpp
@@ -296,40 +296,12 @@ void SDetailSingleItemRow::Construct( const FArguments& InArgs, FDetailLayoutCus
 			// create outer splitter
 			TSharedRef<SSplitter> OuterSplitter =
 				SNew(SSplitter)
-				.Style(FEditorStyle::Get(), "DetailsView.Splitter.Outer")
+				.Style(FEditorStyle::Get(), "DetailsView.Splitter")
 				.PhysicalSplitterHandleSize(1.0f)
 				.HitDetectionSplitterHandleSize(5.0f);
 
 			Widget = OuterSplitter;
 
-			// create Left column:
-			// | Left  | Name | Value | Right |
-			TSharedRef<SHorizontalBox> LeftColumnBox =
-				SNew(SHorizontalBox)
-				.Clipping(EWidgetClipping::OnDemand);
-
-			// edit condition widget
-			LeftColumnBox->AddSlot()
-				.Padding(6.0f, 0.0f, 6.0f, 0.0f)
-				.HAlign(HAlign_Center)
-				.VAlign(VAlign_Center)
-				.AutoWidth()
-				[
-					SNew(SConstrainedBox)
-					.MinWidth(20)
-					[
-						SNew(SEditConditionWidget)
-						.EditConditionValue(WidgetRow.EditConditionValue)
-						.OnEditConditionValueChanged(WidgetRow.OnEditConditionValueChanged)
-					]
-				];
-
-			OuterSplitter->AddSlot()
-				.SizeRule(SSplitter::ESizeRule::SizeToContent)
-				[
-					LeftColumnBox
-				];
-
 			TSharedPtr<SSplitter> InnerSplitter;
 
 			// create inner splitter
@@ -366,15 +338,6 @@ void SDetailSingleItemRow::Construct( const FArguments& InArgs, FDetailLayoutCus
 					SNew(SDetailRowIndent, SharedThis(this))
 				];
 
-			NameColumnBox->AddSlot()
-				.HAlign(HAlign_Left)
-				.VAlign(VAlign_Center)
-				.Padding(5,0,0,0)
-				.AutoWidth()
-				[
-					SNew(SDetailExpanderArrow, SharedThis(this))
-				];
-			
 			TSharedPtr<FPropertyNode> PropertyNode = Customization->GetPropertyNode();
 			if (PropertyNode.IsValid())
 			{
@@ -391,7 +354,7 @@ void SDetailSingleItemRow::Construct( const FArguments& InArgs, FDetailLayoutCus
 					NameColumnBox->AddSlot()
 						.HAlign(HAlign_Left)
 						.VAlign(VAlign_Center)
-						.Padding(0)
+						.Padding(-2, 0, -10, 0)
 						.AutoWidth()
 						[
 							ArrayHandle
@@ -411,20 +374,32 @@ void SDetailSingleItemRow::Construct( const FArguments& InArgs, FDetailLayoutCus
 				}
 			}
 
-			const bool bIsReorderable = PropertyNode.IsValid() && PropertyNode->IsReorderable();
-			auto GetLeftRowPadding = [this, bIsReorderable]()
-			{
-				FMargin Padding = DetailWidgetConstants::LeftRowPadding;
-				Padding.Left -= bIsReorderable ? 16 : 0;
-				return Padding;
-			};
+			NameColumnBox->AddSlot()
+				.HAlign(HAlign_Left)
+				.VAlign(VAlign_Center)
+				.Padding(8,0,0,0)
+				.AutoWidth()
+				[
+					SNew(SDetailExpanderArrow, SharedThis(this))
+				];
+
+			NameColumnBox->AddSlot()
+				.VAlign(VAlign_Center)
+				.HAlign(HAlign_Left)
+				.Padding(3,0,0,0)
+				.AutoWidth()
+				[
+					SNew(SEditConditionWidget)
+					.EditConditionValue(WidgetRow.EditConditionValue)
+					.OnEditConditionValueChanged(WidgetRow.OnEditConditionValueChanged)
+				];
 
 			if (bHasMultipleColumns)
 			{
 				NameColumnBox->AddSlot()
 					.HAlign(WidgetRow.NameWidget.HorizontalAlignment)
 					.VAlign(WidgetRow.NameWidget.VerticalAlignment)
-					.Padding(TAttribute<FMargin>::Create(GetLeftRowPadding))
+					.Padding(5,0,0,0)
 					[
 						NameWidget.ToSharedRef()
 					];
@@ -471,7 +446,7 @@ void SDetailSingleItemRow::Construct( const FArguments& InArgs, FDetailLayoutCus
 				NameColumnBox->AddSlot()
 					.HAlign(WidgetRow.WholeRowWidget.HorizontalAlignment)
 					.VAlign(WidgetRow.WholeRowWidget.VerticalAlignment)
-					.Padding(TAttribute<FMargin>::Create(GetLeftRowPadding))
+					.Padding(5,0,0,0)
 					[
 						WidgetRow.WholeRowWidget.Widget
 					];
diff --git a/Engine/Source/Editor/PropertyEditor/Private/SDetailsViewBase.cpp b/Engine/Source/Editor/PropertyEditor/Private/SDetailsViewBase.cpp
--- a/Engine/Source/Editor/PropertyEditor/Private/SDetailsViewBase.cpp
+++ b/Engine/Source/Editor/PropertyEditor/Private/SDetailsViewBase.cpp
@@ -424,7 +424,9 @@ void SDetailsViewBase::UpdateSinglePropertyMap(TSharedPtr<FComplexPropertyNode>
 				.VAlign(VAlign_Center)
 				.HAlign(HAlign_Center)
 				[
-					SNew(STextBlock).Text(NSLOCTEXT("DetailLayoutHelpers", "AddToFavoritesDescription", "Right-click on a property to add it to your Favorites."))
+					SNew(STextBlock)
+					.Text(NSLOCTEXT("DetailLayoutHelpers", "AddToFavoritesDescription", "Right-click on a property to add it to your Favorites."))
+					.TextStyle(FAppStyle::Get(), "HintText")
 				];
 		}
 	}
diff --git a/Engine/Source/Editor/PropertyEditor/Public/IDetailPropertyRow.h b/Engine/Source/Editor/PropertyEditor/Public/IDetailPropertyRow.h
--- a/Engine/Source/Editor/PropertyEditor/Public/IDetailPropertyRow.h
+++ b/Engine/Source/Editor/PropertyEditor/Public/IDetailPropertyRow.h
@@ -20,7 +20,7 @@ public:
 	
 	FDetailColumnSizeData()
 	{
-		ValueColumnWidthValue = 0.65f;
+		ValueColumnWidthValue = 0.5f;
 		RightColumnWidthValue = 0.05f;
 		HoveredSplitterIndexValue = INDEX_NONE;
 
diff --git a/Engine/Source/Runtime/SlateCore/Private/Styling/StarshipCoreStyle.cpp b/Engine/Source/Runtime/SlateCore/Private/Styling/StarshipCoreStyle.cpp
--- a/Engine/Source/Runtime/SlateCore/Private/Styling/StarshipCoreStyle.cpp
+++ b/Engine/Source/Runtime/SlateCore/Private/Styling/StarshipCoreStyle.cpp
@@ -1047,7 +1047,10 @@ void FStarshipCoreStyle::SetupTextStyles(TSharedRef<FStyle>& Style)
 		Style->Set("InlineEditableTextBlockSmallStyle", InlineEditableTextBlockSmallStyle);
 	}
 
-
+	const FTextBlockStyle HintText = FTextBlockStyle(NormalText)
+		.SetFont(FONT(8, "Italic"))
+		.SetColorAndOpacity(FSlateColor::UseSubduedForeground());
+	Style->Set("HintText", HintText);
 }
 
 void FStarshipCoreStyle::SetupButtonStyles(TSharedRef<FStyle>& Style)
