SDetailExpanderArrow and SDetailRowIndent no longer hold TSharedPtr to their owning rows, which was causing the owning rows to never be GC'd.

diff --git a/Engine/Source/Editor/PropertyEditor/Private/SDetailExpanderArrow.cpp b/Engine/Source/Editor/PropertyEditor/Private/SDetailExpanderArrow.cpp
--- a/Engine/Source/Editor/PropertyEditor/Private/SDetailExpanderArrow.cpp
+++ b/Engine/Source/Editor/PropertyEditor/Private/SDetailExpanderArrow.cpp
@@ -35,24 +35,36 @@ void SDetailExpanderArrow::Construct(const FArguments& InArgs, TSharedRef<SDetai
 
 EVisibility SDetailExpanderArrow::GetExpanderVisibility() const
 {
-	return Row->DoesItemHaveChildren() ? EVisibility::Visible : EVisibility::Collapsed;
+	TSharedPtr<SDetailTableRowBase> RowPtr = Row.Pin();
+	if (!RowPtr.IsValid())
+	{
+		return EVisibility::Collapsed;
+	}
+
+	return RowPtr->DoesItemHaveChildren() ? EVisibility::Visible : EVisibility::Collapsed;
 }
 
 const FSlateBrush* SDetailExpanderArrow::GetExpanderImage() const
 {
-	const bool bIsItemExpanded = Row->IsItemExpanded();
+	TSharedPtr<SDetailTableRowBase> RowPtr = Row.Pin();
+	if (!RowPtr.IsValid())
+	{
+		return FAppStyle::Get().GetBrush("NoBrush");
+	}
+
+	const bool bIsItemExpanded = RowPtr->IsItemExpanded();
 
 	FName ResourceName;
 	if (bIsItemExpanded)
 	{
 		if (ExpanderArrow->IsHovered())
 		{
-			static FName ExpandedHoveredName = "TreeArrow_Expanded_Hovered";
+			static const FName ExpandedHoveredName = "TreeArrow_Expanded_Hovered";
 			ResourceName = ExpandedHoveredName;
 		}
 		else
 		{
-			static FName ExpandedName = "TreeArrow_Expanded";
+			static const FName ExpandedName = "TreeArrow_Expanded";
 			ResourceName = ExpandedName;
 		}
 	}
@@ -60,12 +72,12 @@ const FSlateBrush* SDetailExpanderArrow::GetExpanderImage() const
 	{
 		if (ExpanderArrow->IsHovered())
 		{
-			static FName CollapsedHoveredName = "TreeArrow_Collapsed_Hovered";
+			static const FName CollapsedHoveredName = "TreeArrow_Collapsed_Hovered";
 			ResourceName = CollapsedHoveredName;
 		}
 		else
 		{
-			static FName CollapsedName = "TreeArrow_Collapsed";
+			static const FName CollapsedName = "TreeArrow_Collapsed";
 			ResourceName = CollapsedName;
 		}
 	}
@@ -75,15 +87,21 @@ const FSlateBrush* SDetailExpanderArrow::GetExpanderImage() const
 
 FReply SDetailExpanderArrow::OnExpanderClicked()
 {
+	TSharedPtr<SDetailTableRowBase> RowPtr = Row.Pin();
+	if (!RowPtr.IsValid())
+	{
+		return FReply::Unhandled();
+	}
+
 	// Recurse the expansion if "shift" is being pressed
 	const FModifierKeysState ModKeyState = FSlateApplication::Get().GetModifierKeys();
 	if (ModKeyState.IsShiftDown())
 	{
-		Row->Private_OnExpanderArrowShiftClicked();
+		RowPtr->Private_OnExpanderArrowShiftClicked();
 	}
 	else
 	{
-		Row->ToggleExpansion();
+		RowPtr->ToggleExpansion();
 	}
 
 	return FReply::Handled();
diff --git a/Engine/Source/Editor/PropertyEditor/Private/SDetailExpanderArrow.h b/Engine/Source/Editor/PropertyEditor/Private/SDetailExpanderArrow.h
--- a/Engine/Source/Editor/PropertyEditor/Private/SDetailExpanderArrow.h
+++ b/Engine/Source/Editor/PropertyEditor/Private/SDetailExpanderArrow.h
@@ -23,6 +23,6 @@ private:
 	FReply OnExpanderClicked();
 	
 private:
-	TSharedPtr<SDetailTableRowBase> Row;
+	TWeakPtr<SDetailTableRowBase> Row;
 	TSharedPtr<SButton> ExpanderArrow;
 };
\ No newline at end of file
diff --git a/Engine/Source/Editor/PropertyEditor/Private/SDetailRowIndent.cpp b/Engine/Source/Editor/PropertyEditor/Private/SDetailRowIndent.cpp
--- a/Engine/Source/Editor/PropertyEditor/Private/SDetailRowIndent.cpp
+++ b/Engine/Source/Editor/PropertyEditor/Private/SDetailRowIndent.cpp
@@ -19,10 +19,16 @@ void SDetailRowIndent::Construct(const FArguments& InArgs, TSharedRef<SDetailTab
 
 int32 SDetailRowIndent::OnPaint(const FPaintArgs& Args, const FGeometry& AllottedGeometry, const FSlateRect& MyCullingRect, FSlateWindowElementList& OutDrawElements, int32 LayerId, const FWidgetStyle& InWidgetStyle, bool bParentEnabled) const
 {
+	TSharedPtr<SDetailTableRowBase> RowPtr = Row.Pin();
+	if (!RowPtr.IsValid())
+	{
+		return LayerId;
+	}
+
 	const FSlateBrush* BackgroundBrush = FAppStyle::Get().GetBrush("DetailsView.CategoryMiddle");
 	const FSlateBrush* DropShadowBrush = FAppStyle::Get().GetBrush("DetailsView.ArrayDropShadow");
 
-	int32 IndentLevel = Row->GetIndentLevelForBackgroundColor();
+	int32 IndentLevel = RowPtr->GetIndentLevelForBackgroundColor();
 	for (int32 i = 0; i < IndentLevel; ++i)
 	{
 		FSlateColor BackgroundColor = GetRowBackgroundColor(i);
@@ -49,16 +55,23 @@ int32 SDetailRowIndent::OnPaint(const FPaintArgs& Args, const FGeometry& Allotte
 
 FOptionalSize SDetailRowIndent::GetIndentWidth() const
 {
-	int32 IndentLevel = Row->GetIndentLevelForBackgroundColor();
+	int32 IndentLevel = 0;
+
+	TSharedPtr<SDetailTableRowBase> RowPtr = Row.Pin();
+	if (RowPtr.IsValid())
+	{
+		IndentLevel = RowPtr->GetIndentLevelForBackgroundColor();
+	}
 	return IndentLevel * 16.0f;
 }
 
-FSlateColor SDetailRowIndent::GetRowBackgroundColor(int32 i) const
+FSlateColor SDetailRowIndent::GetRowBackgroundColor(int32 IndentLevel) const
 {
-	if (Row->IsHovered())
+	TSharedPtr<SDetailTableRowBase> RowPtr = Row.Pin();
+	if (RowPtr.IsValid() && RowPtr->IsHovered())
 	{
 		return FAppStyle::Get().GetSlateColor("Colors.Header");
 	}
 
-	return PropertyEditorConstants::GetRowBackgroundColor(i);
+	return PropertyEditorConstants::GetRowBackgroundColor(IndentLevel);
 }
diff --git a/Engine/Source/Editor/PropertyEditor/Private/SDetailRowIndent.h b/Engine/Source/Editor/PropertyEditor/Private/SDetailRowIndent.h
--- a/Engine/Source/Editor/PropertyEditor/Private/SDetailRowIndent.h
+++ b/Engine/Source/Editor/PropertyEditor/Private/SDetailRowIndent.h
@@ -21,8 +21,8 @@ public:
 private:
 	int32 OnPaint(const FPaintArgs& Args, const FGeometry& AllottedGeometry, const FSlateRect& MyCullingRect, FSlateWindowElementList& OutDrawElements, int32 LayerId, const FWidgetStyle& InWidgetStyle, bool bParentEnabled) const;
 	FOptionalSize GetIndentWidth() const;
-	FSlateColor GetRowBackgroundColor(int32 i) const;
+	FSlateColor GetRowBackgroundColor(int32 IndentLevel) const;
 
 private:
-	TSharedPtr<SDetailTableRowBase> Row;
+	TWeakPtr<SDetailTableRowBase> Row;
 };
