Added FDetailPropertyRow::SetWidgetRowProperties to handle the setting of the FDetailWidgetRow's custom properties correctly in both cases that it's used.

diff --git a/Engine/Source/Editor/PropertyEditor/Private/DetailPropertyRow.cpp b/Engine/Source/Editor/PropertyEditor/Private/DetailPropertyRow.cpp
--- a/Engine/Source/Editor/PropertyEditor/Private/DetailPropertyRow.cpp
+++ b/Engine/Source/Editor/PropertyEditor/Private/DetailPropertyRow.cpp
@@ -1,17 +1,18 @@
 // Copyright Epic Games, Inc. All Rights Reserved.
 
 #include "DetailPropertyRow.h"
-#include "Modules/ModuleManager.h"
-#include "PropertyCustomizationHelpers.h"
-#include "DetailItemNode.h"
-#include "DetailCategoryGroupNode.h"
-#include "ObjectPropertyNode.h"
+#include "CategoryPropertyNode.h"
 #include "CustomChildBuilder.h"
-#include "StructurePropertyNode.h"
+#include "DetailCategoryGroupNode.h"
+#include "DetailItemNode.h"
+#include "DetailWidgetRow.h"
 #include "ItemPropertyNode.h"
 #include "ObjectPropertyNode.h"
-#include "DetailWidgetRow.h"
-#include "CategoryPropertyNode.h"
+#include "ObjectPropertyNode.h"
+#include "PropertyCustomizationHelpers.h"
+#include "PropertyEditorHelpers.h"
+#include "StructurePropertyNode.h"
+#include "Modules/ModuleManager.h"
 #include "Widgets/Layout/SSpacer.h"
 
 #define LOCTEXT_NAMESPACE	"DetailPropertyRow"
@@ -145,6 +146,7 @@ void FDetailPropertyRow::GetDefaultWidgets( TSharedPtr<SWidget>& OutNameWidget,
 		CustomTypeInterface->CustomizeHeader(PropertyHandle.ToSharedRef(), *CustomTypeRow, *this);
 	}
 
+	SetWidgetRowProperties(Row);
 	MakeNameOrKeyWidget(Row, CustomTypeRow);
 	MakeValueWidget(Row, CustomTypeRow, bAddWidgetDecoration);
 
@@ -208,18 +210,11 @@ FDetailWidgetRow FDetailPropertyRow::GetWidgetRow()
 	if( HasColumns() )
 	{
 		FDetailWidgetRow Row;
-	
+
+		SetWidgetRowProperties(Row);
 		MakeNameOrKeyWidget( Row, CustomPropertyWidget );
 		MakeValueWidget( Row, CustomPropertyWidget );
 
-		if (CustomPropertyWidget.IsValid())
-		{
-			Row.CopyMenuAction = CustomPropertyWidget->CopyMenuAction;
-			Row.PasteMenuAction = CustomPropertyWidget->PasteMenuAction;
-			Row.CustomMenuItems = CustomPropertyWidget->CustomMenuItems;
-			Row.CustomResetToDefault = CustomPropertyWidget->CustomResetToDefault;
-		}
-		
 		return Row;
 	}
 	else
@@ -654,6 +649,49 @@ bool FDetailPropertyRow::GetForceAutoExpansion() const
 	return bForceAutoExpansion;
 }
 
+static void TogglePropertyEditorEditCondition(bool bValue, TWeakPtr<FPropertyEditor> PropertyEditorWeak)
+{
+	TSharedPtr<FPropertyEditor> PropertyEditorPtr = PropertyEditorWeak.Pin();
+	if (PropertyEditorPtr.IsValid() && PropertyEditorPtr->IsEditConditionMet() != bValue)
+	{
+		PropertyEditorPtr->ToggleEditConditionState();
+	}
+}
+
+void FDetailPropertyRow::SetWidgetRowProperties(FDetailWidgetRow& Row) const
+{
+	// set edit condition handlers - use customized if provided
+	TAttribute<bool> EditConditionValue = CustomEditConditionValue;
+	if (!EditConditionValue.IsSet() && !EditConditionValue.IsBound())
+	{
+		EditConditionValue = TAttribute<bool>(PropertyEditor.ToSharedRef(), &FPropertyEditor::IsEditConditionMet);
+	}
+
+	FOnBooleanValueChanged OnEditConditionValueChanged = CustomEditConditionValueChanged;
+	if (!OnEditConditionValueChanged.IsBound() && PropertyEditor->SupportsEditConditionToggle())
+	{
+		TWeakPtr<FPropertyEditor> PropertyEditorWeak = PropertyEditor;
+		OnEditConditionValueChanged = FOnBooleanValueChanged::CreateStatic(&TogglePropertyEditorEditCondition, PropertyEditorWeak);
+	}
+
+	Row.EditCondition(EditConditionValue, OnEditConditionValueChanged);
+	Row.CustomResetToDefault = CustomResetToDefault;
+
+	// set custom actions and reset to default
+	if (CustomPropertyWidget.IsValid())
+	{
+		Row.CopyMenuAction = CustomPropertyWidget->CopyMenuAction;
+		Row.PasteMenuAction = CustomPropertyWidget->PasteMenuAction;
+		Row.CustomMenuItems = CustomPropertyWidget->CustomMenuItems;
+
+		if (CustomPropertyWidget->CustomResetToDefault.IsSet())
+		{
+			checkf(!CustomResetToDefault.IsSet(), TEXT("Duplicate reset to default handlers set on both FDetailPropertyRow and CustomWidget()!"));
+			Row.CustomResetToDefault = CustomPropertyWidget->CustomResetToDefault;
+		}
+	}
+}
+
 void FDetailPropertyRow::MakeNameOrKeyWidget( FDetailWidgetRow& Row, const TSharedPtr<FDetailWidgetRow> InCustomRow ) const
 {
 	EVerticalAlignment VerticalAlignment = VAlign_Center;
@@ -835,27 +873,6 @@ void FDetailPropertyRow::MakeValueWidget( FDetailWidgetRow& Row, const TSharedPt
 		[
 			ValueWidget
 		];
-	
-	// set edit condition handlers - use customized if provided
-	TAttribute<bool> EditConditionValue = CustomEditConditionValue;
-	if (!EditConditionValue.IsSet())
-	{
-		EditConditionValue = TAttribute<bool>(PropertyEditor.ToSharedRef(), &FPropertyEditor::IsEditConditionMet);
-	}
-
-	FOnBooleanValueChanged OnEditConditionValueChanged = CustomEditConditionValueChanged;
-	if (!OnEditConditionValueChanged.IsBound() && PropertyEditor->SupportsEditConditionToggle())
-	{
-		OnEditConditionValueChanged = FOnBooleanValueChanged::CreateLambda([this](bool bValue) 
-		{ 
-			if (PropertyEditor->IsEditConditionMet() != bValue)
-			{
-				PropertyEditor->ToggleEditConditionState();
-			}
-		});
-	}
-
-	Row.EditCondition(EditConditionValue, OnEditConditionValueChanged);
 }
 
 #undef LOCTEXT_NAMESPACE
diff --git a/Engine/Source/Editor/PropertyEditor/Private/DetailPropertyRow.h b/Engine/Source/Editor/PropertyEditor/Private/DetailPropertyRow.h
--- a/Engine/Source/Editor/PropertyEditor/Private/DetailPropertyRow.h
+++ b/Engine/Source/Editor/PropertyEditor/Private/DetailPropertyRow.h
@@ -120,6 +120,11 @@ private:
 	 */
 	void MakeValueWidget( FDetailWidgetRow& Row, const TSharedPtr<FDetailWidgetRow> InCustomRow, bool bAddWidgetDecoration = true ) const;
 
+	/**
+	 * Set the given widget row's custom properties. Note: Should be called before MakeNameOrKeyWidget and MakeValueWidget.
+	 */
+	void SetWidgetRowProperties( FDetailWidgetRow& Row ) const;
+
 	/**
 	 * @return true if this row has an edit condition
 	 */
