SEditConditionWidget now has a GetVisibility() instead of using a lambda, which makes it easier to see that the class controls its own visibility.

diff --git a/Engine/Source/Editor/PropertyEditor/Private/PropertyEditorHelpers.cpp b/Engine/Source/Editor/PropertyEditor/Private/PropertyEditorHelpers.cpp
--- a/Engine/Source/Editor/PropertyEditor/Private/PropertyEditorHelpers.cpp
+++ b/Engine/Source/Editor/PropertyEditor/Private/PropertyEditorHelpers.cpp
@@ -333,10 +333,15 @@ void SEditConditionWidget::Construct( const FArguments& Args )
 		SNew(SCheckBox)
 		.OnCheckStateChanged(this, &SEditConditionWidget::OnEditConditionCheckChanged)
 		.IsChecked(this, &SEditConditionWidget::OnGetEditConditionCheckState)
-		.Visibility_Lambda([this]() { return HasEditConditionToggle() ? EVisibility::Visible : EVisibility::Collapsed; })
+		.Visibility(this, &SEditConditionWidget::GetVisibility)
 	];
 }
 
+EVisibility SEditConditionWidget::GetVisibility() const
+{
+	return HasEditConditionToggle() ? EVisibility::Visible : EVisibility::Collapsed;
+}
+
 bool SEditConditionWidget::HasEditConditionToggle() const
 {
 	return OnEditConditionValueChanged.IsBound();
diff --git a/Engine/Source/Editor/PropertyEditor/Private/PropertyEditorHelpers.h b/Engine/Source/Editor/PropertyEditor/Private/PropertyEditorHelpers.h
--- a/Engine/Source/Editor/PropertyEditor/Private/PropertyEditorHelpers.h
+++ b/Engine/Source/Editor/PropertyEditor/Private/PropertyEditorHelpers.h
@@ -94,6 +94,7 @@ private:
 	void OnEditConditionCheckChanged( ECheckBoxState CheckState );
 	bool HasEditConditionToggle() const;
 	ECheckBoxState OnGetEditConditionCheckState() const;
+	EVisibility GetVisibility() const;
 
 private:
 	TAttribute<bool> EditConditionValue;
