Fixed bug with instanced objects not updating the reset to default state in the parent object when changed.

diff --git a/Engine/Source/Editor/PropertyEditor/Private/PropertyNode.cpp b/Engine/Source/Editor/PropertyEditor/Private/PropertyNode.cpp
--- a/Engine/Source/Editor/PropertyEditor/Private/PropertyNode.cpp
+++ b/Engine/Source/Editor/PropertyEditor/Private/PropertyNode.cpp
@@ -899,45 +899,45 @@ bool FPropertyNode::IsEditConst() const
 {
 	if (bUpdateEditConstState)
 	{
-	// Ask the objects whether this property can be changed
-	const FObjectPropertyNode* ObjectPropertyNode = FindObjectItemParent();
+		// Ask the objects whether this property can be changed
+		const FObjectPropertyNode* ObjectPropertyNode = FindObjectItemParent();
 
 		bIsEditConst = IsPropertyConst();
 		if (!bIsEditConst && Property != nullptr && ObjectPropertyNode)
-	{
-		// travel up the chain to see if this property's owner struct is editconst - if it is, so is this property
-		FPropertyNode* NextParent = ParentNode;
-			while (NextParent != nullptr && CastField<FStructProperty>(NextParent->GetProperty()) != NULL)
 		{
-			if (NextParent->IsEditConst())
+			// travel up the chain to see if this property's owner struct is editconst - if it is, so is this property
+			FPropertyNode* NextParent = ParentNode;
+			while (NextParent != nullptr && CastField<FStructProperty>(NextParent->GetProperty()) != NULL)
 			{
-					bIsEditConst = true;
-					break;
+				if (NextParent->IsEditConst())
+				{
+						bIsEditConst = true;
+						break;
+				}
+				NextParent = NextParent->ParentNode;
 			}
-			NextParent = NextParent->ParentNode;
-		}
 
 			if (!bIsEditConst)
 			{
-		for (TPropObjectConstIterator CurObjectIt(ObjectPropertyNode->ObjectConstIterator()); CurObjectIt; ++CurObjectIt)
-		{
-			const TWeakObjectPtr<UObject> CurObject = *CurObjectIt;
-			if (CurObject.IsValid())
-			{
-				if (!CurObject->CanEditChange(Property.Get()))
+				for (TPropObjectConstIterator CurObjectIt(ObjectPropertyNode->ObjectConstIterator()); CurObjectIt; ++CurObjectIt)
 				{
-					// At least one of the objects didn't like the idea of this property being changed.
+					const TWeakObjectPtr<UObject> CurObject = *CurObjectIt;
+					if (CurObject.IsValid())
+					{
+						if (!CurObject->CanEditChange(Property.Get()))
+						{
+							// At least one of the objects didn't like the idea of this property being changed.
 							bIsEditConst = true;
 							break;
 						}
+					}
 				}
 			}
 		}
-	}
 
-	// check edit condition
+		// check edit condition
 		if (!bIsEditConst && HasEditCondition())
-	{
+		{
 			bIsEditConst = !IsEditConditionMet();
 		}
 
@@ -1897,10 +1897,10 @@ bool FPropertyNode::GetDiffersFromDefaultForObject( FPropertyItemValueDataTracke
 	bool bHasParent = GetParentNode() != nullptr;
 
 	if (bIsValidTracker && bHasDefaultValue && bHasParent)
-		{
-			//////////////////////////
-			// Check the property against its default.
-			// If the property is an object property, we have to take special measures.
+	{
+		//////////////////////////
+		// Check the property against its default.
+		// If the property is an object property, we have to take special measures.
 
 		if (FArrayProperty* OuterArrayProperty = InProperty->GetOwner<FArrayProperty>())
 		{
@@ -2451,11 +2451,20 @@ void FPropertyNode::NotifyPostChange( FPropertyChangedEvent& InPropertyChangedEv
 	// remember the property that was the chain's original active property; this will correspond to the outermost property of struct/array that was modified
 	FProperty* const OriginalActiveProperty = PropertyChain->GetActiveMemberNode()->GetValue();
 
+	// invalidate the entire chain of objects in the hierarchy 
+	FObjectPropertyNode* CurrentObjectNode = FindObjectItemParent();
+	while (CurrentObjectNode != nullptr)
+	{
+		CurrentObjectNode->InvalidateCachedState();
+
+		// FindObjectItemParent returns itself if the node is an object, so step up the hierarchy to get to its actual parent object
+		FPropertyNode* CurrentParent = CurrentObjectNode->GetParentNode();
+		CurrentObjectNode = CurrentParent != nullptr ? CurrentParent->FindObjectItemParent() : nullptr;
+	}
+
 	FObjectPropertyNode* ObjectNode = FindObjectItemParent();
 	if( ObjectNode )
 	{
-		ObjectNode->InvalidateCachedState();
-
 		FProperty* CurProperty = InPropertyChangedEvent.Property;
 
 		// Fire ULevel::LevelDirtiedEvent when falling out of scope.
