Changing a check to an ensure so we can find the indeterminate behavior but not crash.

diff --git a/Engine/Source/Editor/PropertyEditor/Private/DetailPropertyRow.cpp b/Engine/Source/Editor/PropertyEditor/Private/DetailPropertyRow.cpp
--- a/Engine/Source/Editor/PropertyEditor/Private/DetailPropertyRow.cpp
+++ b/Engine/Source/Editor/PropertyEditor/Private/DetailPropertyRow.cpp
@@ -686,7 +686,7 @@ void FDetailPropertyRow::SetWidgetRowProperties(FDetailWidgetRow& Row) const
 
 		if (CustomPropertyWidget->CustomResetToDefault.IsSet())
 		{
-			checkf(!CustomResetToDefault.IsSet(), TEXT("Duplicate reset to default handlers set on both FDetailPropertyRow and CustomWidget()!"));
+			ensureMsgf(!CustomResetToDefault.IsSet(), TEXT("Duplicate reset to default handlers set on both FDetailPropertyRow and CustomWidget()!"));
 			Row.CustomResetToDefault = CustomPropertyWidget->CustomResetToDefault;
 		}
 	}
