Moved unnecessary padding from FDetailPropertyRow that was causing 4 pixels of padding to always appear on the right and throw off alignment in the details panel.

diff --git a/Engine/Source/Editor/PropertyEditor/Private/DetailPropertyRow.cpp b/Engine/Source/Editor/PropertyEditor/Private/DetailPropertyRow.cpp
--- a/Engine/Source/Editor/PropertyEditor/Private/DetailPropertyRow.cpp
+++ b/Engine/Source/Editor/PropertyEditor/Private/DetailPropertyRow.cpp
@@ -823,7 +823,6 @@ void FDetailPropertyRow::MakeValueWidget( FDetailWidgetRow& Row, const TSharedPt
 	else
 	{
 		ValueWidget->AddSlot()
-		.Padding( 0.0f, 0.0f, 4.0f, 0.0f )
 		[
 			SAssignNew( PropertyValue, SPropertyValueWidget, PropertyEditor, GetPropertyUtilities() )
 			.ShowPropertyButtons( false ) // We handle this ourselves
@@ -845,7 +844,7 @@ void FDetailPropertyRow::MakeValueWidget( FDetailWidgetRow& Row, const TSharedPt
 				.AutoWidth()
 				.HAlign(HAlign_Left)
 				.VAlign(VAlign_Center)
-				.Padding(2.0f, 1.0f)
+				.Padding(4.0f, 1.0f, 0.0f, 1.0f)
 				[ 
 					RequiredButtons[ButtonIndex]
 				];
@@ -858,7 +857,7 @@ void FDetailPropertyRow::MakeValueWidget( FDetailWidgetRow& Row, const TSharedPt
 			.AutoWidth()
 			.VAlign(VAlign_Center)
 			.HAlign(HAlign_Left)
-			.Padding(0.0f, 0.0f, 4.0f, 0.0f)
+			.Padding(4.0f, 0.0f, 4.0f, 0.0f)
 			[
 				PropertyCustomizationHelpers::MakeEditConfigHierarchyButton(FSimpleDelegate::CreateSP(PropertyEditor.ToSharedRef(), &FPropertyEditor::EditConfigHierarchy))
 			];
