FDetailItemNode now also checks if the IsEnabled attribute IsSet() and not just IsBound() in GenerateStandaloneWidget().

diff --git a/Engine/Source/Editor/PropertyEditor/Private/DetailItemNode.cpp b/Engine/Source/Editor/PropertyEditor/Private/DetailItemNode.cpp
--- a/Engine/Source/Editor/PropertyEditor/Private/DetailItemNode.cpp
+++ b/Engine/Source/Editor/PropertyEditor/Private/DetailItemNode.cpp
@@ -287,11 +287,6 @@ bool FDetailItemNode::GenerateStandaloneWidget(FDetailWidgetRow& OutRow) const
 			TSharedPtr<SWidget> ValueWidget;
 
 			NameWidget = Row.NameWidget.Widget;
-			if (Row.IsEnabledAttr.IsBound())
-			{
-				NameWidget->SetEnabled(Row.IsEnabledAttr);
-			}
-
 			ValueWidget =
 				SNew(SConstrainedBox)
 				.MinWidth(Row.ValueWidget.MinWidth)
@@ -300,8 +295,9 @@ bool FDetailItemNode::GenerateStandaloneWidget(FDetailWidgetRow& OutRow) const
 					Row.ValueWidget.Widget
 				];
 
-			if (Row.IsEnabledAttr.IsBound())
+			if (Row.IsEnabledAttr.IsSet() || Row.IsEnabledAttr.IsBound())
 			{
+				NameWidget->SetEnabled(Row.IsEnabledAttr);
 				ValueWidget->SetEnabled(Row.IsEnabledAttr);
 			}
 
@@ -323,10 +319,7 @@ bool FDetailItemNode::GenerateStandaloneWidget(FDetailWidgetRow& OutRow) const
 			];
 		}
 
-		if (Row.CustomResetToDefault.IsSet())
-		{
-			OutRow.OverrideResetToDefault(Row.CustomResetToDefault.GetValue());
-		}
+		OutRow.CustomResetToDefault = Row.CustomResetToDefault;
 
 		OutRow.EditConditionValue = Row.EditConditionValue;
 		OutRow.OnEditConditionValueChanged = Row.OnEditConditionValueChanged;
