All columns in the details view are now in the same splitter, which makes resizing work more intuitively.

diff --git a/Engine/Source/Editor/PropertyEditor/Private/SDetailSingleItemRow.cpp b/Engine/Source/Editor/PropertyEditor/Private/SDetailSingleItemRow.cpp
--- a/Engine/Source/Editor/PropertyEditor/Private/SDetailSingleItemRow.cpp
+++ b/Engine/Source/Editor/PropertyEditor/Private/SDetailSingleItemRow.cpp
@@ -291,40 +291,23 @@ void SDetailSingleItemRow::Construct( const FArguments& InArgs, FDetailLayoutCus
 			NameWidget->SetEnabled(IsEnabledAttribute);
 			ValueWidget->SetEnabled(IsEnabledAttribute);
 
-			TSharedRef<SHorizontalBox> RowBox = SNew(SHorizontalBox);
-
-			// create outer splitter
-			TSharedRef<SSplitter> OuterSplitter =
-				SNew(SSplitter)
-				.Style(FEditorStyle::Get(), "DetailsView.Splitter")
-				.PhysicalSplitterHandleSize(1.0f)
-				.HitDetectionSplitterHandleSize(5.0f);
-
-			Widget = OuterSplitter;
-
-			TSharedPtr<SSplitter> InnerSplitter;
+			TSharedRef<SSplitter> Splitter = SNew(SSplitter)
+					.Style(FEditorStyle::Get(), "DetailsView.Splitter")
+					.PhysicalSplitterHandleSize(1.0f)
+					.HitDetectionSplitterHandleSize(5.0f)
+					.HighlightedHandleIndex(ColumnSizeData.HoveredSplitterIndex)
+					.OnHandleHovered(ColumnSizeData.OnSplitterHandleHovered);
 
-			// create inner splitter
-			OuterSplitter->AddSlot()
-				.Value(ColumnSizeData.PropertyColumnWidth)
-				.OnSlotResized(ColumnSizeData.OnPropertyColumnResized)
-			[
-				SNew(SBorder)
+			Widget = SNew(SBorder)
 				.BorderImage(FAppStyle::Get().GetBrush("DetailsView.CategoryMiddle"))
 				.BorderBackgroundColor(this, &SDetailSingleItemRow::GetInnerBackgroundColor)
 				.Padding(0)
 				[
-					SAssignNew(InnerSplitter, SSplitter)
-					.Style(FEditorStyle::Get(), "DetailsView.Splitter")
-					.PhysicalSplitterHandleSize(1.0f)
-					.HitDetectionSplitterHandleSize(5.0f)
-					.HighlightedHandleIndex(ColumnSizeData.HoveredSplitterIndex)
-					.OnHandleHovered(ColumnSizeData.OnSplitterHandleHovered)
-				]
-			];
+					Splitter
+				];
 
 			// create Name column:
-			// | Left  | Name | Value | Right |
+			// | Name | Value | Right |
 			TSharedRef<SHorizontalBox> NameColumnBox = SNew(SHorizontalBox)
 				.Clipping(EWidgetClipping::OnDemand);
 
@@ -404,7 +387,7 @@ void SDetailSingleItemRow::Construct( const FArguments& InArgs, FDetailLayoutCus
 						NameWidget.ToSharedRef()
 					];
 		
-				InnerSplitter->AddSlot()
+				Splitter->AddSlot()
 					.Value(ColumnSizeData.NameColumnWidth)
 					.OnSlotResized(ColumnSizeData.OnNameColumnResized)
 					[
@@ -415,8 +398,8 @@ void SDetailSingleItemRow::Construct( const FArguments& InArgs, FDetailLayoutCus
 				ExtensionWidget->SetEnabled(IsEnabledAttribute);
 
 				// create Value column:
-				// | Left  | Name | Value | Right |
-				InnerSplitter->AddSlot()
+				// | Name | Value | Right |
+				Splitter->AddSlot()
 					.Value(ColumnSizeData.ValueColumnWidth)
 					.OnSlotResized(ColumnSizeData.OnValueColumnResized) 
 					[
@@ -451,7 +434,9 @@ void SDetailSingleItemRow::Construct( const FArguments& InArgs, FDetailLayoutCus
 						WidgetRow.WholeRowWidget.Widget
 					];
 
-				InnerSplitter->AddSlot()
+				Splitter->AddSlot()
+					.Value(ColumnSizeData.PropertyColumnWidth)
+					.OnSlotResized(ColumnSizeData.OnPropertyColumnResized)
 					[
 						NameColumnBox
 					];
@@ -495,21 +480,25 @@ void SDetailSingleItemRow::Construct( const FArguments& InArgs, FDetailLayoutCus
 				ToolbarBuilder.AddToolBarButton(Extension.UIAction, NAME_None, Extension.Label, Extension.ToolTip, Extension.Icon);
 			}
 
-			TSharedRef<SHorizontalBox> RightColumnBox = SNew(SHorizontalBox);
-			RightColumnBox->AddSlot()
-				.HAlign(HAlign_Fill)
-				.VAlign(VAlign_Center)
-				.Padding(5,0,5,0)
-				[
-					ToolbarBuilder.MakeWidget()
-				];
-
-			OuterSplitter->AddSlot()
+			Splitter->AddSlot()
 				.Value(ColumnSizeData.RightColumnWidth)
 				.OnSlotResized(ColumnSizeData.OnRightColumnResized)
 				.MinSize(50)
 			[
-				RightColumnBox
+				SNew(SBorder)
+				.BorderImage(FAppStyle::Get().GetBrush("DetailsView.CategoryMiddle"))
+				.BorderBackgroundColor(this, &SDetailSingleItemRow::GetOuterBackgroundColor)
+				.Padding(0)
+				[
+					SNew(SHorizontalBox)
+					+ SHorizontalBox::Slot()
+					.HAlign(HAlign_Fill)
+					.VAlign(VAlign_Center)
+					.Padding(5,0,5,0)
+					[
+						ToolbarBuilder.MakeWidget()
+					]
+				]
 			];
 		}
 	}
