Changed numeric property bitmask to display up to 3 values in the combo box.

diff --git a/Engine/Source/Editor/PropertyEditor/Private/UserInterface/PropertyEditor/SPropertyEditorNumeric.h b/Engine/Source/Editor/PropertyEditor/Private/UserInterface/PropertyEditor/SPropertyEditorNumeric.h
--- a/Engine/Source/Editor/PropertyEditor/Private/UserInterface/PropertyEditor/SPropertyEditorNumeric.h
+++ b/Engine/Source/Editor/PropertyEditor/Private/UserInterface/PropertyEditor/SPropertyEditorNumeric.h
@@ -135,21 +135,25 @@ public:
 					NumericType BitmaskValue = Value.GetValue();
 					if (BitmaskValue != 0)
 					{
-						if (TBitmaskValueHelpers<NumericType>::BitwiseAND(BitmaskValue, BitmaskValue - static_cast<NumericType>(1)))
-						{
-							return LOCTEXT("BitmaskButtonContentMultipleBitsSet", "(Mixed Flags)");
-						}
-						else
+						TArray<FBitmaskFlagInfo> BitmaskFlags = CreateBitmaskFlagsArray(Property);
+
+						TArray<FText> SetFlags;
+						SetFlags.Reserve(BitmaskFlags.Num());
+
+						for (const FBitmaskFlagInfo& FlagInfo : BitmaskFlags)
 						{
-							TArray<FBitmaskFlagInfo> BitmaskFlags = CreateBitmaskFlagsArray(Property);
-							for (int i = 0; i < BitmaskFlags.Num(); ++i)
+							if (TBitmaskValueHelpers<NumericType>::BitwiseAND(Value.GetValue(), FlagInfo.Value))
 							{
-								if (TBitmaskValueHelpers<NumericType>::BitwiseAND(BitmaskValue, BitmaskFlags[i].Value))
-								{
-									return BitmaskFlags[i].DisplayName;
-								}
+								SetFlags.Add(FlagInfo.DisplayName);
 							}
 						}
+						if (SetFlags.Num() > 3)
+						{
+							SetFlags.SetNum(3);
+							SetFlags.Add(FText::FromString("..."));
+						}
+
+						return FText::Join(FText::FromString(" | "), SetFlags);
 					}
 
 					return LOCTEXT("BitmaskButtonContentNoFlagsSet", "(No Flags Set)");
