Added Toggle Favorites button to toggle the visibility of the Favorites category in details panels that use Favorites.

Added drop shadow below the details search box.

diff --git a/Engine/Source/Editor/PropertyEditor/Private/SDetailsView.cpp b/Engine/Source/Editor/PropertyEditor/Private/SDetailsView.cpp
--- a/Engine/Source/Editor/PropertyEditor/Private/SDetailsView.cpp
+++ b/Engine/Source/Editor/PropertyEditor/Private/SDetailsView.cpp
@@ -19,7 +19,7 @@
 #include "DetailLayoutBuilderImpl.h"
 #include "DetailsViewPropertyGenerationUtilities.h"
 #include "DetailsViewGenericObjectFilter.h"
-
+#include "Styling/StyleColors.h"
 
 #define LOCTEXT_NAMESPACE "SDetailsView"
 
@@ -187,7 +187,30 @@ void SDetailsView::Construct(const FArguments& InArgs, const FDetailsViewArgs& I
 			FSlateIcon(),
 			FUIAction(FExecuteAction::CreateSP(this, &SDetailsView::SetRootExpansionStates, /*bExpanded=*/true, /*bRecurse=*/false )));
 
-	FilterRow = 
+	TSharedRef<SVerticalBox> VerticalBox = SNew(SVerticalBox);
+		
+	// Create the name area which does not change when selection changes
+	SAssignNew(NameArea, SDetailNameArea, &SelectedObjects)
+		// the name area is only for actors
+		.Visibility(this, &SDetailsView::GetActorNameAreaVisibility)
+		.OnLockButtonClicked(this, &SDetailsView::OnLockButtonClicked)
+		.IsLocked(this, &SDetailsView::IsLocked)
+		.ShowLockButton(DetailsViewArgs.bLockable)
+		.ShowActorLabel(DetailsViewArgs.bShowActorLabel)
+		// only show the selection tip if we're not selecting objects
+		.SelectionTip(!DetailsViewArgs.bHideSelectionTip);
+
+	if( !DetailsViewArgs.bCustomNameAreaLocation )
+	{
+		VerticalBox->AddSlot()
+		.AutoHeight()
+		.Padding(0.0f, 0.0f, 0.0f, 4.0f)
+		[
+			NameArea.ToSharedRef()
+		];
+	}
+
+	TSharedRef<SHorizontalBox> FilterRowHBox = 
 		SNew( SHorizontalBox )
 		.Visibility( this, &SDetailsView::GetFilterBoxVisibility )
 		+SHorizontalBox::Slot()
@@ -204,8 +227,9 @@ void SDetailsView::Construct(const FArguments& InArgs, const FDetailsViewArgs& I
 
 	if (DetailsViewArgs.bShowPropertyMatrixButton)
 	{
-		FilterRow->AddSlot()
-			.Padding(4.0f, 0.0f, 0.0f, 0.0f)
+		FilterRowHBox->AddSlot()
+			.Padding(0)
+			.HAlign(HAlign_Right)
 			.VAlign(VAlign_Center)
 			.AutoWidth()
 			[
@@ -221,11 +245,33 @@ void SDetailsView::Construct(const FArguments& InArgs, const FDetailsViewArgs& I
 					.Image(FAppStyle::Get().GetBrush("DetailsView.EditRawProperties"))
 				]
 			];
+	} 
+
+	if (DetailsViewArgs.bAllowFavoriteSystem)
+	{
+		FilterRowHBox->AddSlot()
+			.Padding(0)
+			.HAlign(HAlign_Right)
+			.VAlign(VAlign_Center)
+			.AutoWidth()
+			[
+				// Create the toggle favorites button
+				SNew(SButton)
+				.ButtonStyle( FAppStyle::Get(), "SimpleButton" )
+				.OnClicked(this, &SDetailsView::OnToggleFavoritesClicked)
+				.ToolTipText(LOCTEXT("ToggleFavorites", "Toggle Favorites"))
+				[
+					SNew(SImage)
+					.ColorAndOpacity(this, &SDetailsView::GetToggleFavoritesColor)
+					.Image(FAppStyle::Get().GetBrush("PropertyWindow.Favorites_Enabled"))
+				]
+			];
 	}
 
 	if (DetailsViewArgs.bShowOptions)
 	{
-		FilterRow->AddSlot()
+		FilterRowHBox->AddSlot()
+			.Padding(0)
 			.HAlign(HAlign_Right)
 			.VAlign(VAlign_Center)
 			.AutoWidth()
@@ -249,28 +295,18 @@ void SDetailsView::Construct(const FArguments& InArgs, const FDetailsViewArgs& I
 			];
 	}
 
-	// Create the name area which does not change when selection changes
-	SAssignNew(NameArea, SDetailNameArea, &SelectedObjects)
-		// the name area is only for actors
-		.Visibility(this, &SDetailsView::GetActorNameAreaVisibility)
-		.OnLockButtonClicked(this, &SDetailsView::OnLockButtonClicked)
-		.IsLocked(this, &SDetailsView::IsLocked)
-		.ShowLockButton(DetailsViewArgs.bLockable)
-		.ShowActorLabel(DetailsViewArgs.bShowActorLabel)
-		// only show the selection tip if we're not selecting objects
-		.SelectionTip(!DetailsViewArgs.bHideSelectionTip);
-
-	TSharedRef<SVerticalBox> VerticalBox = SNew(SVerticalBox);
-
-	if( !DetailsViewArgs.bCustomNameAreaLocation )
-	{
-		VerticalBox->AddSlot()
+	FilterRow = SNew(SVerticalBox)
+		+ SVerticalBox::Slot()
 		.AutoHeight()
-		.Padding(0.0f, 0.0f, 0.0f, 4.0f)
 		[
-			NameArea.ToSharedRef()
+			FilterRowHBox
+		]
+		+ SVerticalBox::Slot()
+		.AutoHeight()
+		[
+			SNew(SImage)
+			.Image(FAppStyle::Get().GetBrush("DetailsView.AdvancedDropdownBorder.Open"))
 		];
-	}
 
 	if( !DetailsViewArgs.bCustomFilterAreaLocation )
 	{
@@ -285,7 +321,6 @@ void SDetailsView::Construct(const FArguments& InArgs, const FDetailsViewArgs& I
 		];
 	}
 
-
 	VerticalBox->AddSlot()
 	.FillHeight(1)
 	.Padding(0)
@@ -934,4 +969,24 @@ void SDetailsView::OnShowHiddenPropertiesWhilePlayingClicked()
 	ForceRefresh();
 }
 
+FSlateColor SDetailsView::GetToggleFavoritesColor() const
+{
+	if (DetailsViewArgs.bAllowFavoriteSystem && CurrentFilter.bShowFavoritesCategory)
+	{
+		return FSlateColor(EStyleColor::AccentBlue);
+	}
+	else
+	{
+		return FSlateColor(EStyleColor::Foreground);
+	}
+}
+
+FReply SDetailsView::OnToggleFavoritesClicked()
+{
+	CurrentFilter.bShowFavoritesCategory = !CurrentFilter.bShowFavoritesCategory;
+	RerunCurrentFilter();
+
+	return FReply::Handled();
+}
+
 #undef LOCTEXT_NAMESPACE
diff --git a/Engine/Source/Editor/PropertyEditor/Private/SDetailsView.h b/Engine/Source/Editor/PropertyEditor/Private/SDetailsView.h
--- a/Engine/Source/Editor/PropertyEditor/Private/SDetailsView.h
+++ b/Engine/Source/Editor/PropertyEditor/Private/SDetailsView.h
@@ -157,6 +157,12 @@ private:
 	/** Called when show hidden properties while playing is clicked */
 	void OnShowHiddenPropertiesWhilePlayingClicked();
 
+	/** Get the color of the toggle favorites button. */
+	FSlateColor GetToggleFavoritesColor() const;
+
+	/** Called when the toggle favorites button is clicked. */
+	FReply OnToggleFavoritesClicked();
+
 private:
 	/** The filter for objects viewed by this details panel */
 	TSharedPtr<FDetailsViewObjectFilter> ObjectFilter;
