Missed files from previous submit.

diff --git a/Engine/Source/Editor/PropertyEditor/Private/DetailCategoryBuilderImpl.cpp b/Engine/Source/Editor/PropertyEditor/Private/DetailCategoryBuilderImpl.cpp
--- a/Engine/Source/Editor/PropertyEditor/Private/DetailCategoryBuilderImpl.cpp
+++ b/Engine/Source/Editor/PropertyEditor/Private/DetailCategoryBuilderImpl.cpp
@@ -1001,6 +1001,11 @@ void FDetailCategoryImpl::FilterNode(const FDetailFilter& InFilter)
 
 	bHasVisibleDetails = false;
 
+	if (bFavoriteCategory && !InFilter.bShowFavoritesCategory)
+	{
+		return;
+	}
+
 	if (InlinePropertyNode.IsValid())
 	{
 		bHasVisibleDetails = true;
diff --git a/Engine/Source/Editor/PropertyEditor/Private/SDetailsViewBase.h b/Engine/Source/Editor/PropertyEditor/Private/SDetailsViewBase.h
--- a/Engine/Source/Editor/PropertyEditor/Private/SDetailsViewBase.h
+++ b/Engine/Source/Editor/PropertyEditor/Private/SDetailsViewBase.h
@@ -49,6 +49,7 @@ struct FDetailFilter
 		, bShowAllChildrenIfCategoryMatches(true)
 		, bShowKeyable(false)
 		, bShowAnimated(false)
+		, bShowFavoritesCategory(false)
 	{}
 
 	bool IsEmptyFilter() const { return FilterStrings.Num() == 0 && bShowOnlyModifiedProperties == false && bShowAllAdvanced == false && bShowOnlyDiffering == false && bShowAllChildrenIfCategoryMatches == false && bShowKeyable == false && bShowAnimated == false; }
@@ -67,6 +68,8 @@ struct FDetailFilter
 	bool bShowKeyable;
 	/** If we should only show animated properties */
 	bool bShowAnimated;
+	/** If we should show the favorites category. */
+	bool bShowFavoritesCategory;
 	TSet<FPropertyPath> WhitelistedProperties;
 };
 
@@ -366,7 +369,7 @@ protected:
 	/** The current detail layout based on objects in this details panel.  There is one layout for each top level object node.*/
 	FDetailLayoutList DetailLayouts;
 	/** Row for searching and view options */
-	TSharedPtr<SHorizontalBox>  FilterRow;
+	TSharedPtr<SWidget>  FilterRow;
 	/** Search box */
 	TSharedPtr<SSearchBox> SearchBox;
 	/** Customization instances that need to be destroyed when safe to do so */
