Removed margin in array header which was causing arrays and struct headers to not be aligned.

diff --git a/Engine/Source/Editor/PropertyEditor/Private/UserInterface/PropertyEditor/SPropertyEditorArray.h b/Engine/Source/Editor/PropertyEditor/Private/UserInterface/PropertyEditor/SPropertyEditorArray.h
--- a/Engine/Source/Editor/PropertyEditor/Private/UserInterface/PropertyEditor/SPropertyEditorArray.h
+++ b/Engine/Source/Editor/PropertyEditor/Private/UserInterface/PropertyEditor/SPropertyEditorArray.h
@@ -52,7 +52,7 @@ public:
 			.Content()
 			[
 				SNew(STextBlock)
-				.Margin(FMargin(4.0f, 2.0f))
+				.Margin(FMargin(0.0f, 0.0f, 4.0f, 0.0f))
 				.Text(TextAttr)
 				.Font(InArgs._Font)
 			]
