Allow the detail view header to prioritize showing actors or components, while still showing all objects as a fallback

diff --git a/Engine/Source/Editor/LevelEditor/Private/SActorDetails.cpp b/Engine/Source/Editor/LevelEditor/Private/SActorDetails.cpp
--- a/Engine/Source/Editor/LevelEditor/Private/SActorDetails.cpp
+++ b/Engine/Source/Editor/LevelEditor/Private/SActorDetails.cpp
@@ -34,7 +34,6 @@
 #include "Elements/Framework/TypedElementRegistry.h"
 #include "Elements/Framework/TypedElementSelectionSet.h"
 #include "Elements/Framework/EngineElementsLibrary.h"
-#include "Elements/SMInstance/SMInstanceElementId.h" // TODO: Only included to define UE_ENABLE_SMINSTANCE_ELEMENTS
 
 class SActorDetailsUneditableComponentWarning : public SCompoundWidget
 {
@@ -102,12 +101,7 @@ void SActorDetails::Construct(const FArguments& InArgs, UTypedElementSelectionSe
 	DetailsViewArgs.bUpdatesFromSelection = true;
 	DetailsViewArgs.bLockable = true;
 	DetailsViewArgs.bAllowFavoriteSystem = true;
-#if UE_ENABLE_SMINSTANCE_ELEMENTS
-	// TODO: This allows the proxy objects used by SMInstance elements to show their name, but it does also change the behavior for actors/components
-	DetailsViewArgs.NameAreaSettings = FDetailsViewArgs::ObjectsUseNameArea;
-#else
-	DetailsViewArgs.NameAreaSettings = FDetailsViewArgs::ComponentsAndActorsUseNameArea;
-#endif
+	DetailsViewArgs.NameAreaSettings = FDetailsViewArgs::ObjectsUseNameArea | FDetailsViewArgs::ComponentsAndActorsUseNameArea;
 	DetailsViewArgs.NotifyHook = GUnrealEd;
 	DetailsViewArgs.ViewIdentifier = TabIdentifier;
 	DetailsViewArgs.bCustomNameAreaLocation = true;
diff --git a/Engine/Source/Editor/PropertyEditor/Private/SDetailNameArea.cpp b/Engine/Source/Editor/PropertyEditor/Private/SDetailNameArea.cpp
--- a/Engine/Source/Editor/PropertyEditor/Private/SDetailNameArea.cpp
+++ b/Engine/Source/Editor/PropertyEditor/Private/SDetailNameArea.cpp
@@ -28,49 +28,44 @@ void SDetailNameArea::Construct( const FArguments& InArgs, const TArray< TWeakOb
 	CustomContent = SNullWidget::NullWidget;
 }
 
-void SDetailNameArea::Refresh( const TArray< TWeakObjectPtr<UObject> >& SelectedObjects )
+void SDetailNameArea::Refresh( const TArray< TWeakObjectPtr<UObject> >& SelectedObjects, int32 NameAreaSettings )
 {
-	ChildSlot
-	[
-		BuildObjectNameArea( SelectedObjects )
-	];
-}
-
-void SDetailNameArea::Refresh( const TArray< TWeakObjectPtr<AActor> >& SelectedActors, const TArray< TWeakObjectPtr<UObject> >& SelectedObjects, FDetailsViewArgs::ENameAreaSettings NameAreaSettings  )
-{
-	// Convert the actor array to base object type
-
 	TArray< TWeakObjectPtr<UObject> > FinalSelectedObjects;
-	if(NameAreaSettings == FDetailsViewArgs::ActorsUseNameArea)
+	FinalSelectedObjects.Reserve(SelectedObjects.Num());
+
+	// Apply the name area filter in priority order
+	if ((NameAreaSettings & FDetailsViewArgs::ActorsUseNameArea) && FinalSelectedObjects.Num() == 0)
 	{
-		for(auto Actor : SelectedActors)
+		for (const TWeakObjectPtr<UObject>& Object : SelectedObjects)
 		{
-			const TWeakObjectPtr<UObject> ObjectWeakPtr = Actor.Get();
-			FinalSelectedObjects.Add(ObjectWeakPtr);
+			if (AActor* Actor = Cast<AActor>(Object.Get()))
+			{
+				FinalSelectedObjects.Add(Actor);
+			}
 		}
 	}
-	else if( NameAreaSettings == FDetailsViewArgs::ComponentsAndActorsUseNameArea )
+	if ((NameAreaSettings & FDetailsViewArgs::ComponentsAndActorsUseNameArea) && FinalSelectedObjects.Num() == 0)
 	{
-		for(auto Actor : SelectedActors)
+		for (const TWeakObjectPtr<UObject>& Object : SelectedObjects)
 		{
-			const TWeakObjectPtr<UObject> ObjectWeakPtr = Actor.Get();
-			FinalSelectedObjects.Add(ObjectWeakPtr);
-		}
-
-		// Note: assumes that actors and components are not selected together.
-		if( FinalSelectedObjects.Num() == 0 )
-		{
-			for(auto Object : SelectedObjects)
+			if (UActorComponent* ActorComp = Cast<UActorComponent>(Object.Get()))
 			{
-				UActorComponent* ActorComp = Cast<UActorComponent>(Object.Get());
-				if(ActorComp && ActorComp->GetOwner())
+				if (AActor* Actor = ActorComp->GetOwner())
 				{
-					FinalSelectedObjects.AddUnique(ActorComp->GetOwner());
+					FinalSelectedObjects.AddUnique(Actor);
 				}
 			}
 		}
 	}
-	Refresh( FinalSelectedObjects );
+	if ((NameAreaSettings & FDetailsViewArgs::ObjectsUseNameArea) && FinalSelectedObjects.Num() == 0)
+	{
+		FinalSelectedObjects = SelectedObjects;
+	}
+
+	ChildSlot
+	[
+		BuildObjectNameArea( FinalSelectedObjects )
+	];
 }
 
 const FSlateBrush* SDetailNameArea::OnGetLockButtonImageResource() const
diff --git a/Engine/Source/Editor/PropertyEditor/Private/SDetailNameArea.h b/Engine/Source/Editor/PropertyEditor/Private/SDetailNameArea.h
--- a/Engine/Source/Editor/PropertyEditor/Private/SDetailNameArea.h
+++ b/Engine/Source/Editor/PropertyEditor/Private/SDetailNameArea.h
@@ -35,15 +35,7 @@ public:
 	 *
 	 * @param SelectedObjects	the new list of selected objects
 	 */
-	void Refresh( const TArray< TWeakObjectPtr<UObject> >& SelectedObjects );
-	
-	/**
-	 * Refreshes the name area when selection changes
-	 *
-	 * @param SelectedActors	the new list of selected actors
-	 */
-	void Refresh( const TArray< TWeakObjectPtr<AActor> >& SelectedActors, const TArray< TWeakObjectPtr<UObject> >& SelectedObjects, FDetailsViewArgs::ENameAreaSettings NameAreaSettings  );
-
+	void Refresh( const TArray< TWeakObjectPtr<UObject> >& SelectedObjects, int32 NameAreaSettings );
 
 	/**
 	 * Inserts Custom Content (typically tool buttons) before the lock
diff --git a/Engine/Source/Editor/PropertyEditor/Private/SDetailsView.cpp b/Engine/Source/Editor/PropertyEditor/Private/SDetailsView.cpp
--- a/Engine/Source/Editor/PropertyEditor/Private/SDetailsView.cpp
+++ b/Engine/Source/Editor/PropertyEditor/Private/SDetailsView.cpp
@@ -657,19 +657,11 @@ void SDetailsView::SetObjectArrayPrivate(const TArray<UObject*>& InObjects)
 	}
 
 	// Selection changed, refresh the detail area
-	if (DetailsViewArgs.NameAreaSettings != FDetailsViewArgs::ActorsUseNameArea && DetailsViewArgs.NameAreaSettings != FDetailsViewArgs::ComponentsAndActorsUseNameArea)
-	{
-		NameArea->Refresh(SelectedObjects);
-	}
-	else
-	{
-		NameArea->Refresh(SelectedActors, SelectedObjects, DetailsViewArgs.NameAreaSettings);
-	}
+	NameArea->Refresh(SelectedObjects, DetailsViewArgs.NameAreaSettings);
 	
 	// When selection changes rebuild information about the selection
 	SelectedActorInfo = AssetSelectionUtils::BuildSelectedActorInfo(SelectedRawActors);
 
-
 	PostSetObject(Roots);
 
 	// Set the title of the window based on the objects we are viewing
diff --git a/Engine/Source/Editor/PropertyEditor/Public/IDetailsView.h b/Engine/Source/Editor/PropertyEditor/Public/IDetailsView.h
--- a/Engine/Source/Editor/PropertyEditor/Public/IDetailsView.h
+++ b/Engine/Source/Editor/PropertyEditor/Public/IDetailsView.h
@@ -36,13 +36,13 @@ struct FDetailsViewArgs
 	enum ENameAreaSettings
 	{
 		/** The name area should never be displayed */
-		HideNameArea,
-		/** All object types use name area */
-		ObjectsUseNameArea,
-		/** Only Actors use name area */
-		ActorsUseNameArea,
-		/** Components and actors use the name area. Components will display their actor owner as the name */
-		ComponentsAndActorsUseNameArea,
+		HideNameArea = 0,
+		/** All object types use the name area */
+		ObjectsUseNameArea = 1<<0,
+		/** Actors prioritize using the name area */
+		ActorsUseNameArea = 1<<1,
+		/** Components and actors prioritize using the name area. Components will display their actor owner as the name */
+		ComponentsAndActorsUseNameArea = 1<<2 | ActorsUseNameArea,
 	};
 
 	/** Controls how CPF_DisableEditOnInstance nodes will be treated */
@@ -58,8 +58,8 @@ struct FDetailsViewArgs
 	FName ViewIdentifier;
 	/** Notify hook to call when properties are changed */
 	FNotifyHook* NotifyHook;
-	/** Settings for displaying the name area */
-	ENameAreaSettings NameAreaSettings;
+	/** Settings for displaying the name area (@see ENameAreaSettings) */
+	int32 NameAreaSettings;
 	/** The default column width */
 	float ColumnWidth;
 	/** True if the viewed objects updates from editor selection */
