Fixed wrong shade of background in the scrollbar well of the material instance details view.

diff --git a/Engine/Source/Editor/PropertyEditor/Private/SDetailSingleItemRow.cpp b/Engine/Source/Editor/PropertyEditor/Private/SDetailSingleItemRow.cpp
--- a/Engine/Source/Editor/PropertyEditor/Private/SDetailSingleItemRow.cpp
+++ b/Engine/Source/Editor/PropertyEditor/Private/SDetailSingleItemRow.cpp
@@ -521,7 +521,7 @@ void SDetailSingleItemRow::Construct( const FArguments& InArgs, FDetailLayoutCus
 	{
 		return SDetailTableRowBase::IsScrollBarVisible(OwnerTableViewWeak) ?
 			FSlateColor(EStyleColor::White) : 
-			this->GetInnerBackgroundColor();
+			this->GetOuterBackgroundColor();
 	};
 
 	this->ChildSlot
