Bitmask flags editors now support up to 64-bit bitmasks.

Reported on UDN.

diff --git a/Engine/Source/Editor/PropertyEditor/Private/UserInterface/PropertyEditor/SPropertyEditorNumeric.h b/Engine/Source/Editor/PropertyEditor/Private/UserInterface/PropertyEditor/SPropertyEditorNumeric.h
--- a/Engine/Source/Editor/PropertyEditor/Private/UserInterface/PropertyEditor/SPropertyEditorNumeric.h
+++ b/Engine/Source/Editor/PropertyEditor/Private/UserInterface/PropertyEditor/SPropertyEditorNumeric.h
@@ -74,7 +74,7 @@ public:
 				if (BitmaskEnum)
 				{
 					const bool bUseEnumValuessAsMaskValues = BitmaskEnum->GetBoolMetaData(PropertyEditorConstants::MD_UseEnumValuesAsMaskValuesInEditor);
-					auto AddNewBitmaskFlagLambda = [BitmaskEnum, &Result](int32 InEnumIndex, int32 InFlagValue)
+					auto AddNewBitmaskFlagLambda = [BitmaskEnum, &Result](int32 InEnumIndex, int64 InFlagValue)
 					{
 						Result.Emplace();
 						FBitmaskFlagInfo* BitmaskFlag = &Result.Last();
@@ -97,14 +97,14 @@ public:
 						{
 							if (bUseEnumValuessAsMaskValues)
 							{
-								if (EnumValue < MAX_int32 && FMath::IsPowerOfTwo(EnumValue))
+								if (EnumValue < MAX_int64 && FMath::IsPowerOfTwo(EnumValue))
 								{
-									AddNewBitmaskFlagLambda(BitmaskEnumIndex, static_cast<int32>(EnumValue));
+									AddNewBitmaskFlagLambda(BitmaskEnumIndex, EnumValue);
 								}
 							}
 							else if (EnumValue < BitmaskBitCount)
 							{
-								AddNewBitmaskFlagLambda(BitmaskEnumIndex, TBitmaskValueHelpers<NumericType>::LeftShift(static_cast<NumericType>(1), static_cast<int32>(EnumValue)));
+								AddNewBitmaskFlagLambda(BitmaskEnumIndex, TBitmaskValueHelpers<NumericType>::LeftShift(static_cast<NumericType>(1), EnumValue));
 							}
 						}
 					}
@@ -621,30 +621,33 @@ private:
 	};
 
 	/** Integral bitmask value helper methods. */
-	template<typename T, typename U = void>
+	template<typename T>
 	struct TBitmaskValueHelpers
 	{
 		static T BitwiseAND(T Base, T Mask) { return Base & Mask; }
 		static T BitwiseXOR(T Base, T Mask) { return Base ^ Mask; }
-		static T LeftShift(T Base, int32 Shift) { return Base << Shift; }
+		template <typename U>
+		static T LeftShift(T Base, U Shift) { return Base << Shift; }
 	};
 
 	/** Explicit specialization for numeric 'float' types (these will not be used). */
-	template<typename U>
-	struct TBitmaskValueHelpers<float, U>
+	template<>
+	struct TBitmaskValueHelpers<float>
 	{
 		static float BitwiseAND(float Base, float Mask) { return 0.0f; }
 		static float BitwiseXOR(float Base, float Mask) { return 0.0f; }
-		static float LeftShift(float Base, int32 Shift) { return 0.0f; }
+		template <typename U>
+		static float LeftShift(float Base, U Shift) { return 0.0f; }
 	};
 
 	/** Explicit specialization for numeric 'double' types (these will not be used). */
-	template<typename U>
-	struct TBitmaskValueHelpers<double, U>
+	template<>
+	struct TBitmaskValueHelpers<double>
 	{
 		static double BitwiseAND(double Base, double Mask) { return 0.0f; }
 		static double BitwiseXOR(double Base, double Mask) { return 0.0f; }
-		static double LeftShift(double Base, int32 Shift)  { return 0.0f; }
+		template <typename U>
+		static double LeftShift(double Base, U Shift)  { return 0.0f; }
 	};
 
 private:
