Removed margin from the array header that was causing array headers and struct headers to be misaligned.

diff --git a/Engine/Source/Editor/PropertyEditor/Private/UserInterface/PropertyEditor/SPropertyEditorArray.h b/Engine/Source/Editor/PropertyEditor/Private/UserInterface/PropertyEditor/SPropertyEditorArray.h
--- a/Engine/Source/Editor/PropertyEditor/Private/UserInterface/PropertyEditor/SPropertyEditorArray.h
+++ b/Engine/Source/Editor/PropertyEditor/Private/UserInterface/PropertyEditor/SPropertyEditorArray.h
@@ -52,7 +52,6 @@ public:
 			.Content()
 			[
 				SNew(STextBlock)
-				.Margin(FMargin(0.0f, 0.0f, 4.0f, 0.0f))
 				.Text(TextAttr)
 				.Font(InArgs._Font)
 			]
