Fixes LOD Section values creeping out of thier section in the Skeletal Mesh Editor

diff --git a/Engine/Source/Editor/PropertyEditor/Private/UserInterface/PropertyEditor/SPropertyEditorAsset.cpp b/Engine/Source/Editor/PropertyEditor/Private/UserInterface/PropertyEditor/SPropertyEditorAsset.cpp
--- a/Engine/Source/Editor/PropertyEditor/Private/UserInterface/PropertyEditor/SPropertyEditorAsset.cpp
+++ b/Engine/Source/Editor/PropertyEditor/Private/UserInterface/PropertyEditor/SPropertyEditorAsset.cpp
@@ -525,13 +525,16 @@ void SPropertyEditorAsset::Construct(const FArguments& InArgs, const TSharedPtr<
 			SNew(SVerticalBox)
 			+ SVerticalBox::Slot()
 			.VAlign(VAlign_Center)
+			.AutoHeight()
 			[
 				AssetComboButton.ToSharedRef()
 			]
 			+ SVerticalBox::Slot()
 			.VAlign(VAlign_Center)
+			.AutoHeight()
 			[
 				SAssignNew(CustomContentBox, SHorizontalBox)
+
 				+SHorizontalBox::Slot()
 				.AutoWidth()
 				[
