Fix crash in SPropertyEditorAsset due to null dereference.

diff --git a/Engine/Source/Editor/PropertyEditor/Private/UserInterface/PropertyEditor/SPropertyEditorAsset.cpp b/Engine/Source/Editor/PropertyEditor/Private/UserInterface/PropertyEditor/SPropertyEditorAsset.cpp
--- a/Engine/Source/Editor/PropertyEditor/Private/UserInterface/PropertyEditor/SPropertyEditorAsset.cpp
+++ b/Engine/Source/Editor/PropertyEditor/Private/UserInterface/PropertyEditor/SPropertyEditorAsset.cpp
@@ -392,12 +392,15 @@ void SPropertyEditorAsset::Construct(const FArguments& InArgs, const TSharedPtr<
 					// This might be a standalone struct inside an asset, check if the package's asset is a template type like a user struct
 					for (UPackage* Package : ParentPackages)
 					{
-						UObject* ParentAsset = Package->FindAssetInPackage();
-						if (ParentAsset && (ParentAsset->IsTemplate() || ParentAsset->IsA(UStruct::StaticClass())))
+						if (Package != nullptr)
 						{
-							IsEnabledAttribute.Set(false);
-							TooltipAttribute.Set(LOCTEXT("VariableHasDisableEditOnTemplate", "Editing this value in structure's defaults is not allowed"));
-							break;
+							UObject* ParentAsset = Package->FindAssetInPackage();
+							if (ParentAsset && (ParentAsset->IsTemplate() || ParentAsset->IsA(UStruct::StaticClass())))
+							{
+								IsEnabledAttribute.Set(false);
+								TooltipAttribute.Set(LOCTEXT("VariableHasDisableEditOnTemplate", "Editing this value in structure's defaults is not allowed"));
+								break;
+							}
 						}
 					}
 				}
