GetOptions now works correctly with multi-selection.

Added test actor to PropertyEditorTestObject.h to allow for testing actor things, like this GetOptions multiselection bug.

Fixed whitespace errors caused by mismerge.

diff --git a/Engine/Source/Editor/PropertyEditor/Private/PropertyHandleImpl.cpp b/Engine/Source/Editor/PropertyEditor/Private/PropertyHandleImpl.cpp
--- a/Engine/Source/Editor/PropertyEditor/Private/PropertyHandleImpl.cpp
+++ b/Engine/Source/Editor/PropertyEditor/Private/PropertyHandleImpl.cpp
@@ -3000,8 +3000,6 @@ bool FPropertyHandleBase::GeneratePossibleValues(TArray< TSharedPtr<FString> >&
 
 			if (OutObjects.Num() > 0)
 			{
-				FCachedPropertyPath Path(GetOptionsFunctionName);
-				
 				TArray<FString> OptionIntersection;
 				TSet<FString> OptionIntersectionSet;
 
@@ -3011,11 +3009,8 @@ bool FPropertyHandleBase::GeneratePossibleValues(TArray< TSharedPtr<FString> >&
 					{
 						FEditorScriptExecutionGuard ScriptExecutionGuard;
 
-						if (PropertyPathHelpers::GetPropertyValue(Target, Path, StringOptions))
-						{
-							// No-Op
-						}
-						else
+						FCachedPropertyPath Path(GetOptionsFunctionName);
+						if (!PropertyPathHelpers::GetPropertyValue(Target, Path, StringOptions))
 						{
 							TArray<FName> NameOptions;
 							if (PropertyPathHelpers::GetPropertyValue(Target, Path, NameOptions))
diff --git a/Engine/Source/Editor/UnrealEd/Classes/Editor/PropertyEditorTestObject.h b/Engine/Source/Editor/UnrealEd/Classes/Editor/PropertyEditorTestObject.h
--- a/Engine/Source/Editor/UnrealEd/Classes/Editor/PropertyEditorTestObject.h
+++ b/Engine/Source/Editor/UnrealEd/Classes/Editor/PropertyEditorTestObject.h
@@ -761,4 +761,16 @@ struct FTestSparseClassDataStorage
 
 	UPROPERTY(EditDefaultsOnly, Category=Default)
 	TMap<int, int> Map;
+};
+
+UCLASS()
+class APropertyEditorTestActor : public AActor
+{
+	GENERATED_BODY()
+
+	UPROPERTY(EditAnywhere, Category=Default, meta=(GetOptions=GetOptionsFunc))
+	FName GetOptionsValue;
+
+	UFUNCTION()
+	TArray<FString> GetOptionsFunc() const;
 };
\ No newline at end of file
diff --git a/Engine/Source/Editor/UnrealEd/Private/PropertyEditorTestObject.cpp b/Engine/Source/Editor/UnrealEd/Private/PropertyEditorTestObject.cpp
--- a/Engine/Source/Editor/UnrealEd/Private/PropertyEditorTestObject.cpp
+++ b/Engine/Source/Editor/UnrealEd/Private/PropertyEditorTestObject.cpp
@@ -10,4 +10,9 @@ bool UPropertyEditorTestObject::CanEditChange(const FProperty* InProperty) const
 	}
 
 	return true;
+}
+
+TArray<FString> APropertyEditorTestActor::GetOptionsFunc() const
+{
+	return TArray<FString> { TEXT("One"), TEXT("Two"), TEXT("Three") };
 }
\ No newline at end of file
