Fixed potential nullptr access to this in FPropertyHandleImpl::ExecuteCustomResetToDefault.

Fixed PreviewSceneCustomization causing reset to default to be deferred.

diff --git a/Engine/Source/Editor/Persona/Private/PreviewSceneCustomizations.cpp b/Engine/Source/Editor/Persona/Private/PreviewSceneCustomizations.cpp
--- a/Engine/Source/Editor/Persona/Private/PreviewSceneCustomizations.cpp
+++ b/Engine/Source/Editor/Persona/Private/PreviewSceneCustomizations.cpp
@@ -183,7 +183,7 @@ void FPreviewSceneDescriptionCustomization::CustomizeDetails(IDetailLayoutBuilde
 		.AddProperty(SkeletalMeshProperty)
 		.CustomWidget()
 		.OverrideResetToDefault(FResetToDefaultOverride::Create(
-			FIsResetToDefaultVisible::CreateLambda([this](TSharedPtr<IPropertyHandle> PropertyHandle) -> bool 
+			TAttribute<bool>::CreateLambda([this]()
 			{
 				if (PreviewScene.IsValid())
 				{
@@ -191,7 +191,7 @@ void FPreviewSceneDescriptionCustomization::CustomizeDetails(IDetailLayoutBuilde
 				}
 				return false;
 			}),
-			FResetToDefaultHandler::CreateLambda([this](TSharedPtr<IPropertyHandle> PropertyHandle) 
+			FSimpleDelegate::CreateLambda([this]() 
 			{
 				if (PreviewScene.IsValid())
 				{
diff --git a/Engine/Source/Editor/PropertyEditor/Private/PropertyHandleImpl.cpp b/Engine/Source/Editor/PropertyEditor/Private/PropertyHandleImpl.cpp
--- a/Engine/Source/Editor/PropertyEditor/Private/PropertyHandleImpl.cpp
+++ b/Engine/Source/Editor/PropertyEditor/Private/PropertyHandleImpl.cpp
@@ -3287,13 +3287,24 @@ bool FPropertyHandleBase::CanResetToDefault() const
 void FPropertyHandleBase::ExecuteCustomResetToDefault(const FResetToDefaultOverride& InOnCustomResetToDefault)
 {
 	// This action must be deferred until next tick so that we avoid accessing invalid data before we have a chance to tick
-	Implementation->GetPropertyUtilities()->EnqueueDeferredAction(FSimpleDelegate::CreateLambda([this, InOnCustomResetToDefault]() { OnCustomResetToDefault(InOnCustomResetToDefault); }));
+	TSharedPtr<IPropertyUtilities> PropertyUtilities = Implementation->GetPropertyUtilities();
+	if (PropertyUtilities.IsValid())
+	{
+		TWeakPtr<FPropertyHandleBase> ThisWeak = SharedThis(this);
+		PropertyUtilities->EnqueueDeferredAction(FSimpleDelegate::CreateLambda([ThisWeak, InOnCustomResetToDefault]()
+			{
+				TSharedPtr<FPropertyHandleBase> ThisPinned = ThisWeak.Pin();
+				if (ThisPinned.IsValid())
+				{
+					ThisPinned->OnCustomResetToDefault(InOnCustomResetToDefault);
+				}
+			}));
+	}
 }
 
 FName FPropertyHandleBase::GetDefaultCategoryName() const
 {
 	FProperty* Property = GetProperty();
-
 	if (Property)
 	{
 		return FObjectEditorUtils::GetCategoryFName(Property);
