Restored array drag and drop rearranging in details view arrays.

diff --git a/Engine/Source/Editor/PropertyEditor/Private/SDetailSingleItemRow.cpp b/Engine/Source/Editor/PropertyEditor/Private/SDetailSingleItemRow.cpp
--- a/Engine/Source/Editor/PropertyEditor/Private/SDetailSingleItemRow.cpp
+++ b/Engine/Source/Editor/PropertyEditor/Private/SDetailSingleItemRow.cpp
@@ -352,7 +352,7 @@ void SDetailSingleItemRow::Construct( const FArguments& InArgs, FDetailLayoutCus
 				{
 					ArrayDragDelegate = FOnTableRowDragEnter::CreateSP(this, &SDetailSingleItemRow::OnArrayDragEnter);
 					ArrayDragLeaveDelegate = FOnTableRowDragLeave::CreateSP(this, &SDetailSingleItemRow::OnArrayDragLeave);
-					ArrayDropDelegate = FOnTableRowDrop::CreateSP(this, &SDetailSingleItemRow::OnArrayHeaderDrop);
+					ArrayDropDelegate = FOnTableRowDrop::CreateSP(this, PropertyNode->IsReorderable() ? &SDetailSingleItemRow::OnArrayDrop : &SDetailSingleItemRow::OnArrayHeaderDrop);
 					ArrayAcceptDropDelegate = FOnCanAcceptDrop::CreateSP(this, &SDetailSingleItemRow::OnArrayCanAcceptDrop);
 				}
 			}
