Fix SPropertyEditorAsset to properly handle SoftObjectPtr for partitioned worlds by testing existence of the unloaded actors instead of failing if the world is loaded but not the actor.
Also hide the browse button on SoftObjectPtr properties when in a partitioned world, as it used to load the corresponding sublevel. This could be replaced with a "load referenced actor" button in the future.

diff --git a/Engine/Source/Editor/PropertyEditor/Private/UserInterface/PropertyEditor/SPropertyEditorAsset.cpp b/Engine/Source/Editor/PropertyEditor/Private/UserInterface/PropertyEditor/SPropertyEditorAsset.cpp
--- a/Engine/Source/Editor/PropertyEditor/Private/UserInterface/PropertyEditor/SPropertyEditorAsset.cpp
+++ b/Engine/Source/Editor/PropertyEditor/Private/UserInterface/PropertyEditor/SPropertyEditorAsset.cpp
@@ -592,7 +592,31 @@ void SPropertyEditorAsset::Construct(const FArguments& InArgs, const TSharedPtr<
 		];
 	}
 
-	if( InArgs._DisplayBrowse )
+	bool bDisplayBrowse = InArgs._DisplayBrowse;
+	if (bDisplayBrowse && bIsActor)
+	{
+		FObjectOrAssetData Value;
+		GetValue(Value);
+
+		UWorld* World = Value.Object ? CastChecked<AActor>(Value.Object)->GetWorld() : nullptr;
+
+		if (!World)
+		{			
+			FSoftObjectPath MapObjectPath = FSoftObjectPath(Value.ObjectPath.GetAssetPathName(), FString());
+
+			if (UObject* MapObject = MapObjectPath.ResolveObject())
+			{
+				World = CastChecked<UWorld>(MapObject);
+			}
+		}
+
+		if (World && World->GetWorldPartition())
+		{
+			bDisplayBrowse = false;
+		}
+	}
+
+	if( bDisplayBrowse )
 	{
 		ButtonBox->AddSlot()
 		.Padding( 2.0f, 0.0f )
@@ -677,9 +701,20 @@ SPropertyEditorAsset::EActorReferenceState SPropertyEditorAsset::GetActorReferen
 			// Get a path pointing to the owning map
 			FSoftObjectPath MapObjectPath = FSoftObjectPath(Value.ObjectPath.GetAssetPathName(), FString());
 
-			if (MapObjectPath.ResolveObject())
+			if (UObject* MapObject = MapObjectPath.ResolveObject())
 			{
-				// If the map is valid but the object is not
+				UWorld* World = CastChecked<UWorld>(MapObject);
+
+				// In a partitioned world, the world object will exist but the actor itself can be unloaded
+				if (World->GetWorldPartition())
+				{
+					UObject* Object = nullptr;
+					if (World->LoadSubobject(*Value.ObjectPath.GetSubPathString(), Object, /*bOnlyTestExistence*/true))
+					{
+						return EActorReferenceState::Exists;
+					}
+				}
+
 				return EActorReferenceState::Error;
 			}
 
@@ -862,6 +897,10 @@ FText SPropertyEditorAsset::OnGetToolTip() const
 			{
 				ToolTipText = FText::Format(LOCTEXT("BrokenActorReference", "Broken reference to Actor ID '{Actor}', it was deleted or renamed"), Args);
 			}
+			else if (State == EActorReferenceState::Exists)
+			{
+				ToolTipText = FText::Format(LOCTEXT("ExistsActorReference", "Unloaded reference to Actor ID '{Actor}', use World Partition editor to load its corresponding cells"), Args);
+			}
 			else if (State == EActorReferenceState::Unknown)
 			{
 				ToolTipText = FText::Format(LOCTEXT("UnknownActorReference", "Unloaded reference to Actor ID '{Actor}', use Browse to load level"), Args);
diff --git a/Engine/Source/Editor/PropertyEditor/Private/UserInterface/PropertyEditor/SPropertyEditorAsset.h b/Engine/Source/Editor/PropertyEditor/Private/UserInterface/PropertyEditor/SPropertyEditorAsset.h
--- a/Engine/Source/Editor/PropertyEditor/Private/UserInterface/PropertyEditor/SPropertyEditorAsset.h
+++ b/Engine/Source/Editor/PropertyEditor/Private/UserInterface/PropertyEditor/SPropertyEditorAsset.h
@@ -37,6 +37,8 @@ public:
 		Null,
 		// The pointed to actor is fully loaded in memory
 		Loaded,
+		// The pointed to actor exists because the pointed to map is loaded but not the actor
+		Exists,
 		// The pointed to actor is unknown because the pointed to map is not loaded 
 		Unknown,
 		// This is a known bad reference, the owning map is loaded but the actor does not exist
diff --git a/Engine/Source/Runtime/CoreUObject/Private/UObject/SoftObjectPath.cpp b/Engine/Source/Runtime/CoreUObject/Private/UObject/SoftObjectPath.cpp
--- a/Engine/Source/Runtime/CoreUObject/Private/UObject/SoftObjectPath.cpp
+++ b/Engine/Source/Runtime/CoreUObject/Private/UObject/SoftObjectPath.cpp
@@ -477,7 +477,7 @@ UObject* FSoftObjectPath::TryLoad(FUObjectSerializeContext* InLoadContext) const
 			// it knows how to load that specific object.
 			if (!LoadedObject && TopLevelObject)
 			{
-				LoadedObject = TopLevelObject->LoadSubobject(*SubPathString);
+				TopLevelObject->LoadSubobject(*SubPathString, LoadedObject);
 			}
 #endif
 		}
diff --git a/Engine/Source/Runtime/CoreUObject/Public/UObject/Object.h b/Engine/Source/Runtime/CoreUObject/Public/UObject/Object.h
--- a/Engine/Source/Runtime/CoreUObject/Public/UObject/Object.h
+++ b/Engine/Source/Runtime/CoreUObject/Public/UObject/Object.h
@@ -263,8 +263,15 @@ public:
 	 * Called to defer loading a subobject to its top-level container object. Usefull with the package override
 	 * feature, where an object might not be saved in the same package as its outer, especially if the the top-level
 	 * object does lazy loading of these objects.
+	 * 
+ 	 * @param	SubObjectPath		object path to handle, relative to the current implementer
+	 * @param	OutObject			receives the loaded object, or existing loaded object if bOnlyTestExistence is true and object is already loaded
+	 * @param	bOnlyTestExistence	if true, don't load the object, just return true if it exists
+	 * 
+	 * @return true if the object loaded or exists, depending on bOnlyTestExistence
+* 
 	 */
-	virtual UObject* LoadSubobject(const TCHAR* SubObjectPath) { return nullptr; }
+	virtual bool LoadSubobject(const TCHAR* SubObjectPath, UObject*& OutObject, bool bOnlyTestExistence = false) { return false; }
 #endif
 
 	/**
diff --git a/Engine/Source/Runtime/Engine/Classes/Engine/Level.h b/Engine/Source/Runtime/Engine/Classes/Engine/Level.h
--- a/Engine/Source/Runtime/Engine/Classes/Engine/Level.h
+++ b/Engine/Source/Runtime/Engine/Classes/Engine/Level.h
@@ -785,7 +785,7 @@ public:
 	virtual void PostEditChangeProperty(FPropertyChangedEvent& PropertyChangedEvent) override;
 	virtual void BeginCacheForCookedPlatformData(const ITargetPlatform *TargetPlatform) override;
 	virtual bool CanEditChange(const FProperty* PropertyThatWillChange) const override;
-	virtual UObject* LoadSubobject(const TCHAR* SubObjectPath) override;
+	virtual bool LoadSubobject(const TCHAR* SubObjectPath, UObject*& OutObject, bool bOnlyTestExistence) override;
 #endif // WITH_EDITOR
 	virtual void PostLoad() override;
 	virtual void PreSave(const class ITargetPlatform* TargetPlatform) override;
diff --git a/Engine/Source/Runtime/Engine/Classes/Engine/World.h b/Engine/Source/Runtime/Engine/Classes/Engine/World.h
--- a/Engine/Source/Runtime/Engine/Classes/Engine/World.h
+++ b/Engine/Source/Runtime/Engine/Classes/Engine/World.h
@@ -2613,7 +2613,7 @@ public:
 	virtual bool Rename(const TCHAR* NewName = NULL, UObject* NewOuter = NULL, ERenameFlags Flags = REN_None) override;
 	virtual void GetAssetRegistryTags(TArray<FAssetRegistryTag>& OutTags) const override;
 	virtual bool IsNameStableForNetworking() const override;
-	virtual UObject* LoadSubobject(const TCHAR* SubObjectPath) override;
+	virtual bool LoadSubobject(const TCHAR* SubObjectPath, UObject*& OutObject, bool bOnlyTestExistence) override;
 #endif
 	virtual void PostDuplicate(bool bDuplicateForPIE) override;
 	//~ End UObject Interface
diff --git a/Engine/Source/Runtime/Engine/Private/Level.cpp b/Engine/Source/Runtime/Engine/Private/Level.cpp
--- a/Engine/Source/Runtime/Engine/Private/Level.cpp
+++ b/Engine/Source/Runtime/Engine/Private/Level.cpp
@@ -2633,19 +2633,21 @@ bool ULevel::CanEditChange(const FProperty* PropertyThatWillChange) const
 	return Super::CanEditChange(PropertyThatWillChange);
 }
 
-UObject* ULevel::LoadSubobject(const TCHAR* SubObjectPath)
+bool ULevel::LoadSubobject(const TCHAR* SubObjectPath, UObject*& OutObject, bool bOnlyTestExistence)
 {
-	UObject* LoadedObject = StaticFindObject(nullptr, this, SubObjectPath);
+	OutObject = StaticFindObject(nullptr, this, SubObjectPath);
 
-	if (!LoadedObject)
+	if (!OutObject)
 	{
 		if (UWorldPartition* WorldPartition = GetWorldPartition())
 		{
-			LoadedObject = WorldPartition->LoadSubobject(SubObjectPath);
+			return WorldPartition->LoadSubobject(SubObjectPath, OutObject, bOnlyTestExistence);
 		}
+
+		return false;
 	}
 
-	return LoadedObject;
+	return true;
 }
 
 void ULevel::FixupForPIE(int32 InPIEInstanceID, TFunctionRef<void(int32, FSoftObjectPath&)> InCustomFixupFunction)
diff --git a/Engine/Source/Runtime/Engine/Private/World.cpp b/Engine/Source/Runtime/Engine/Private/World.cpp
--- a/Engine/Source/Runtime/Engine/Private/World.cpp
+++ b/Engine/Source/Runtime/Engine/Private/World.cpp
@@ -8002,21 +8002,20 @@ bool UWorld::IsNameStableForNetworking() const
 	return bIsNameStableForNetworking || Super::IsNameStableForNetworking();
 }
 
-UObject* UWorld::LoadSubobject(const TCHAR* SubObjectPath)
+bool UWorld::LoadSubobject(const TCHAR* SubObjectPath, UObject*& OutObject, bool bOnlyTestExistence)
 {
-	UObject* LoadedObject = nullptr;
-
 	FString SubObjectName;
 	FString SubObjectContext;	
 	if (FString(SubObjectPath).Split(TEXT("."), &SubObjectContext, &SubObjectName))
 	{
 		if (UObject* SubObject = StaticFindObject(nullptr, this, *SubObjectContext))
 		{
-			LoadedObject = SubObject->LoadSubobject(*SubObjectName);
+			return SubObject->LoadSubobject(*SubObjectName, OutObject, bOnlyTestExistence);
 		}
 	}
 
-	return LoadedObject;
+	OutObject = nullptr;
+	return false;
 }
 #endif
 
diff --git a/Engine/Source/Runtime/Engine/Private/WorldPartition/WorldPartition.cpp b/Engine/Source/Runtime/Engine/Private/WorldPartition/WorldPartition.cpp
--- a/Engine/Source/Runtime/Engine/Private/WorldPartition/WorldPartition.cpp
+++ b/Engine/Source/Runtime/Engine/Private/WorldPartition/WorldPartition.cpp
@@ -878,7 +878,7 @@ void UWorldPartition::BeginDestroy()
 }
 
 #if WITH_EDITOR
-UObject* UWorldPartition::LoadSubobject(const TCHAR* SubObjectPath)
+bool UWorldPartition::LoadSubobject(const TCHAR* SubObjectPath, UObject*& OutObject, bool bOnlyTestExistence)
 {
 	for (UActorDescContainer::TIterator<> ActorDescIterator(this); ActorDescIterator; ++ActorDescIterator)
 	{
@@ -886,12 +886,18 @@ UObject* UWorldPartition::LoadSubobject(const TCHAR* SubObjectPath)
 
 		if (FString(ActorDesc->ActorPath.ToString()).EndsWith(SubObjectPath))
 		{
-			int32 ReferenceIndex = LoadedSubobjects.Emplace(this, ActorDesc->GetGuid());
-			return LoadedSubobjects[ReferenceIndex]->GetActor();
+			if (!bOnlyTestExistence)
+			{
+				LoadedSubobjects.Emplace(this, ActorDesc->GetGuid());
+			}
+
+			OutObject = ActorDescIterator->GetActor();
+
+			return true;
 		}
 	}
 
-	return nullptr;
+	return false;
 }
 #endif
 
diff --git a/Engine/Source/Runtime/Engine/Public/WorldPartition/WorldPartition.h b/Engine/Source/Runtime/Engine/Public/WorldPartition/WorldPartition.h
--- a/Engine/Source/Runtime/Engine/Public/WorldPartition/WorldPartition.h
+++ b/Engine/Source/Runtime/Engine/Public/WorldPartition/WorldPartition.h
@@ -122,7 +122,7 @@ private:
 #if WITH_EDITOR
 public:
 	//~ Begin UObject Interface
-	virtual UObject* LoadSubobject(const TCHAR* SubObjectPath) override;
+	virtual bool LoadSubobject(const TCHAR* SubObjectPath, UObject*& OutObject, bool bOnlyTestExistence) override;
 	//~ End UObject Interface
 
 	FName GetWorldPartitionEditorName();
