Fix crash when selecting child actors in a BP class.

diff --git a/Engine/Source/Editor/PropertyEditor/Private/UserInterface/PropertyEditor/SPropertyEditorAsset.cpp b/Engine/Source/Editor/PropertyEditor/Private/UserInterface/PropertyEditor/SPropertyEditorAsset.cpp
--- a/Engine/Source/Editor/PropertyEditor/Private/UserInterface/PropertyEditor/SPropertyEditorAsset.cpp
+++ b/Engine/Source/Editor/PropertyEditor/Private/UserInterface/PropertyEditor/SPropertyEditorAsset.cpp
@@ -606,7 +606,7 @@ void SPropertyEditorAsset::Construct(const FArguments& InArgs, const TSharedPtr<
 
 			if (UObject* MapObject = MapObjectPath.ResolveObject())
 			{
-				World = CastChecked<UWorld>(MapObject);
+				World = Cast<UWorld>(MapObject);
 			}
 		}
 
@@ -703,10 +703,10 @@ SPropertyEditorAsset::EActorReferenceState SPropertyEditorAsset::GetActorReferen
 
 			if (UObject* MapObject = MapObjectPath.ResolveObject())
 			{
-				UWorld* World = CastChecked<UWorld>(MapObject);
+				UWorld* World = Cast<UWorld>(MapObject);
 
 				// In a partitioned world, the world object will exist but the actor itself can be unloaded
-				if (World->GetWorldPartition())
+				if (World && World->GetWorldPartition())
 				{
 					UObject* Object = nullptr;
 					if (World->LoadSubobject(*Value.ObjectPath.GetSubPathString(), Object, /*bOnlyTestExistence*/true))
