Resetting Light Units to default in point lights now functions correctly.

The issue is that FPropertyHandleBase does not have a very long lifetime, so this instance was getting ref counted and deleted when trying to be used the next frame.

diff --git a/Engine/Source/Editor/PropertyEditor/Private/PropertyHandleImpl.cpp b/Engine/Source/Editor/PropertyEditor/Private/PropertyHandleImpl.cpp
--- a/Engine/Source/Editor/PropertyEditor/Private/PropertyHandleImpl.cpp
+++ b/Engine/Source/Editor/PropertyEditor/Private/PropertyHandleImpl.cpp
@@ -3290,13 +3290,12 @@ void FPropertyHandleBase::ExecuteCustomResetToDefault(const FResetToDefaultOverr
 	TSharedPtr<IPropertyUtilities> PropertyUtilities = Implementation->GetPropertyUtilities();
 	if (PropertyUtilities.IsValid())
 	{
-		TWeakPtr<FPropertyHandleBase> ThisWeak = SharedThis(this);
-		PropertyUtilities->EnqueueDeferredAction(FSimpleDelegate::CreateLambda([ThisWeak, InOnCustomResetToDefault]()
+		TSharedPtr<FPropertyHandleBase> ThisShared = SharedThis(this);
+		PropertyUtilities->EnqueueDeferredAction(FSimpleDelegate::CreateLambda([ThisShared, InOnCustomResetToDefault]()
 			{
-				TSharedPtr<FPropertyHandleBase> ThisPinned = ThisWeak.Pin();
-				if (ThisPinned.IsValid())
+				if (ThisShared.IsValid())
 				{
-					ThisPinned->OnCustomResetToDefault(InOnCustomResetToDefault);
+					ThisShared->OnCustomResetToDefault(InOnCustomResetToDefault);
 				}
 			}));
 	}
