Cleanup unsed details panel category brush logic

diff --git a/Engine/Source/Editor/PropertyEditor/Private/DetailCategoryGroupNode.cpp b/Engine/Source/Editor/PropertyEditor/Private/DetailCategoryGroupNode.cpp
--- a/Engine/Source/Editor/PropertyEditor/Private/DetailCategoryGroupNode.cpp
+++ b/Engine/Source/Editor/PropertyEditor/Private/DetailCategoryGroupNode.cpp
@@ -121,16 +121,16 @@ const FSlateBrush* SDetailCategoryTableRow::GetBackgroundImage() const
 	{
 		if (bIsInnerCategory)
 		{
-			return FEditorStyle::GetBrush("DetailsView.CategoryMiddle");
+			return FAppStyle::Get().GetBrush("DetailsView.CategoryMiddle");
 		}
 
 		if (IsHovered())
 		{
-			return IsItemExpanded() ? FEditorStyle::GetBrush("DetailsView.CategoryTop_Hovered") : FEditorStyle::GetBrush("DetailsView.CollapsedCategory_Hovered");
+			return FAppStyle::Get().GetBrush("DetailsView.CategoryTop_Hovered");
 		}
 		else
 		{
-			return IsItemExpanded() ? FEditorStyle::GetBrush("DetailsView.CategoryTop") : FEditorStyle::GetBrush("DetailsView.CollapsedCategory");
+			return FAppStyle::Get().GetBrush("DetailsView.CategoryTop");
 		}
 	}
 
