FDetailsViewArgs::RightColumnMinWidth property added to allow different views to customize how much space is reserved for the right column.

FDetailsViewArgs::ColumnWidth is now respected.

diff --git a/Engine/Source/Editor/PropertyEditor/Private/SDetailSingleItemRow.cpp b/Engine/Source/Editor/PropertyEditor/Private/SDetailSingleItemRow.cpp
--- a/Engine/Source/Editor/PropertyEditor/Private/SDetailSingleItemRow.cpp
+++ b/Engine/Source/Editor/PropertyEditor/Private/SDetailSingleItemRow.cpp
@@ -483,7 +483,7 @@ void SDetailSingleItemRow::Construct( const FArguments& InArgs, FDetailLayoutCus
 			Splitter->AddSlot()
 				.Value(ColumnSizeData.RightColumnWidth)
 				.OnSlotResized(ColumnSizeData.OnRightColumnResized)
-				.MinSize(50)
+				.MinSize(ColumnSizeData.RightColumnMinWidth)
 			[
 				SNew(SBorder)
 				.BorderImage(FAppStyle::Get().GetBrush("DetailsView.CategoryMiddle"))
diff --git a/Engine/Source/Editor/PropertyEditor/Private/SDetailsView.cpp b/Engine/Source/Editor/PropertyEditor/Private/SDetailsView.cpp
--- a/Engine/Source/Editor/PropertyEditor/Private/SDetailsView.cpp
+++ b/Engine/Source/Editor/PropertyEditor/Private/SDetailsView.cpp
@@ -41,6 +41,9 @@ void SDetailsView::Construct(const FArguments& InArgs, const FDetailsViewArgs& I
 {
 	DetailsViewArgs = InDetailsViewArgs;
 
+	ColumnSizeData.SetValueColumnWidth(DetailsViewArgs.ColumnWidth);
+	ColumnSizeData.RightColumnMinWidth = DetailsViewArgs.RightColumnMinWidth;
+
 	SetObjectFilter(InDetailsViewArgs.ObjectFilter);
 
 	bViewingClassDefaultObject = false;
diff --git a/Engine/Source/Editor/PropertyEditor/Public/IDetailPropertyRow.h b/Engine/Source/Editor/PropertyEditor/Public/IDetailPropertyRow.h
--- a/Engine/Source/Editor/PropertyEditor/Public/IDetailPropertyRow.h
+++ b/Engine/Source/Editor/PropertyEditor/Public/IDetailPropertyRow.h
@@ -21,7 +21,7 @@ public:
 	FDetailColumnSizeData()
 	{
 		ValueColumnWidthValue = 0.7f;
-		RightColumnWidthValue = 0.05f;
+		RightColumnWidthValue = 0;
 		HoveredSplitterIndexValue = INDEX_NONE;
 
 		NameColumnWidth.BindRaw(this, &FDetailColumnSizeData::GetNameColumnWidth);
@@ -29,7 +29,7 @@ public:
 		PropertyColumnWidth.BindRaw(this, &FDetailColumnSizeData::GetPropertyColumnWidth);
 		RightColumnWidth.BindRaw(this, &FDetailColumnSizeData::GetRightColumnWidth);
 		HoveredSplitterIndex.BindRaw(this, &FDetailColumnSizeData::GetHoveredSplitterIndex);
-		OnValueColumnResized.BindRaw(this, &FDetailColumnSizeData::OnSetValueColumnWidth);
+		OnValueColumnResized.BindRaw(this, &FDetailColumnSizeData::SetValueColumnWidth);
 		OnRightColumnResized.BindRaw(this, &FDetailColumnSizeData::OnSetRightColumnWidth);
 		OnSplitterHandleHovered.BindRaw(this, &FDetailColumnSizeData::OnSetHoveredSplitterIndex);
 
@@ -38,6 +38,8 @@ public:
 		OnNameColumnResized.BindLambda([](float) {}); 
 	}
 
+	float RightColumnMinWidth;
+
 	TAttribute<float> NameColumnWidth;
 	TAttribute<float> ValueColumnWidth;
 	TAttribute<float> PropertyColumnWidth;
@@ -49,6 +51,12 @@ public:
 	SSplitter::FOnSlotResized OnRightColumnResized;
 	SSplitter::FOnHandleHovered OnSplitterHandleHovered;
 
+	void SetValueColumnWidth(float NewWidth)
+	{ 
+		ensure(NewWidth <= 1.0f); 
+		ValueColumnWidthValue = NewWidth;
+	} 
+
 private:
 	float ValueColumnWidthValue;
 	float RightColumnWidthValue;
@@ -60,12 +68,7 @@ private:
 	float GetPropertyColumnWidth() const { return 1.0f - RightColumnWidthValue; }
 	int32 GetHoveredSplitterIndex() const { return HoveredSplitterIndexValue; }
 
-	void OnSetValueColumnWidth(float NewWidth)
-	{ 
-		ensure(NewWidth <= 1.0f); 
-		ValueColumnWidthValue = NewWidth;
-	} 
-
+	
 	void OnSetRightColumnWidth(float NewWidth)
 	{
 		ensure(NewWidth <= 1.0f);
