Added RightColumnMinWidth to FDetailsViewArgs to control the width of the right column in a new details view. Details views that will not have much in the way of extension widgets (eg. modes panels) can opt to have this value be smaller than the default of 50.

diff --git a/Engine/Source/Editor/PropertyEditor/Public/IDetailsView.h b/Engine/Source/Editor/PropertyEditor/Public/IDetailsView.h
--- a/Engine/Source/Editor/PropertyEditor/Public/IDetailsView.h
+++ b/Engine/Source/Editor/PropertyEditor/Public/IDetailsView.h
@@ -62,6 +62,8 @@ struct FDetailsViewArgs
 	int32 NameAreaSettings;
 	/** The default column width */
 	float ColumnWidth;
+	/** The minimum width of the right column in Slate units. */
+	float RightColumnMinWidth;
 	/** True if the viewed objects updates from editor selection */
 	bool bUpdatesFromSelection : 1;
 	/** True if this property view can be locked */
@@ -109,6 +111,7 @@ public:
 		, NotifyHook(nullptr)
 		, NameAreaSettings(ActorsUseNameArea)
 		, ColumnWidth(.65f)
+		, RightColumnMinWidth(50)
 		, bUpdatesFromSelection(false)
 		, bLockable(false)
 		, bAllowSearch(true)
