Removed code dealing with FPropertyEditor having property children, because the code is unused.

diff --git a/Engine/Source/Editor/PropertyEditor/Private/Presentation/PropertyEditor/PropertyEditor.cpp b/Engine/Source/Editor/PropertyEditor/Private/Presentation/PropertyEditor/PropertyEditor.cpp
--- a/Engine/Source/Editor/PropertyEditor/Private/Presentation/PropertyEditor/PropertyEditor.cpp
+++ b/Engine/Source/Editor/PropertyEditor/Private/Presentation/PropertyEditor/PropertyEditor.cpp
@@ -42,8 +42,6 @@ FPropertyEditor::FPropertyEditor( const TSharedRef<FPropertyNode>& InPropertyNod
 	// FPropertyEditor isn't built to handle CategoryNodes
 	check( InPropertyNode->AsCategoryNode() == NULL );
 
-	FProperty* Property = InPropertyNode->GetProperty();
-
 	PropertyHandle = PropertyEditorHelpers::GetPropertyHandle( InPropertyNode, PropertyUtilities->GetNotifyHook(), PropertyUtilities );
 	check( PropertyHandle.IsValid() && PropertyHandle->IsValidHandle() );
 }
@@ -585,21 +583,6 @@ void FPropertyEditor::RequestRefresh()
 	PropertyUtilities->RequestRefresh();
 }
 
-void FPropertyEditor::AddPropertyEditorChild( const TSharedRef<FPropertyEditor>& Child )
-{
-	ChildPropertyEditors.Add( Child );
-}
-
-void FPropertyEditor::RemovePropertyEditorChild( const TSharedRef<FPropertyEditor>& Child )
-{
-	ChildPropertyEditors.Remove( Child );
-}
-
-const TArray< TSharedRef< FPropertyEditor > >& FPropertyEditor::GetPropertyEditorChildren() const
-{
-	return ChildPropertyEditors;
-}
-
 TSharedRef< FPropertyNode > FPropertyEditor::GetPropertyNode() const
 {
 	return PropertyNode;
diff --git a/Engine/Source/Editor/PropertyEditor/Private/Presentation/PropertyEditor/PropertyEditor.h b/Engine/Source/Editor/PropertyEditor/Private/Presentation/PropertyEditor/PropertyEditor.h
--- a/Engine/Source/Editor/PropertyEditor/Private/Presentation/PropertyEditor/PropertyEditor.h
+++ b/Engine/Source/Editor/PropertyEditor/Private/Presentation/PropertyEditor/PropertyEditor.h
@@ -73,10 +73,6 @@ public:
 	/**	@return Whether the property passes the current filter restrictions. If no there are no filter restrictions false will be returned. */
 	bool DoesPassFilterRestrictions() const;
 
-	void AddPropertyEditorChild( const TSharedRef<FPropertyEditor>& Child );
-	void RemovePropertyEditorChild( const TSharedRef<FPropertyEditor>& Child );
-	const TArray< TSharedRef< FPropertyEditor > >& GetPropertyEditorChildren() const;
-
 	void UseSelected();
 	void AddItem();
 	void AddGivenItem(const FString& InGivenItem);
@@ -116,8 +112,6 @@ private:
 
 private:
 
-	TArray< TSharedRef< FPropertyEditor > >	ChildPropertyEditors;
-
 	/** Property handle for actually reading/writing the value of a property */
 	TSharedPtr< class IPropertyHandle > PropertyHandle;
 
