Added SetIsCustomRowReadOnlyDelegate to enable customizing the read-only state of custom rows en-masse.

Removed SetIsCustomRowVisibilityFilteredDelegate, because it was unused and needlessly complex since the delegate that you pass into SetIsCustomRowVisibleDelegate can run the same logic and early out.

Category and property filtering through Properties To Show and Categories to Show now works properly.The problem was that visiblity was only considered once on creation and then never again.

diff --git a/Engine/Source/Editor/PropertyEditor/Private/CustomChildBuilder.cpp b/Engine/Source/Editor/PropertyEditor/Private/CustomChildBuilder.cpp
--- a/Engine/Source/Editor/PropertyEditor/Private/CustomChildBuilder.cpp
+++ b/Engine/Source/Editor/PropertyEditor/Private/CustomChildBuilder.cpp
@@ -32,13 +32,6 @@ FDetailWidgetRow& FCustomChildrenBuilder::AddCustomRow( const FText& SearchStrin
 	FDetailLayoutCustomization NewCustomization;
 
 	NewRow->FilterString( SearchString );
-	
-	IDetailsViewPrivate* DetailsView = ParentCategory.Pin()->GetDetailsView();
-	if (DetailsView && DetailsView->IsCustomRowVisibilityFiltered() && !DetailsView->IsCustomRowVisible(FName(*SearchString.ToString()), FName(*GetParentCategory().GetDisplayName().ToString())))
-	{
-		NewRow->Visibility(TAttribute<EVisibility>(EVisibility::Collapsed));
-	}
-
 	NewCustomization.WidgetDecl = NewRow;
 
 	ChildCustomizations.Add( NewCustomization );
diff --git a/Engine/Source/Editor/PropertyEditor/Private/DetailAdvancedDropdownNode.cpp b/Engine/Source/Editor/PropertyEditor/Private/DetailAdvancedDropdownNode.cpp
--- a/Engine/Source/Editor/PropertyEditor/Private/DetailAdvancedDropdownNode.cpp
+++ b/Engine/Source/Editor/PropertyEditor/Private/DetailAdvancedDropdownNode.cpp
@@ -171,3 +171,13 @@ FReply FAdvancedDropdownNode::OnAdvancedDropDownClicked()
 
 	return FReply::Handled();
 }
+
+ENodeVisibility FAdvancedDropdownNode::GetVisibility() const
+{
+	IDetailsViewPrivate* DetailsView = GetDetailsView();
+	if (!DetailsView || DetailsView->IsCustomRowVisible(FName(NAME_None), FName(*GetParentCategory()->GetDisplayName().ToString())))
+	{
+		return ENodeVisibility::Visible;
+	}
+	return ENodeVisibility::ForcedHidden;
+}
\ No newline at end of file
diff --git a/Engine/Source/Editor/PropertyEditor/Private/DetailAdvancedDropdownNode.h b/Engine/Source/Editor/PropertyEditor/Private/DetailAdvancedDropdownNode.h
--- a/Engine/Source/Editor/PropertyEditor/Private/DetailAdvancedDropdownNode.h
+++ b/Engine/Source/Editor/PropertyEditor/Private/DetailAdvancedDropdownNode.h
@@ -39,7 +39,7 @@ private:
 	virtual void GetChildren(FDetailNodeList& OutChildren) override {}
 	virtual void OnItemExpansionChanged( bool bIsExpanded, bool bShouldSaveState) override {}
 	virtual bool ShouldBeExpanded() const override { return false; }
-	virtual ENodeVisibility GetVisibility() const override { return ENodeVisibility::Visible; }
+	virtual ENodeVisibility GetVisibility() const override;
 	virtual void FilterNode( const FDetailFilter& InFilter ) override {}
 	virtual void Tick( float DeltaTime ) override {}
 	virtual bool ShouldShowOnlyChildren() const override { return false; }
diff --git a/Engine/Source/Editor/PropertyEditor/Private/DetailCategoryBuilderImpl.cpp b/Engine/Source/Editor/PropertyEditor/Private/DetailCategoryBuilderImpl.cpp
--- a/Engine/Source/Editor/PropertyEditor/Private/DetailCategoryBuilderImpl.cpp
+++ b/Engine/Source/Editor/PropertyEditor/Private/DetailCategoryBuilderImpl.cpp
@@ -151,12 +151,6 @@ FDetailWidgetRow& FDetailCategoryImpl::AddCustomRow(const FText& FilterString, b
 	NewCustomization.WidgetDecl = MakeShareable(new FDetailWidgetRow);
 	NewCustomization.WidgetDecl->FilterString(FilterString);
 
-	IDetailsViewPrivate* DetailsView = GetDetailsView();
-	if (DetailsView && DetailsView->IsCustomRowVisibilityFiltered() && !GetDetailsView()->IsCustomRowVisible(FName(*FilterString.ToString()), FName(*DisplayName.ToString())))
-	{
-		NewCustomization.WidgetDecl->Visibility(TAttribute<EVisibility>(EVisibility::Collapsed));
-	}
-
 	AddCustomLayout(NewCustomization, bForAdvanced);
 
 	return *NewCustomization.WidgetDecl;
diff --git a/Engine/Source/Editor/PropertyEditor/Private/DetailItemNode.cpp b/Engine/Source/Editor/PropertyEditor/Private/DetailItemNode.cpp
--- a/Engine/Source/Editor/PropertyEditor/Private/DetailItemNode.cpp
+++ b/Engine/Source/Editor/PropertyEditor/Private/DetailItemNode.cpp
@@ -57,19 +57,7 @@ void FDetailItemNode::Initialize()
 		SetExpansionState(bShouldExpand, bSaveState);
 	}
 
-	// Cache the visibility of customizations that can set it
-	if( Customization.HasCustomWidget() )
-	{	
-		CachedItemVisibility = Customization.WidgetDecl->VisibilityAttr.Get();
-	}
-	else if( Customization.HasPropertyNode() )
-	{
-		CachedItemVisibility = Customization.PropertyRow->GetPropertyVisibility();
-	}
-	else if( Customization.HasGroup() )
-	{
-		CachedItemVisibility = Customization.DetailGroup->GetGroupVisibility();
-	}
+	CachedItemVisibility = ComputeItemVisibility();
 
 	const bool bUpdateFilteredNodes = false;
 	GenerateChildren( bUpdateFilteredNodes );
@@ -173,7 +161,6 @@ void FDetailItemNode::InitCustomBuilder()
 		const bool bSaveState = false;
 		SetExpansionState(ParentCategory.Pin()->GetSavedExpansionState(*this), bSaveState);
 	}
-
 }
 
 void FDetailItemNode::InitGroup()
@@ -433,14 +420,14 @@ bool FDetailItemNode::ShouldBeExpanded() const
 
 ENodeVisibility FDetailItemNode::GetVisibility() const
 {
-	ENodeVisibility Visibility;
+	ENodeVisibility Visibility = CachedItemVisibility == EVisibility::Collapsed ? ENodeVisibility::ForcedHidden : ENodeVisibility::Visible;
 	if(Customization.IsHidden())
 	{
 		Visibility = ENodeVisibility::ForcedHidden;
 	}
 	else
 	{
-		Visibility = (bShouldBeVisibleDueToFiltering || bShouldBeVisibleDueToChildFiltering) ? ENodeVisibility::Visible : ENodeVisibility::HiddenDueToFiltering;
+		Visibility = (bShouldBeVisibleDueToFiltering || bShouldBeVisibleDueToChildFiltering) ? Visibility : ENodeVisibility::HiddenDueToFiltering;
 	}
 	return Visibility;
 }
@@ -626,19 +613,7 @@ void FDetailItemNode::Tick( float DeltaTime )
 		}
 
 		// Recache visibility
-		EVisibility NewVisibility;
-		if( Customization.HasCustomWidget() )
-		{	
-			NewVisibility = Customization.WidgetDecl->VisibilityAttr.Get();
-		}
-		else if( Customization.HasPropertyNode() )
-		{
-			NewVisibility = Customization.PropertyRow->GetPropertyVisibility();
-		}
-		else if( Customization.HasGroup() )
-		{
-			NewVisibility = Customization.DetailGroup->GetGroupVisibility();
-		}
+		EVisibility NewVisibility = ComputeItemVisibility();
 	
 		if( CachedItemVisibility != NewVisibility )
 		{
@@ -650,6 +625,44 @@ void FDetailItemNode::Tick( float DeltaTime )
 	}
 }
 
+EVisibility FDetailItemNode::ComputeItemVisibility() const
+{
+	EVisibility NewVisibility;
+	if (Customization.HasCustomWidget())
+	{	
+		NewVisibility = Customization.WidgetDecl->VisibilityAttr.Get();
+
+		IDetailsViewPrivate* DetailsView = GetDetailsView();
+		if (DetailsView && !DetailsView->IsCustomRowVisible(FName(*Customization.WidgetDecl->FilterTextString.ToString()), FName(*GetParentCategory()->GetDisplayName().ToString())))
+		{
+			NewVisibility = EVisibility::Collapsed;
+		}
+	}
+	else if (Customization.HasPropertyNode())
+	{
+		NewVisibility = Customization.PropertyRow->GetPropertyVisibility();
+	}
+	else if (Customization.HasGroup())
+	{
+		NewVisibility = Customization.DetailGroup->GetGroupVisibility();
+	}
+	else if (Customization.HasCustomBuilder())
+	{
+		NewVisibility = EVisibility::Collapsed;
+
+		for (TSharedRef<FDetailTreeNode> Child : Children)
+		{
+			if (Child->GetVisibility() == ENodeVisibility::Visible)
+			{
+				NewVisibility = EVisibility::Visible;
+				break;
+			}
+		}
+	}
+
+	return NewVisibility;
+}
+
 bool FDetailItemNode::ShouldShowOnlyChildren() const
 {
 	// Show only children of this node if there is no content for custom details or the property node requests that only children be shown
@@ -696,6 +709,25 @@ FPropertyPath FDetailItemNode::GetPropertyPath() const
 	return Ret;
 }
 
+TAttribute<bool> FDetailItemNode::IsPropertyEditingEnabled() const
+{
+	return TAttribute<bool>::Create([this]()
+		{
+			bool IsParentEnabledValue = IsParentEnabled.Get(true);
+			if (Customization.HasCustomWidget())
+			{
+				IDetailsViewPrivate* DetailsView = GetDetailsView();
+				if (DetailsView)
+				{
+					return IsParentEnabledValue && 
+						!DetailsView->IsCustomRowReadOnly(FName(*Customization.WidgetDecl->FilterTextString.ToString()), FName(*GetParentCategory()->GetDisplayName().ToString()));
+				}
+			}
+
+			return IsParentEnabledValue;
+		});
+}
+
 TSharedPtr<FPropertyNode> FDetailItemNode::GetPropertyNode() const
 {
 	return Customization.GetPropertyNode();
diff --git a/Engine/Source/Editor/PropertyEditor/Private/DetailItemNode.h b/Engine/Source/Editor/PropertyEditor/Private/DetailItemNode.h
--- a/Engine/Source/Editor/PropertyEditor/Private/DetailItemNode.h
+++ b/Engine/Source/Editor/PropertyEditor/Private/DetailItemNode.h
@@ -52,6 +52,11 @@ public:
 	 */
 	bool HasGeneratedChildren() const { return Children.Num() > 0;}
 
+	/**
+	 * @return The new, uncached visibility of this item.
+	 */
+	EVisibility ComputeItemVisibility() const;
+
 	/** FDetailTreeNode interface */
 	virtual IDetailsView* GetNodeDetailsView() const override { return ParentCategory.Pin()->GetNodeDetailsView(); }
 	virtual IDetailsViewPrivate* GetDetailsView() const override{ return ParentCategory.Pin()->GetDetailsView(); }
@@ -70,7 +75,7 @@ public:
 	virtual void SetIsHighlighted(bool bInIsHighlighted) override { bIsHighlighted = bInIsHighlighted; }
 	virtual bool IsHighlighted() const override { return bIsHighlighted; }
 	virtual bool IsLeaf() override { return Children.Num() == 0;  }
-	virtual TAttribute<bool> IsPropertyEditingEnabled() const override { return IsParentEnabled; }
+	virtual TAttribute<bool> IsPropertyEditingEnabled() const override;
 	virtual TSharedPtr<FPropertyNode> GetPropertyNode() const override;
 	virtual TSharedPtr<FComplexPropertyNode> GetExternalRootPropertyNode() const override;
 	virtual TSharedPtr<IDetailPropertyRow> GetRow() const override;
diff --git a/Engine/Source/Editor/PropertyEditor/Private/IDetailsViewPrivate.h b/Engine/Source/Editor/PropertyEditor/Private/IDetailsViewPrivate.h
--- a/Engine/Source/Editor/PropertyEditor/Private/IDetailsViewPrivate.h
+++ b/Engine/Source/Editor/PropertyEditor/Private/IDetailsViewPrivate.h
@@ -81,14 +81,14 @@ public:
 	virtual bool IsPropertyReadOnly( const struct FPropertyAndParent& PropertyAndParent ) const = 0;
 
 	/**
-	 * @return Whether the IsCustomRowVisible check is pertinent, i.e. always return true no matter the specified row and parent names.
+	 * @return Whether a custom row with the specified name and parent name is visible.
 	 */
-	virtual bool IsCustomRowVisibilityFiltered() const = 0;
+	virtual bool IsCustomRowVisible(FName InRowName, FName InParentName) const = 0;
 
 	/**
-	 * @return Whether a custom row with the specified name and parent name is visible.
+	 * @return Whether a custom row with the specified name and parent name is read-only.
 	 */
-	virtual bool IsCustomRowVisible(FName InRowName, FName InParentName) const = 0;
+	virtual bool IsCustomRowReadOnly(FName InRowName, FName InParentName) const = 0;
 
 	/**
 	 * @return The thumbnail pool that should be used for thumbnails being rendered in this view
diff --git a/Engine/Source/Editor/PropertyEditor/Private/SDetailsViewBase.cpp b/Engine/Source/Editor/PropertyEditor/Private/SDetailsViewBase.cpp
--- a/Engine/Source/Editor/PropertyEditor/Private/SDetailsViewBase.cpp
+++ b/Engine/Source/Editor/PropertyEditor/Private/SDetailsViewBase.cpp
@@ -463,16 +463,16 @@ void SDetailsViewBase::SetIsPropertyReadOnlyDelegate(FIsPropertyReadOnly InIsPro
 	IsPropertyReadOnlyDelegate = InIsPropertyReadOnly;
 }
 
-void SDetailsViewBase::SetIsCustomRowVisibilityFilteredDelegate(FIsCustomRowVisibilityFiltered InIsCustomRowVisibilityFilteredDelegate)
-{
-	IsCustomRowVisibilityFilteredDelegate = InIsCustomRowVisibilityFilteredDelegate;
-};
-
 void SDetailsViewBase::SetIsCustomRowVisibleDelegate(FIsCustomRowVisible InIsCustomRowVisible)
 {
 	IsCustomRowVisibleDelegate = InIsCustomRowVisible;
 }
 
+void SDetailsViewBase::SetIsCustomRowReadOnlyDelegate(FIsCustomRowReadOnly InIsCustomRowReadOnly)
+{
+	IsCustomRowReadOnlyDelegate = InIsCustomRowReadOnly;
+}
+
 void SDetailsViewBase::SetIsPropertyEditingEnabledDelegate(FIsPropertyEditingEnabled IsPropertyEditingEnabled)
 {
 	IsPropertyEditingEnabledDelegate = IsPropertyEditingEnabled;
@@ -566,18 +566,14 @@ bool SDetailsViewBase::IsPropertyReadOnly( const FPropertyAndParent& PropertyAnd
 	return IsPropertyReadOnlyDelegate.IsBound() ? IsPropertyReadOnlyDelegate.Execute(PropertyAndParent) : false;
 }
 
-bool SDetailsViewBase::IsCustomRowVisibilityFiltered() const
+bool SDetailsViewBase::IsCustomRowVisible(FName InRowName, FName InParentName) const
 {
-	if (!IsCustomRowVisibleDelegate.IsBound())
-	{
-		return false;
-	}
-	return IsCustomRowVisibilityFilteredDelegate.IsBound() ? IsCustomRowVisibilityFilteredDelegate.Execute() : true;
+	return IsCustomRowVisibleDelegate.IsBound() ? IsCustomRowVisibleDelegate.Execute(InRowName, InParentName) : true;
 }
 
-bool SDetailsViewBase::IsCustomRowVisible(FName InRowName, FName InParentName) const
+bool SDetailsViewBase::IsCustomRowReadOnly(FName InRowName, FName InParentName) const
 {
-	return IsCustomRowVisibleDelegate.IsBound() ? IsCustomRowVisibleDelegate.Execute(InRowName, InParentName) : true;
+	return IsCustomRowReadOnlyDelegate.IsBound() ? IsCustomRowReadOnlyDelegate.Execute(InRowName, InParentName) : false;
 }
 
 TSharedPtr<IPropertyUtilities> SDetailsViewBase::GetPropertyUtilities()
diff --git a/Engine/Source/Editor/PropertyEditor/Private/SDetailsViewBase.h b/Engine/Source/Editor/PropertyEditor/Private/SDetailsViewBase.h
--- a/Engine/Source/Editor/PropertyEditor/Private/SDetailsViewBase.h
+++ b/Engine/Source/Editor/PropertyEditor/Private/SDetailsViewBase.h
@@ -127,12 +127,12 @@ public:
 	virtual FOnDisplayedPropertiesChanged& GetOnDisplayedPropertiesChanged() override { return OnDisplayedPropertiesChangedDelegate; }
 	virtual void SetDisableCustomDetailLayouts( bool bInDisableCustomDetailLayouts ) override { bDisableCustomDetailLayouts = bInDisableCustomDetailLayouts; }
 	virtual void SetIsPropertyVisibleDelegate(FIsPropertyVisible InIsPropertyVisible) override;
-	virtual FIsCustomRowVisibilityFiltered& GetIsCustomRowVisibilityFilteredDelegate() override { return IsCustomRowVisibilityFilteredDelegate; };
-	virtual void SetIsCustomRowVisibilityFilteredDelegate(FIsCustomRowVisibilityFiltered InIsCustomRowVisibilityFilteredDelegate) override;
 	virtual FIsPropertyVisible& GetIsPropertyVisibleDelegate() override { return IsPropertyVisibleDelegate; }
 	virtual void SetIsCustomRowVisibleDelegate(FIsCustomRowVisible InIsCustomRowVisible) override;
 	virtual FIsCustomRowVisible& GetIsCustomRowVisibleDelegate() override { return IsCustomRowVisibleDelegate; }
 	virtual void SetIsPropertyReadOnlyDelegate(FIsPropertyReadOnly InIsPropertyReadOnly) override;
+	virtual FIsCustomRowReadOnly& GetIsCustomRowReadOnlyDelegate() override { return IsCustomRowReadOnlyDelegate; }
+	virtual void SetIsCustomRowReadOnlyDelegate(FIsCustomRowReadOnly InIsCustomRowReadOnly) override;
 	virtual FIsPropertyReadOnly& GetIsPropertyReadOnlyDelegate() override { return IsPropertyReadOnlyDelegate; }
 	virtual void SetIsPropertyEditingEnabledDelegate(FIsPropertyEditingEnabled IsPropertyEditingEnabled) override;
 	virtual FIsPropertyEditingEnabled& GetIsPropertyEditingEnabledDelegate() override { return IsPropertyEditingEnabledDelegate; }
@@ -143,8 +143,8 @@ public:
 	virtual TSharedPtr<IDetailPropertyExtensionHandler> GetExtensionHandler() override;
 	virtual bool IsPropertyVisible(const struct FPropertyAndParent& PropertyAndParent) const override;
 	virtual bool IsPropertyReadOnly(const struct FPropertyAndParent& PropertyAndParent) const override;
-    virtual bool IsCustomRowVisibilityFiltered() const override;
 	virtual bool IsCustomRowVisible(FName InRowName, FName InParentName) const override;
+	virtual bool IsCustomRowReadOnly(FName InRowName, FName InParentName) const override;
 	virtual void SetGenericLayoutDetailsDelegate(FOnGetDetailCustomizationInstance OnGetGenericDetails) override;
 	virtual FOnGetDetailCustomizationInstance& GetGenericLayoutDetailsDelegate() override { return GenericLayoutDelegate; }
 	virtual bool IsLocked() const override { return bIsLocked; }
@@ -386,12 +386,12 @@ protected:
 	FDetailNodeList RootTreeNodes;
 	/** Delegate executed to determine if a property should be visible */
 	FIsPropertyVisible IsPropertyVisibleDelegate;
-	/** Delegate executed to check if custom row visibility is filtered. */
-	FIsCustomRowVisibilityFiltered IsCustomRowVisibilityFilteredDelegate;
 	/** Delegate executed to determine if a custom row should be visible. */
 	FIsCustomRowVisible IsCustomRowVisibleDelegate;
 	/** Delegate executed to determine if a property should be read-only */
 	FIsPropertyReadOnly IsPropertyReadOnlyDelegate;
+	/** Delegate executed to determine if a custom row should be read-only. */
+	FIsCustomRowReadOnly IsCustomRowReadOnlyDelegate;
 	/** Delegate called to see if a property editing is enabled */
 	FIsPropertyEditingEnabled IsPropertyEditingEnabledDelegate;
 	/** Delegate called when the details panel finishes editing a property (after post edit change is called) */
diff --git a/Engine/Source/Editor/PropertyEditor/Public/DetailLayoutBuilder.h b/Engine/Source/Editor/PropertyEditor/Public/DetailLayoutBuilder.h
--- a/Engine/Source/Editor/PropertyEditor/Public/DetailLayoutBuilder.h
+++ b/Engine/Source/Editor/PropertyEditor/Public/DetailLayoutBuilder.h
@@ -9,6 +9,7 @@
 #include "IDetailPropertyRow.h"
 
 class IDetailCategoryBuilder;
+class IDetailsView;
 
 namespace ECategoryPriority
 {
@@ -55,7 +56,12 @@ public:
 	/**
 	 * @return the parent detail view for this layout builder
 	 */
-	virtual const class IDetailsView* GetDetailsView() const = 0;
+	virtual const IDetailsView* GetDetailsView() const = 0;
+
+	/**
+	 * @return the parent detail view for this layout builder
+	 */
+	virtual IDetailsView* GetDetailsView() = 0;
 
 	/**
 	 * @return The base class of the objects being customized in this detail layout
diff --git a/Engine/Source/Editor/PropertyEditor/Public/IDetailsView.h b/Engine/Source/Editor/PropertyEditor/Public/IDetailsView.h
--- a/Engine/Source/Editor/PropertyEditor/Public/IDetailsView.h
+++ b/Engine/Source/Editor/PropertyEditor/Public/IDetailsView.h
@@ -270,17 +270,16 @@ public:
 	virtual FIsPropertyReadOnly& GetIsPropertyReadOnlyDelegate() = 0;
 
 	/**
-	 * Sets a delegate to call to check if custom row visibility is filtered,
-	 * i.e. whether the FIsCustomRowVisible delegate will always return true no matter the parameters.
+	 * Sets a delegate to call to determine if a specific custom row should be visible in this instance of the details view
 	 */
-	virtual void SetIsCustomRowVisibilityFilteredDelegate(FIsCustomRowVisibilityFiltered InIsCustomRowVisibilityFiltered) = 0;
-	virtual FIsCustomRowVisibilityFiltered& GetIsCustomRowVisibilityFilteredDelegate() = 0;
+	virtual void SetIsCustomRowVisibleDelegate(FIsCustomRowVisible InIsCustomRowVisible) = 0;
+	virtual FIsCustomRowVisible& GetIsCustomRowVisibleDelegate() = 0;
 
 	/**
 	 * Sets a delegate to call to determine if a specific custom row should be visible in this instance of the details view
 	 */
-	virtual void SetIsCustomRowVisibleDelegate(FIsCustomRowVisible InIsCustomRowVisible) = 0;
-	virtual FIsCustomRowVisible& GetIsCustomRowVisibleDelegate() = 0;
+	virtual void SetIsCustomRowReadOnlyDelegate(FIsCustomRowReadOnly InIsCustomRowVisible) = 0;
+	virtual FIsCustomRowReadOnly& GetIsCustomRowReadOnlyDelegate() = 0;
 
 	/**
 	 * Sets a delegate to call to layout generic details not specific to an object being viewed
diff --git a/Engine/Source/Editor/PropertyEditor/Public/PropertyEditorDelegates.h b/Engine/Source/Editor/PropertyEditor/Public/PropertyEditorDelegates.h
--- a/Engine/Source/Editor/PropertyEditor/Public/PropertyEditorDelegates.h
+++ b/Engine/Source/Editor/PropertyEditor/Public/PropertyEditorDelegates.h
@@ -36,15 +36,12 @@ DECLARE_DELEGATE_RetVal_OneParam( bool, FIsPropertyVisible, const FPropertyAndPa
 /** Delegate called to see if a property should be read-only */
 DECLARE_DELEGATE_RetVal_OneParam( bool, FIsPropertyReadOnly, const FPropertyAndParent& );
 
-/** 
- * Delegate called to check if custom row visibility is filtered, 
- * i.e. whether FIsCustomRowVisible delegate will always return true no matter the parameters. 
- */
-DECLARE_DELEGATE_RetVal(bool, FIsCustomRowVisibilityFiltered);
-
 /** Delegate called to determine if a custom row should be visible. */
 DECLARE_DELEGATE_RetVal_TwoParams(bool, FIsCustomRowVisible, FName /*InRowName*/, FName /*InParentName*/);
 
+/** Delegate called to determine if a custom row should be read-only. */
+DECLARE_DELEGATE_RetVal_TwoParams(bool, FIsCustomRowReadOnly, FName /*InRowName*/, FName /*InParentName*/);
+
 /** Delegate called to get a detail layout for a specific object class */
 DECLARE_DELEGATE_RetVal( TSharedRef<IDetailCustomization>, FOnGetDetailCustomizationInstance );
 
diff --git a/Engine/Source/Editor/UMGEditor/Private/Components/DetailsView.cpp b/Engine/Source/Editor/UMGEditor/Private/Components/DetailsView.cpp
--- a/Engine/Source/Editor/UMGEditor/Private/Components/DetailsView.cpp
+++ b/Engine/Source/Editor/UMGEditor/Private/Components/DetailsView.cpp
@@ -79,9 +79,7 @@ void UDetailsView::BuildContentWidget()
 			
 			DetailViewWidget->SetCustomFilterLabel(LOCTEXT("ShowAllParameters", "Show All Parameters"));
 			DetailViewWidget->SetCustomFilterDelegate(FSimpleDelegate::CreateUObject(this, &UDetailsView::ToggleWhitelistedProperties));
-
 			DetailViewWidget->SetIsPropertyVisibleDelegate(FIsPropertyVisible::CreateUObject(this, &UDetailsView::GetIsPropertyVisible));
-			DetailViewWidget->SetIsCustomRowVisibilityFilteredDelegate(FIsCustomRowVisibilityFiltered::CreateUObject(this, &UDetailsView::IsRowVisibilityFiltered));
 			DetailViewWidget->SetIsCustomRowVisibleDelegate(FIsCustomRowVisible::CreateUObject(this, &UDetailsView::GetIsRowVisible));
 			DetailViewWidget->SetObject(ViewedObject);
 			if (DetailViewWidget.IsValid())
@@ -162,14 +160,9 @@ void UDetailsView::ToggleWhitelistedProperties()
 }
 
 
-bool UDetailsView::IsRowVisibilityFiltered() const
-{
-	return bShowOnlyWhitelisted && (PropertiesToShow.Num() > 0 || CategoriesToShow.Num() > 0);
-}
-
 bool UDetailsView::GetIsPropertyVisible(const FPropertyAndParent& PropertyAndParent) const
 {
-    if (!IsRowVisibilityFiltered())
+    if (!bShowOnlyWhitelisted)
 	{
 		return true;
 	}
@@ -186,7 +179,7 @@ bool UDetailsView::GetIsPropertyVisible(const FPropertyAndParent& PropertyAndPar
 
 bool UDetailsView::GetIsRowVisible(FName InRowName, FName InParentName) const
 {
-    if (!IsRowVisibilityFiltered())
+    if (!bShowOnlyWhitelisted)
 	{
 		return true;
 	}
