Added code for favoriting custom builders.

Also added code for favoriting groups, but this currently isn't enabled work since header widgets suffer from the same problem as custom widget rows do - their SWidget is created in place rather than through a function call, so thus cannot be duplicated. This needs some more thought, but the code for storing them is the same either way.

diff --git a/Engine/Source/Editor/EditorConfig/Private/EditorMetadataOverrides.cpp b/Engine/Source/Editor/EditorConfig/Private/EditorMetadataOverrides.cpp
--- a/Engine/Source/Editor/EditorConfig/Private/EditorMetadataOverrides.cpp
+++ b/Engine/Source/Editor/EditorConfig/Private/EditorMetadataOverrides.cpp
@@ -78,8 +78,10 @@ void UEditorMetadataOverrides::OnCompleted(bool bSuccess)
 	}
 }
 
-const FMetadataSet* UEditorMetadataOverrides::FindFieldMetadata(const FField* Field, FName Key) const
+const FMetadataSet* UEditorMetadataOverrides::FindFieldMetadata(const FField* Field) const
 {
+	check(Field != nullptr);
+	
 	if (!SourceConfig.IsValid())
 	{
 		return nullptr;
@@ -100,9 +102,36 @@ const FMetadataSet* UEditorMetadataOverrides::FindFieldMetadata(const FField* Fi
 	const FMetadataSet* FieldMetadata = StructMetadata->Fields.Find(Field->GetFName());
 	return FieldMetadata;
 }
+
+FMetadataSet* UEditorMetadataOverrides::FindFieldMetadata(const FField* Field)
+{
+	check(Field != nullptr);
+	
+	if (!SourceConfig.IsValid())
+	{
+		return nullptr;
+	}
+
+	const UStruct* OwnerStruct = Field->GetOwnerStruct();
+	if (OwnerStruct == nullptr)
+	{
+		return nullptr;
+	}
+
+	FStructMetadata* StructMetadata = LoadedMetadata.Classes.Find(OwnerStruct->GetFName());
+	if (StructMetadata == nullptr)
+	{
+		return nullptr;
+	}
+
+	FMetadataSet* FieldMetadata = StructMetadata->Fields.Find(Field->GetFName());
+	return FieldMetadata;
+}
 	
-FMetadataSet* UEditorMetadataOverrides::FindOrAddFieldMetadata(const FField* Field, FName Key)
+FMetadataSet* UEditorMetadataOverrides::FindOrAddFieldMetadata(const FField* Field)
 {
+	check(Field != nullptr);
+
 	if (!SourceConfig.IsValid())
 	{
 		return nullptr;
@@ -120,19 +149,69 @@ FMetadataSet* UEditorMetadataOverrides::FindOrAddFieldMetadata(const FField* Fie
 	return &FieldMetadata;
 }
 
-bool UEditorMetadataOverrides::HasMetadata(const FField* Field, FName Key) const
+EMetadataType UEditorMetadataOverrides::GetMetadataType(const FField* Field, FName Key) const
 {
 	check(Field != nullptr);
 
-	const FMetadataSet* FieldMetadata = FindFieldMetadata(Field, Key);
-	return FieldMetadata != nullptr;
+	const FMetadataSet* MetadataSet = FindFieldMetadata(Field);
+	if (MetadataSet == nullptr)
+	{
+		return EMetadataType::None;
+	}
+
+	if (MetadataSet->Bools.Contains(Key))
+	{
+		return EMetadataType::Bool;
+	}
+	if (MetadataSet->Ints.Contains(Key))
+	{
+		return EMetadataType::Int;
+	}
+	if (MetadataSet->Floats.Contains(Key))
+	{
+		return EMetadataType::Float;
+	}
+	if (MetadataSet->Strings.Contains(Key))
+	{
+		return EMetadataType::String;
+	}
+
+	return EMetadataType::None;
 }
 
-bool UEditorMetadataOverrides::GetStringMetadata(const FField* Field, FName Key, FString& OutValue) const
+EMetadataType UEditorMetadataOverrides::GetMetadataType(const UStruct* Struct, FName Key) const
 {
-	check(Field != nullptr);
+	check(Struct != nullptr);
 
-	const FMetadataSet* FieldMetadata = FindFieldMetadata(Field, Key);
+	const FMetadataSet* MetadataSet = FindStructMetadata(Struct);
+	if (MetadataSet == nullptr)
+	{
+		return EMetadataType::None; 
+	}
+
+	if (MetadataSet->Bools.Contains(Key))
+	{
+		return EMetadataType::Bool;
+	}
+	if (MetadataSet->Ints.Contains(Key))
+	{
+		return EMetadataType::Int;
+	}
+	if (MetadataSet->Floats.Contains(Key))
+	{
+		return EMetadataType::Float;
+	}
+	if (MetadataSet->Strings.Contains(Key))
+	{
+		return EMetadataType::String;
+	}
+
+	return EMetadataType::None;
+}
+
+bool UEditorMetadataOverrides::GetStringMetadata(const FField* Field, FName Key, FString& OutValue) const
+{
+	const FMetadataSet* FieldMetadata = FindFieldMetadata(Field);
 	if (FieldMetadata == nullptr)
 	{
 		return false;
@@ -156,9 +235,7 @@ bool UEditorMetadataOverrides::GetStringMetadata(const FField* Field, FName Key,
 
 void UEditorMetadataOverrides::SetStringMetadata(const FField* Field, FName Key, FStringView Value)
 {
-	check(Field != nullptr);
-		
-	FMetadataSet* FieldMetadata = FindOrAddFieldMetadata(Field, Key);
+	FMetadataSet* FieldMetadata = FindOrAddFieldMetadata(Field);
 	if (FieldMetadata == nullptr)
 	{
 		return;
@@ -170,9 +247,7 @@ void UEditorMetadataOverrides::SetStringMetadata(const FField* Field, FName Key,
 
 bool UEditorMetadataOverrides::GetFloatMetadata(const FField* Field, FName Key, float& OutValue) const
 {
-check(Field != nullptr);
-
-	const FMetadataSet* FieldMetadata = FindFieldMetadata(Field, Key);
+	const FMetadataSet* FieldMetadata = FindFieldMetadata(Field);
 	if (FieldMetadata == nullptr)
 	{
 		return false;
@@ -196,9 +271,7 @@ check(Field != nullptr);
 
 void UEditorMetadataOverrides::SetFloatMetadata(const FField* Field, FName Key, float Value)
 {
-	check(Field != nullptr);
-		
-	FMetadataSet* FieldMetadata = FindOrAddFieldMetadata(Field, Key);
+	FMetadataSet* FieldMetadata = FindOrAddFieldMetadata(Field);
 	if (FieldMetadata == nullptr)
 	{
 		return;
@@ -210,9 +283,7 @@ void UEditorMetadataOverrides::SetFloatMetadata(const FField* Field, FName Key,
 
 bool UEditorMetadataOverrides::GetIntMetadata(const FField* Field, FName Key, int32& OutValue) const
 {
-	check(Field != nullptr);
-
-	const FMetadataSet* FieldMetadata = FindFieldMetadata(Field, Key);
+	const FMetadataSet* FieldMetadata = FindFieldMetadata(Field);
 	if (FieldMetadata == nullptr)
 	{
 		return false;
@@ -236,9 +307,7 @@ bool UEditorMetadataOverrides::GetIntMetadata(const FField* Field, FName Key, in
 
 void UEditorMetadataOverrides::SetIntMetadata(const FField* Field, FName Key, int32 Value)
 {
-	check(Field != nullptr);
-		
-	FMetadataSet* FieldMetadata = FindOrAddFieldMetadata(Field, Key);
+	FMetadataSet* FieldMetadata = FindOrAddFieldMetadata(Field);
 	if (FieldMetadata == nullptr)
 	{
 		return;
@@ -250,9 +319,7 @@ void UEditorMetadataOverrides::SetIntMetadata(const FField* Field, FName Key, in
 
 bool UEditorMetadataOverrides::GetBoolMetadata(const FField* Field, FName Key, bool& OutValue) const
 {
-	check(Field != nullptr);
-
-	const FMetadataSet* FieldMetadata = FindFieldMetadata(Field, Key);
+	const FMetadataSet* FieldMetadata = FindFieldMetadata(Field);
 	if (FieldMetadata == nullptr)
 	{
 		return false;
@@ -276,9 +343,7 @@ bool UEditorMetadataOverrides::GetBoolMetadata(const FField* Field, FName Key, b
 	
 void UEditorMetadataOverrides::SetBoolMetadata(const FField* Field, FName Key, bool Value)
 {
-	check(Field != nullptr);
-		
-	FMetadataSet* FieldMetadata = FindOrAddFieldMetadata(Field, Key);
+	FMetadataSet* FieldMetadata = FindOrAddFieldMetadata(Field);
 	if (FieldMetadata == nullptr)
 	{
 		return;
@@ -290,8 +355,6 @@ void UEditorMetadataOverrides::SetBoolMetadata(const FField* Field, FName Key, b
 
 bool UEditorMetadataOverrides::GetClassMetadata(const FField* Field, FName Key, UClass*& OutValue) const
 {
-	check(Field != nullptr);
-
 	FString ClassName;
 	if (!GetStringMetadata(Field, Key, ClassName))
 	{
@@ -299,13 +362,11 @@ bool UEditorMetadataOverrides::GetClassMetadata(const FField* Field, FName Key,
 	}
 
 	OutValue = FindObject<UClass>(ANY_PACKAGE, *ClassName);
-	return true; // we return true here because we did have a value, it just wasn't a valid class name
+	return true; // we return true here even if the value is null because we did have a value, it just wasn't a valid class name
 }
 
 void UEditorMetadataOverrides::SetClassMetadata(const FField* Field, FName Key, UClass* Value)
 {
-	check(Field != nullptr);
-
 	FString ClassName;
 	if (Value != nullptr)
 	{
@@ -316,6 +377,98 @@ void UEditorMetadataOverrides::SetClassMetadata(const FField* Field, FName Key,
 	bDirty = true;
 }
 
+bool UEditorMetadataOverrides::GetArrayMetadata(const FField* Field, FName Key, TArray<FString>& OutValue) const
+{
+	const FMetadataSet* FieldMetadata = FindFieldMetadata(Field);
+	if (FieldMetadata == nullptr)
+	{
+		return false;
+	}
+
+	const FString* MetaValue = FieldMetadata->Strings.Find(Key);
+	if (MetaValue != nullptr)
+	{
+		MetaValue->ParseIntoArray(OutValue, TEXT(","));
+		return true;
+	}
+
+	// fallback to metadata system
+	if (!Field->HasMetaData(Key))
+	{
+		return false;
+	}
+
+	const FString ArrayString = Field->GetMetaData(Key);
+	ArrayString.ParseIntoArray(OutValue, TEXT(","));
+	return true;
+}
+	
+void UEditorMetadataOverrides::SetArrayMetadata(const FField* Field, FName Key, const TArray<FString>& Value)
+{
+	FMetadataSet* FieldMetadata = FindOrAddFieldMetadata(Field);
+	if (FieldMetadata == nullptr)
+	{
+		return;
+	}
+
+	const FString ValueString = FString::Join(Value, TEXT(","));
+	FieldMetadata->Strings.Add(Key, ValueString);
+	bDirty = true;
+}
+
+void UEditorMetadataOverrides::AddToArrayMetadata(const FField* Field, FName Key, const FString& Value)
+{
+	FMetadataSet* FieldMetadata = FindOrAddFieldMetadata(Field);
+	if (FieldMetadata == nullptr)
+	{
+		return;
+	}
+
+	FString& CurrentValue = FieldMetadata->Strings.FindOrAdd(Key);
+	if (CurrentValue.IsEmpty())
+	{
+		CurrentValue = Value;
+	}
+	else
+	{
+		CurrentValue.Append(TEXT(","));
+		CurrentValue.Append(Value);
+	}
+	bDirty = true;
+}
+
+void UEditorMetadataOverrides::RemoveFromArrayMetadata(const FField* Field, FName Key, const FString& Value)
+{
+	FMetadataSet* FieldMetadata = FindFieldMetadata(Field);
+	if (FieldMetadata == nullptr)
+	{
+		return;
+	}
+
+	FString* CurrentValue = FieldMetadata->Strings.Find(Key);
+	if (CurrentValue == nullptr)
+	{
+		return;
+	}
+	
+	TArray<FString> ArrayValue;
+	CurrentValue->ParseIntoArray(ArrayValue, TEXT(","));
+
+	ArrayValue.Remove(Value);
+
+	if (ArrayValue.IsEmpty())
+	{
+		FieldMetadata->Strings.Remove(Key);
+	}
+	else
+	{
+		const FString ValueString = FString::Join(ArrayValue, TEXT(","));
+		FieldMetadata->Strings.Add(Key, ValueString);
+	}
+
+	bDirty = true;
+}
+
 void UEditorMetadataOverrides::RemoveMetadata(const FField* Field, FName Key)
 {
 	if (!SourceConfig.IsValid())
@@ -348,3 +501,324 @@ void UEditorMetadataOverrides::RemoveMetadata(const FField* Field, FName Key)
 
 	bDirty = true;
 }
+
+const FMetadataSet* UEditorMetadataOverrides::FindStructMetadata(const UStruct* Struct) const
+{
+	check(Struct != nullptr);
+
+	if (!SourceConfig.IsValid())
+	{
+		return nullptr;
+	}
+
+	const FStructMetadata* StructMetadata = LoadedMetadata.Classes.Find(Struct->GetFName());
+	return StructMetadata != nullptr ? &StructMetadata->StructMetadata : nullptr;
+}
+
+FMetadataSet* UEditorMetadataOverrides::FindStructMetadata(const UStruct* Struct)
+{
+	check(Struct != nullptr);
+
+	if (!SourceConfig.IsValid())
+	{
+		return nullptr;
+	}
+
+	FStructMetadata* StructMetadata = LoadedMetadata.Classes.Find(Struct->GetFName());
+	return StructMetadata != nullptr ? &StructMetadata->StructMetadata : nullptr;
+}
+	
+FMetadataSet* UEditorMetadataOverrides::FindOrAddStructMetadata(const UStruct* Struct)
+{
+	check(Struct != nullptr);
+
+	if (!SourceConfig.IsValid())
+	{
+		return nullptr;
+	}
+
+	FStructMetadata& StructMetadata = LoadedMetadata.Classes.FindOrAdd(Struct->GetFName());
+	return &StructMetadata.StructMetadata;
+}
+
+bool UEditorMetadataOverrides::GetStringMetadata(const UStruct* Struct, FName Key, FString& OutValue) const
+{
+	const FMetadataSet* StructMetadata = FindStructMetadata(Struct);
+	if (StructMetadata == nullptr)
+	{
+		return false;
+	}
+
+	const FString* MetaValue = StructMetadata->Strings.Find(Key);
+	if (MetaValue != nullptr)
+	{
+		if (!Struct->HasMetaData(Key))
+		{
+			return false;
+		}
+
+		OutValue = Struct->GetMetaData(Key);
+		return true;
+	}
+
+	OutValue = *MetaValue;
+	return true;
+}
+
+void UEditorMetadataOverrides::SetStringMetadata(const UStruct* Struct, FName Key, FStringView Value)
+{
+	FMetadataSet* StructMetadata = FindOrAddStructMetadata(Struct);
+	if (StructMetadata == nullptr)
+	{
+		return;
+	}
+
+	StructMetadata->Strings.Add(Key, FString(Value));
+	bDirty = true;
+}
+
+bool UEditorMetadataOverrides::GetFloatMetadata(const UStruct* Struct, FName Key, float& OutValue) const
+{
+	const FMetadataSet* StructMetadata = FindStructMetadata(Struct);
+	if (StructMetadata == nullptr)
+	{
+		return false;
+	}
+
+	const float* MetaValue = StructMetadata->Floats.Find(Key);
+	if (MetaValue == nullptr)
+	{
+		if (!Struct->HasMetaData(Key))
+		{
+			return false;
+		}
+
+		OutValue = Struct->GetFloatMetaData(Key);
+		return true;
+	}
+
+	OutValue = *MetaValue;
+	return true;
+}
+
+void UEditorMetadataOverrides::SetFloatMetadata(const UStruct* Struct, FName Key, float Value)
+{
+	FMetadataSet* StructMetadata = FindOrAddStructMetadata(Struct);
+	if (StructMetadata == nullptr)
+	{
+		return;
+	}
+
+	StructMetadata->Floats.Add(Key, Value);
+	bDirty = true;
+}
+
+bool UEditorMetadataOverrides::GetIntMetadata(const UStruct* Struct, FName Key, int32& OutValue) const
+{
+	const FMetadataSet* StructMetadata = FindStructMetadata(Struct);
+	if (StructMetadata == nullptr)
+	{
+		return false;
+	}
+
+	const int32* MetaValue = StructMetadata->Ints.Find(Key);
+	if (MetaValue == nullptr)
+	{
+		if (!Struct->HasMetaData(Key))
+		{
+			return false;
+		}
+
+		OutValue = Struct->GetIntMetaData(Key);
+		return true;
+	}
+
+	OutValue = *MetaValue;
+	return true;
+}
+
+void UEditorMetadataOverrides::SetIntMetadata(const UStruct* Struct, FName Key, int32 Value)
+{
+	FMetadataSet* StructMetadata = FindOrAddStructMetadata(Struct);
+	if (StructMetadata == nullptr)
+	{
+		return;
+	}
+
+	StructMetadata->Ints.Add(Key, Value);
+	bDirty = true;
+}
+
+bool UEditorMetadataOverrides::GetBoolMetadata(const UStruct* Struct, FName Key, bool& OutValue) const
+{
+	const FMetadataSet* StructMetadata = FindStructMetadata(Struct);
+	if (StructMetadata == nullptr)
+	{
+		return false;
+	}
+
+	const bool* MetaValue = StructMetadata->Bools.Find(Key);
+	if (MetaValue == nullptr)
+	{
+		if (!Struct->HasMetaData(Key))
+		{
+			return false;
+		}
+
+		OutValue = Struct->GetBoolMetaData(Key);
+		return true;
+	}
+
+	OutValue = *MetaValue;
+	return true;
+}
+	
+void UEditorMetadataOverrides::SetBoolMetadata(const UStruct* Struct, FName Key, bool Value)
+{
+	FMetadataSet* StructMetadata = FindOrAddStructMetadata(Struct);
+	if (StructMetadata == nullptr)
+	{
+		return;
+	}
+
+	StructMetadata->Bools.Add(Key, Value);
+	bDirty = true;
+}
+
+bool UEditorMetadataOverrides::GetClassMetadata(const UStruct* Struct, FName Key, UClass*& OutValue) const
+{
+	FString ClassName;
+	if (!GetStringMetadata(Struct, Key, ClassName))
+	{
+		return false;
+	}
+
+	OutValue = FindObject<UClass>(ANY_PACKAGE, *ClassName);
+	return true; // we return true here even if the class is null because we did have a value, it just wasn't a valid class name
+}
+
+void UEditorMetadataOverrides::SetClassMetadata(const UStruct* Struct, FName Key, UClass* Value)
+{
+	FString ClassName;
+	if (Value != nullptr)
+	{
+		ClassName = Value->GetName();
+	}
+
+	SetStringMetadata(Struct, Key, ClassName);
+	bDirty = true;
+}
+
+bool UEditorMetadataOverrides::GetArrayMetadata(const UStruct* Struct, FName Key, TArray<FString>& OutValue) const
+{
+	const FMetadataSet* StructMetadata = FindStructMetadata(Struct);
+	if (StructMetadata == nullptr)
+	{
+		return false;
+	}
+
+	const FString* MetaValue = StructMetadata->Strings.Find(Key);
+	if (MetaValue != nullptr)
+	{
+		MetaValue->ParseIntoArray(OutValue, TEXT(","));
+		return true;
+	}
+
+	if (!Struct->HasMetaData(Key))
+	{
+		return false;
+	}
+
+	const FString ArrayString = Struct->GetMetaData(Key);
+	ArrayString.ParseIntoArray(OutValue, TEXT(","));
+	return true;
+}
+	
+void UEditorMetadataOverrides::SetArrayMetadata(const UStruct* Struct, FName Key, const TArray<FString>& Value)
+{
+	FMetadataSet* StructMetadata = FindOrAddStructMetadata(Struct);
+	if (StructMetadata == nullptr)
+	{
+		return;
+	}
+
+	const FString ValueString = FString::Join(Value, TEXT(","));
+
+	StructMetadata->Strings.Add(Key, ValueString);
+	bDirty = true;
+}
+
+void UEditorMetadataOverrides::AddToArrayMetadata(const UStruct* Struct, FName Key, const FString& Value)
+{
+	FMetadataSet* StructMetadata = FindOrAddStructMetadata(Struct);
+	if (StructMetadata == nullptr)
+	{
+		return;
+	}
+
+	FString& CurrentValue = StructMetadata->Strings.FindOrAdd(Key);
+	if (CurrentValue.IsEmpty())
+	{
+		CurrentValue = Value;
+	}
+	else
+	{
+		CurrentValue.Append(TEXT(","));
+		CurrentValue.Append(Value);
+	}
+
+	bDirty = true;
+}
+
+void UEditorMetadataOverrides::RemoveFromArrayMetadata(const UStruct* Struct, FName Key, const FString& Value)
+{
+	FMetadataSet* StructMetadata = FindStructMetadata(Struct);
+	if (StructMetadata == nullptr)
+	{
+		return;
+	}
+
+	FString* CurrentValue = StructMetadata->Strings.Find(Key);
+	if (CurrentValue == nullptr)
+	{
+		return;
+	}
+	
+	TArray<FString> ArrayValue;
+	CurrentValue->ParseIntoArray(ArrayValue, TEXT(","));
+
+	ArrayValue.Remove(Value);
+
+	if (ArrayValue.IsEmpty())
+	{
+		StructMetadata->Strings.Remove(Key);
+	}
+	else
+	{
+		const FString ValueString = FString::Join(ArrayValue, TEXT(","));
+		StructMetadata->Strings.Add(Key, ValueString);
+	}
+
+	bDirty = true;
+}
+
+void UEditorMetadataOverrides::RemoveMetadata(const UStruct* Struct, FName Key)
+{
+	if (!SourceConfig.IsValid())
+	{
+		return;
+	}
+
+	FStructMetadata* StructMetadata = LoadedMetadata.Classes.Find(Struct->GetFName());
+	if (StructMetadata == nullptr)
+	{
+		return;
+	}
+
+	StructMetadata->StructMetadata.Ints.Remove(Key);
+	StructMetadata->StructMetadata.Bools.Remove(Key);
+	StructMetadata->StructMetadata.Floats.Remove(Key);
+	StructMetadata->StructMetadata.Strings.Remove(Key);
+
+	bDirty = true;
+}
diff --git a/Engine/Source/Editor/EditorConfig/Public/EditorMetadataOverrides.h b/Engine/Source/Editor/EditorConfig/Public/EditorMetadataOverrides.h
--- a/Engine/Source/Editor/EditorConfig/Public/EditorMetadataOverrides.h
+++ b/Engine/Source/Editor/EditorConfig/Public/EditorMetadataOverrides.h
@@ -51,6 +51,16 @@ struct FMetadataConfig
 	TMap<FName, FStructMetadata> Classes;
 };
 
+UENUM()
+enum class EMetadataType
+{
+	None,
+	Bool,
+	Int,
+	Float,
+	String
+};
+
 UCLASS()
 class EDITORCONFIG_API UEditorMetadataOverrides : 
 	public UEditorSubsystem, 
@@ -67,7 +77,8 @@ public:
 	bool LoadFromConfig(TSharedPtr<FEditorConfig> Config);
 	void Save();
 
-	bool HasMetadata(const FField* Field, FName Key) const;
+	EMetadataType GetMetadataType(const FField* Field, FName Key) const;
+	EMetadataType GetMetadataType(const UStruct* Struct, FName Key) const;
 
 	bool GetStringMetadata(const FField* Field, FName Key, FString& OutValue) const;
 	void SetStringMetadata(const FField* Field, FName Key, FStringView Value);
@@ -84,14 +95,48 @@ public:
 	bool GetClassMetadata(const FField* Field, FName Key, UClass*& OutValue) const;
 	void SetClassMetadata(const FField* Field, FName Key, UClass* Value);
 
+	bool GetArrayMetadata(const FField* Field, FName Key, TArray<FString>& OutValue) const;
+	void SetArrayMetadata(const FField* Field, FName Key, const TArray<FString>& Value);
+
+	void AddToArrayMetadata(const FField* Field, FName Key, const FString& Value);
+	void RemoveFromArrayMetadata(const FField* Field, FName Key, const FString& Value);
+
 	void RemoveMetadata(const FField* Field, FName Key);
 
+	bool GetStringMetadata(const UStruct* Struct, FName Key, FString& OutValue) const;
+	void SetStringMetadata(const UStruct* Struct, FName Key, FStringView Value);
+
+	bool GetFloatMetadata(const UStruct* Struct, FName Key, float& OutValue) const;
+	void SetFloatMetadata(const UStruct* Struct, FName Key, float Value);
+
+	bool GetIntMetadata(const UStruct* Struct, FName Key, int32& OutValue) const;
+	void SetIntMetadata(const UStruct* Struct, FName Key, int32 Value);
+
+	bool GetBoolMetadata(const UStruct* Struct, FName Key, bool& OutValue) const;
+	void SetBoolMetadata(const UStruct* Struct, FName Key, bool Value);
+
+	bool GetClassMetadata(const UStruct* Struct, FName Key, UClass*& OutValue) const;
+	void SetClassMetadata(const UStruct* Struct, FName Key, UClass* Value);
+
+	bool GetArrayMetadata(const UStruct* Struct, FName Key, TArray<FString>& OutValue) const;
+	void SetArrayMetadata(const UStruct* Struct, FName Key, const TArray<FString>& Value);
+
+	void AddToArrayMetadata(const UStruct* Struct, FName Key, const FString& Value);
+	void RemoveFromArrayMetadata(const UStruct* Struct, FName Key, const FString& Value);
+
+	void RemoveMetadata(const UStruct* Struct, FName Key);
+
 	void Tick(float DeltaTime) override;
 	TStatId GetStatId() const override;
 
 private:
-	const FMetadataSet* FindFieldMetadata(const FField* Field, FName Key) const;
-	FMetadataSet* FindOrAddFieldMetadata(const FField* Field, FName Key);
+	const FMetadataSet* FindFieldMetadata(const FField* Field) const;
+	FMetadataSet* FindFieldMetadata(const FField* Field);
+	FMetadataSet* FindOrAddFieldMetadata(const FField* Field);
+
+	const FMetadataSet* FindStructMetadata(const UStruct* Struct) const;
+	FMetadataSet* FindStructMetadata(const UStruct* Struct);
+	FMetadataSet* FindOrAddStructMetadata(const UStruct* Struct);
 
 	void OnCompleted(bool bSuccess);
 
diff --git a/Engine/Source/Editor/PropertyEditor/Private/DetailCategoryBuilderImpl.cpp b/Engine/Source/Editor/PropertyEditor/Private/DetailCategoryBuilderImpl.cpp
--- a/Engine/Source/Editor/PropertyEditor/Private/DetailCategoryBuilderImpl.cpp
+++ b/Engine/Source/Editor/PropertyEditor/Private/DetailCategoryBuilderImpl.cpp
@@ -23,40 +23,42 @@ namespace DetailLayoutConstants
 	const FMargin RowPadding(10.0f, 2.0f);
 }
 
-void FDetailLayout::AddCustomLayout(const FDetailLayoutCustomization& Layout, bool bAdvanced)
+void FDetailLayout::AddLayout(const FDetailLayoutCustomization& Layout)
 {
-	AddLayoutInternal(Layout, bAdvanced ? CustomAdvancedLayouts : CustomSimpleLayouts);
-}
-
-void FDetailLayout::AddDefaultLayout(const FDetailLayoutCustomization& Layout, bool bAdvanced)
-{
-	AddLayoutInternal(Layout, bAdvanced ? DefaultAdvancedLayouts : DefaultSimpleLayouts);
+	if (Layout.bAdvanced)
+	{
+		AdvancedLayouts.Add(Layout);
+	}
+	else
+	{
+		SimpleLayouts.Add(Layout);
+	}
 }
 
 FDetailLayoutCustomization* FDetailLayout::GetDefaultLayout(const TSharedRef<FPropertyNode>& PropertyNode)
 {
 	FDetailLayoutCustomization* Customization = 
-		DefaultSimpleLayouts.FindByPredicate([&PropertyNode](const FDetailLayoutCustomization& TestCustomization)
+		SimpleLayouts.FindByPredicate([&PropertyNode](const FDetailLayoutCustomization& TestCustomization)
 		{
 			return TestCustomization.GetPropertyNode() == PropertyNode;
 		});
 
 	// Didn't find it in the simple layouts, look in advanced layouts
-	if (!Customization)
+	if (Customization == nullptr)
 	{
 		Customization = 
-			DefaultAdvancedLayouts.FindByPredicate([&PropertyNode](const FDetailLayoutCustomization& TestCustomization)
+			AdvancedLayouts.FindByPredicate([&PropertyNode](const FDetailLayoutCustomization& TestCustomization)
 			{
 				return TestCustomization.GetPropertyNode() == PropertyNode;
 			});
 	}
 
-	return Customization;
-}
+	if (Customization != nullptr && Customization->bCustom)
+	{
+		return nullptr;
+	}
 
-void FDetailLayout::AddLayoutInternal(const FDetailLayoutCustomization& Layout, FCustomizationList& ListToUse)
-{
-	ListToUse.Add(Layout);
+	return Customization;
 }
 
 FDetailLayoutCustomization::FDetailLayoutCustomization()
@@ -148,10 +150,12 @@ FDetailCategoryImpl::~FDetailCategoryImpl()
 FDetailWidgetRow& FDetailCategoryImpl::AddCustomRow(const FText& FilterString, bool bForAdvanced)
 {
 	FDetailLayoutCustomization NewCustomization;
+	NewCustomization.bCustom = true;
+	NewCustomization.bAdvanced = bForAdvanced;
 	NewCustomization.WidgetDecl = MakeShareable(new FDetailWidgetRow);
 	NewCustomization.WidgetDecl->FilterString(FilterString);
 
-	AddCustomLayout(NewCustomization, bForAdvanced);
+	AddCustomLayout(NewCustomization);
 
 	return *NewCustomization.WidgetDecl;
 }
@@ -160,17 +164,40 @@ FDetailWidgetRow& FDetailCategoryImpl::AddCustomRow(const FText& FilterString, b
 void FDetailCategoryImpl::AddCustomBuilder(TSharedRef<IDetailCustomNodeBuilder> InCustomBuilder, bool bForAdvanced)
 {
 	FDetailLayoutCustomization NewCustomization;
+	NewCustomization.bCustom = true;
+	NewCustomization.bAdvanced = bForAdvanced;
 	NewCustomization.CustomBuilderRow = MakeShareable(new FDetailCustomBuilderRow(InCustomBuilder));
 
-	AddCustomLayout(NewCustomization, bForAdvanced);
+	if (!bFavoriteCategory)
+	{
+		TStringBuilder<256> PathToNode;
+		PathToNode.Append(GetCategoryPathName());
+		PathToNode.Append(TEXT("."));
+		PathToNode.Append(InCustomBuilder->GetName().ToString());
+
+		if (GetDetailsView()->IsCustomBuilderFavorite(PathToNode))
+		{
+			TSharedPtr<FDetailLayoutBuilderImpl> LayoutBuilder = DetailLayoutBuilder.Pin();
+			if (LayoutBuilder.IsValid())
+			{
+				static const FName FavoritesCategoryName(TEXT("Favorites"));
+				FDetailCategoryImpl& FavoritesCategory = LayoutBuilder->DefaultCategory(FavoritesCategoryName);
+				FavoritesCategory.AddCustomBuilder(InCustomBuilder, false);
+			}
+		}
+	}
+
+	AddCustomLayout(NewCustomization);
 }
 
 IDetailGroup& FDetailCategoryImpl::AddGroup(FName GroupName, const FText& LocalizedDisplayName, bool bForAdvanced, bool bStartExpanded)
 {
 	FDetailLayoutCustomization NewCustomization;
+	NewCustomization.bCustom = true;
+	NewCustomization.bAdvanced = bForAdvanced;
 	NewCustomization.DetailGroup = MakeShareable(new FDetailGroup(GroupName, AsShared(), LocalizedDisplayName, bStartExpanded));
 
-	AddCustomLayout(NewCustomization, bForAdvanced);
+	AddCustomLayout(NewCustomization);
 
 	return *NewCustomization.DetailGroup;
 }
@@ -179,14 +206,10 @@ int32 FDetailCategoryImpl::GetNumCustomizations() const
 {
 	int32 NumCustomizations = 0;
 
-	for (int32 LayoutIndex = 0; LayoutIndex < LayoutMap.Num(); ++LayoutIndex)
+	for (const FDetailLayout& Layout : LayoutMap)
 	{
-		const FDetailLayout& Layout = LayoutMap[LayoutIndex];
-
-		NumCustomizations += Layout.GetDefaultSimpleLayouts().Num();
-		NumCustomizations += Layout.GetDefaultAdvancedLayouts().Num();
-		NumCustomizations += Layout.GetCustomSimpleLayouts().Num();
-		NumCustomizations += Layout.GetCustomAdvancedLayouts().Num();
+		NumCustomizations += Layout.GetSimpleLayouts().Num();
+		NumCustomizations += Layout.GetAdvancedLayouts().Num();
 	}
 
 	return NumCustomizations;
@@ -195,22 +218,17 @@ int32 FDetailCategoryImpl::GetNumCustomizations() const
 void FDetailCategoryImpl::GetDefaultProperties(TArray<TSharedRef<IPropertyHandle> >& OutDefaultProperties, bool bSimpleProperties, bool bAdvancedProperties)
 {
 	FDetailLayoutBuilderImpl& DetailLayoutBuilderRef = GetParentLayoutImpl();
-	for (int32 LayoutIndex = 0; LayoutIndex < LayoutMap.Num(); ++LayoutIndex)
+	for (const FDetailLayout& Layout : LayoutMap)
 	{
-		const FDetailLayout& Layout = LayoutMap[LayoutIndex];
-
 		if (bSimpleProperties)
 		{
-			const FCustomizationList& CustomizationList = Layout.GetDefaultSimpleLayouts();
-
-			for (int32 CustomizationIndex = 0; CustomizationIndex < CustomizationList.Num(); ++CustomizationIndex)
+			for (const FDetailLayoutCustomization& Customization : Layout.GetSimpleLayouts())
 			{
-				if (CustomizationList[CustomizationIndex].HasPropertyNode())
+				if (Customization.HasPropertyNode())
 				{
-					const TSharedPtr<FPropertyNode>& Node = CustomizationList[CustomizationIndex].GetPropertyNode();
+					const TSharedPtr<FPropertyNode>& Node = Customization.GetPropertyNode();
 
 					TSharedRef<IPropertyHandle> PropertyHandle = DetailLayoutBuilderRef.GetPropertyHandle(Node);
-
 					if (PropertyHandle->IsValidHandle())
 					{
 						OutDefaultProperties.Add(PropertyHandle);
@@ -221,15 +239,16 @@ void FDetailCategoryImpl::GetDefaultProperties(TArray<TSharedRef<IPropertyHandle
 
 		if (bAdvancedProperties)
 		{
-			const FCustomizationList& CustomizationList = Layout.GetDefaultAdvancedLayouts();
-
-			for (int32 CustomizationIndex = 0; CustomizationIndex < CustomizationList.Num(); ++CustomizationIndex)
+			for (const FDetailLayoutCustomization& Customization : Layout.GetAdvancedLayouts())
 			{
-				if (CustomizationList[CustomizationIndex].HasPropertyNode())
+				if (Customization.HasPropertyNode())
 				{
-					const TSharedPtr<FPropertyNode>& Node = CustomizationList[CustomizationIndex].GetPropertyNode();
-
-					OutDefaultProperties.Add(DetailLayoutBuilderRef.GetPropertyHandle(Node));
+					const TSharedPtr<FPropertyNode>& Node = Customization.GetPropertyNode();
+					TSharedRef<IPropertyHandle> PropertyHandle = DetailLayoutBuilderRef.GetPropertyHandle(Node);
+					if (PropertyHandle->IsValidHandle())
+					{
+						OutDefaultProperties.Add(PropertyHandle);
+					}
 				}
 			}
 		}
@@ -281,6 +300,8 @@ IDetailCategoryBuilder& FDetailCategoryImpl::HeaderContent(TSharedRef<SWidget> I
 IDetailPropertyRow& FDetailCategoryImpl::AddProperty(FName PropertyPath, UClass* ClassOutermost, FName InstanceName, EPropertyLocation::Type Location)
 {
 	FDetailLayoutCustomization NewCustomization;
+	NewCustomization.bCustom = true;
+
 	TSharedPtr<FPropertyNode> PropertyNode = GetParentLayoutImpl().GetPropertyNode(PropertyPath, ClassOutermost, InstanceName);
 	if (PropertyNode.IsValid())
 	{
@@ -288,21 +309,9 @@ IDetailPropertyRow& FDetailCategoryImpl::AddProperty(FName PropertyPath, UClass*
 	}
 
 	NewCustomization.PropertyRow = MakeShareable(new FDetailPropertyRow(PropertyNode, AsShared()));
+	NewCustomization.bAdvanced = (Location == EPropertyLocation::Default) ? IsAdvancedLayout(NewCustomization) : (Location == EPropertyLocation::Advanced);
 
-	bool bForAdvanced = false;
-
-	if (Location == EPropertyLocation::Default)
-	{
-		// Get the default location of this property
-		bForAdvanced = IsAdvancedLayout(NewCustomization);
-	}
-	else if (Location == EPropertyLocation::Advanced)
-	{
-		// Force advanced
-		bForAdvanced = true;
-	}
-
-	AddCustomLayout(NewCustomization, bForAdvanced);
+	AddCustomLayout(NewCustomization);
 
 	return *NewCustomization.PropertyRow;
 }
@@ -310,6 +319,8 @@ IDetailPropertyRow& FDetailCategoryImpl::AddProperty(FName PropertyPath, UClass*
 IDetailPropertyRow& FDetailCategoryImpl::AddProperty(TSharedPtr<IPropertyHandle> PropertyHandle, EPropertyLocation::Type Location)
 {
 	FDetailLayoutCustomization NewCustomization;
+	NewCustomization.bCustom = true;
+
 	TSharedPtr<FPropertyNode> PropertyNode = GetParentLayoutImpl().GetPropertyNode(PropertyHandle);
 
 	if (PropertyNode.IsValid())
@@ -318,21 +329,9 @@ IDetailPropertyRow& FDetailCategoryImpl::AddProperty(TSharedPtr<IPropertyHandle>
 	}
 
 	NewCustomization.PropertyRow = MakeShareable(new FDetailPropertyRow(PropertyNode, AsShared()));
+	NewCustomization.bAdvanced = (Location == EPropertyLocation::Default) ? IsAdvancedLayout(NewCustomization) : (Location == EPropertyLocation::Advanced);
 
-	bool bForAdvanced = false;
-
-	if (Location == EPropertyLocation::Default)
-	{
-		// Get the default location of this property
-		bForAdvanced = IsAdvancedLayout(NewCustomization);
-	}
-	else if (Location == EPropertyLocation::Advanced)
-	{
-		// Force advanced
-		bForAdvanced = true;
-	}
-
-	AddCustomLayout(NewCustomization, bForAdvanced);
+	AddCustomLayout(NewCustomization);
 
 	return *NewCustomization.PropertyRow;
 }
@@ -340,6 +339,8 @@ IDetailPropertyRow& FDetailCategoryImpl::AddProperty(TSharedPtr<IPropertyHandle>
 IDetailPropertyRow* FDetailCategoryImpl::AddExternalObjects(const TArray<UObject*>& Objects, EPropertyLocation::Type Location /*= EPropertyLocation::Default*/, const FAddPropertyParams& Params /*= FAddPropertyParams()*/)
 {
 	FDetailLayoutCustomization NewCustomization;
+	NewCustomization.bCustom = true;
+	NewCustomization.bAdvanced = Location == EPropertyLocation::Advanced;
 
 	FAddPropertyParams AddPropertyParams = Params;
 	AddPropertyParams.AllowChildren(true);
@@ -350,14 +351,7 @@ IDetailPropertyRow* FDetailCategoryImpl::AddExternalObjects(const TArray<UObject
 
 	if (NewRow.IsValid())
 	{
-		bool bForAdvanced = false;
-		if (Location == EPropertyLocation::Advanced)
-		{
-			// Force advanced
-			bForAdvanced = true;
-		}
-
-		AddCustomLayout(NewCustomization, bForAdvanced);
+		AddCustomLayout(NewCustomization);
 	}
 
 	return NewRow.Get();
@@ -366,6 +360,8 @@ IDetailPropertyRow* FDetailCategoryImpl::AddExternalObjects(const TArray<UObject
 IDetailPropertyRow* FDetailCategoryImpl::AddExternalObjectProperty(const TArray<UObject*>& Objects, FName PropertyName, EPropertyLocation::Type Location, const FAddPropertyParams& Params)
 {
 	FDetailLayoutCustomization NewCustomization;
+	NewCustomization.bCustom = true;
+	NewCustomization.bAdvanced = Location == EPropertyLocation::Advanced;
 
 	FDetailPropertyRow::MakeExternalPropertyRowCustomization(Objects, PropertyName, AsShared(), NewCustomization, Params);
 
@@ -373,14 +369,7 @@ IDetailPropertyRow* FDetailCategoryImpl::AddExternalObjectProperty(const TArray<
 
 	if (NewRow.IsValid())
 	{
-		bool bForAdvanced = false;
-		if (Location == EPropertyLocation::Advanced)
-		{
-			// Force advanced
-			bForAdvanced = true;
-		}
-
-		AddCustomLayout(NewCustomization, bForAdvanced);
+		AddCustomLayout(NewCustomization);
 
 	}
 
@@ -395,6 +384,8 @@ IDetailPropertyRow* FDetailCategoryImpl::AddExternalStructure(TSharedPtr<FStruct
 IDetailPropertyRow* FDetailCategoryImpl::AddExternalStructureProperty(TSharedPtr<FStructOnScope> StructData, FName PropertyName, EPropertyLocation::Type Location/* = EPropertyLocation::Default*/, const FAddPropertyParams& Params)
 {
 	FDetailLayoutCustomization NewCustomization;
+	NewCustomization.bCustom = true;
+	NewCustomization.bAdvanced = Location == EPropertyLocation::Advanced;
 
 	FDetailPropertyRow::MakeExternalPropertyRowCustomization(StructData, PropertyName, AsShared(), NewCustomization, Params);
 
@@ -405,14 +396,7 @@ IDetailPropertyRow* FDetailCategoryImpl::AddExternalStructureProperty(TSharedPtr
 		TSharedPtr<FPropertyNode> PropertyNode = NewRow->GetPropertyNode();
 		TSharedPtr<FComplexPropertyNode> RootNode = StaticCastSharedRef<FComplexPropertyNode>(PropertyNode->FindComplexParent()->AsShared());
 
-		bool bForAdvanced = false;
-		if (Location == EPropertyLocation::Advanced)
-		{
-			// Force advanced
-			bForAdvanced = true;
-		}
-
-		AddCustomLayout(NewCustomization, bForAdvanced);
+		AddCustomLayout(NewCustomization);
 	}
 
 	return NewRow.Get();
@@ -438,7 +422,6 @@ TArray<TSharedPtr<IPropertyHandle>> FDetailCategoryImpl::AddAllExternalStructure
 
 	FDetailLayoutBuilderImpl& DetailLayoutBuilderRef = GetParentLayoutImpl();
 
-	const bool bForAdvanced = Location == EPropertyLocation::Advanced;
 	if (RootPropertyNode.IsValid())
 	{
 		RootPropertyNode->RebuildChildren();
@@ -450,8 +433,9 @@ TArray<TSharedPtr<IPropertyHandle>> FDetailCategoryImpl::AddAllExternalStructure
 			if (FProperty* Property = PropertyNode->GetProperty())
 			{
 				FDetailLayoutCustomization NewCustomization;
+				NewCustomization.bAdvanced = Location == EPropertyLocation::Advanced;
 				NewCustomization.PropertyRow = MakeShared<FDetailPropertyRow>(PropertyNode, AsShared(), RootPropertyNode);
-				AddDefaultLayout(NewCustomization, bForAdvanced, NAME_None);
+				AddDefaultLayout(NewCustomization, NAME_None);
 
 				Handles.Add(DetailLayoutBuilderRef.GetPropertyHandle(PropertyNode));
 			}
@@ -464,9 +448,9 @@ TArray<TSharedPtr<IPropertyHandle>> FDetailCategoryImpl::AddAllExternalStructure
 void FDetailCategoryImpl::AddPropertyNode(TSharedRef<FPropertyNode> PropertyNode, FName InstanceName)
 {
 	FDetailLayoutCustomization NewCustomization;
-
+	NewCustomization.bAdvanced = IsAdvancedLayout(NewCustomization);
 	NewCustomization.PropertyRow = MakeShareable(new FDetailPropertyRow(PropertyNode, AsShared()));
-	AddDefaultLayout(NewCustomization, IsAdvancedLayout(NewCustomization), InstanceName);
+	AddDefaultLayout(NewCustomization, InstanceName);
 }
 
 bool FDetailCategoryImpl::IsAdvancedLayout(const FDetailLayoutCustomization& LayoutInfo)
@@ -480,14 +464,16 @@ bool FDetailCategoryImpl::IsAdvancedLayout(const FDetailLayoutCustomization& Lay
 	return bAdvanced;
 }
 
-void FDetailCategoryImpl::AddCustomLayout(const FDetailLayoutCustomization& LayoutInfo, bool bForAdvanced)
+void FDetailCategoryImpl::AddCustomLayout(const FDetailLayoutCustomization& LayoutInfo)
 {
-	GetLayoutForInstance(GetParentLayoutImpl().GetCurrentCustomizationVariableName()).AddCustomLayout(LayoutInfo, bForAdvanced);
+	ensure(LayoutInfo.bCustom == true);
+	GetLayoutForInstance(GetParentLayoutImpl().GetCurrentCustomizationVariableName()).AddLayout(LayoutInfo);
 }
 
-void FDetailCategoryImpl::AddDefaultLayout(const FDetailLayoutCustomization& LayoutInfo, bool bForAdvanced, FName InstanceName)
+void FDetailCategoryImpl::AddDefaultLayout(const FDetailLayoutCustomization& LayoutInfo, FName InstanceName)
 {
-	GetLayoutForInstance(InstanceName).AddDefaultLayout(LayoutInfo, bForAdvanced);
+	ensure(LayoutInfo.bCustom == false);
+	GetLayoutForInstance(InstanceName).AddLayout(LayoutInfo);
 }
 
 FDetailLayout& FDetailCategoryImpl::GetLayoutForInstance(FName InstanceName)
@@ -510,11 +496,7 @@ FDetailLayoutCustomization* FDetailCategoryImpl::GetDefaultCustomization(TShared
 	FDetailLayout& Layout = GetLayoutForInstance(GetParentLayoutImpl().GetCurrentCustomizationVariableName());
 	
 	FDetailLayoutCustomization* Customization = Layout.GetDefaultLayout(PropertyNode);
-	if (Customization)
-	{
-		return Customization;
-	}
-	return nullptr;
+	return Customization;
 }
 
 bool FDetailCategoryImpl::ShouldShowAdvanced() const
@@ -806,16 +788,13 @@ static bool ShouldBeInlineNode(const TSharedRef<FDetailItemNode>& Node)
 	return false;
 }
 
-void FDetailCategoryImpl::GenerateNodesFromCustomizations(const FCustomizationList& InCustomizationList, bool bDefaultLayouts, FDetailNodeList& OutNodeList, bool &bOutLastItemHasMultipleColumns)
+void FDetailCategoryImpl::GenerateNodesFromCustomizations(const TArray<FDetailLayoutCustomization>& InCustomizationList, FDetailNodeList& OutNodeList)
 {
 	TAttribute<bool> IsParentEnabled(this, &FDetailCategoryImpl::IsParentEnabled);
 
-	bOutLastItemHasMultipleColumns = false;
-	for (int32 CustomizationIndex = 0; CustomizationIndex < InCustomizationList.Num(); ++CustomizationIndex)
+	for (const FDetailLayoutCustomization& Customization : InCustomizationList)
 	{
-		const FDetailLayoutCustomization& Customization = InCustomizationList[CustomizationIndex];
-		// When building default layouts cull default properties which have been customized
-		if (bFavoriteCategory || (Customization.IsValidCustomization() && (!bDefaultLayouts || !IsCustomProperty(Customization.GetPropertyNode()))))
+		if (bFavoriteCategory || Customization.IsValidCustomization())
 		{
 			TSharedRef<FDetailItemNode> NewNode = MakeShareable(new FDetailItemNode(Customization, AsShared(), IsParentEnabled));
 			NewNode->Initialize();
@@ -827,112 +806,34 @@ void FDetailCategoryImpl::GenerateNodesFromCustomizations(const FCustomizationLi
 				continue;
 			}
 
-			// Add the node unless only its children should be visible or it didnt generate any children or if it is a custom builder which can generate children at any point
+			// Add the node unless only its children should be visible or it didn't generate any children or if it is a custom builder which can generate children at any point
 			if (!NewNode->ShouldShowOnlyChildren() || NewNode->HasGeneratedChildren() || Customization.HasCustomBuilder())
 			{
-				if (CustomizationIndex == InCustomizationList.Num() - 1)
-				{
-					bOutLastItemHasMultipleColumns = NewNode->HasMultiColumnWidget();
-				}
-
 				OutNodeList.Add(NewNode);
 			}
 		}
 	}
 }
 
-
-bool FDetailCategoryImpl::GenerateChildrenForSingleLayout(const FName RequiredGroupName, bool bDefaultLayout, bool bNeedsGroup, const FCustomizationList& LayoutList, FDetailNodeList& OutChildren, bool& bOutLastItemHasMultipleColumns)
-{
-	bool bGeneratedAnyChildren = false;
-	if (LayoutList.Num() > 0)
-	{
-		FDetailNodeList GeneratedChildren;
-		GenerateNodesFromCustomizations(LayoutList, bDefaultLayout, GeneratedChildren, bOutLastItemHasMultipleColumns);
-
-		if (GeneratedChildren.Num() > 0)
-		{
-			bGeneratedAnyChildren = true;
-			if (bNeedsGroup)
-			{
-				TSharedRef<FDetailTreeNode> GroupNode = MakeShareable(new FDetailCategoryGroupNode(GeneratedChildren, RequiredGroupName, *this));
-				OutChildren.Add(GroupNode);
-			}
-			else
-			{
-				OutChildren.Append(GeneratedChildren);
-			}
-		}
-	}
-
-	return bGeneratedAnyChildren;
-}
-
 void FDetailCategoryImpl::GenerateChildrenForLayouts()
 {
-	bool bHasAdvancedLayouts = false;
-	bool bGeneratedAnyChildren = false;
-
-	bool bLastItemHasMultipleColumns = false;
-	// @todo Details each loop can be a function
-	for (int32 LayoutIndex = 0; LayoutIndex < LayoutMap.Num(); ++LayoutIndex)
-	{
-		const FDetailLayout& Layout = LayoutMap[LayoutIndex];
-
-		const FName RequiredGroupName = Layout.GetInstanceName();
-
-		bHasAdvancedLayouts |= Layout.HasAdvancedLayouts();
-
-		const bool bDefaultLayout = false;
-		const bool bShouldShowGroup = LayoutMap.ShouldShowGroup(RequiredGroupName);
-		bGeneratedAnyChildren |= GenerateChildrenForSingleLayout(RequiredGroupName, bDefaultLayout, bShouldShowGroup, Layout.GetCustomSimpleLayouts(), SimpleChildNodes, bLastItemHasMultipleColumns);
-	}
-
-	for (int32 LayoutIndex = 0; LayoutIndex < LayoutMap.Num(); ++LayoutIndex)
+	for (const FDetailLayout& Layout : LayoutMap)
 	{
-		const FDetailLayout& Layout = LayoutMap[LayoutIndex];
-
-		const FName RequiredGroupName = Layout.GetInstanceName();
-
-		const bool bDefaultLayout = true;
-		const bool bShouldShowGroup = LayoutMap.ShouldShowGroup(RequiredGroupName);
-		bGeneratedAnyChildren |= GenerateChildrenForSingleLayout(RequiredGroupName, bDefaultLayout, bShouldShowGroup, Layout.GetDefaultSimpleLayouts(), SimpleChildNodes, bLastItemHasMultipleColumns);
+		GenerateNodesFromCustomizations(Layout.GetSimpleLayouts(), SimpleChildNodes);
 	}
 
-	TAttribute<bool> ShowAdvanced(this, &FDetailCategoryImpl::ShouldShowAdvanced);
-	TAttribute<bool> IsEnabled(this, &FDetailCategoryImpl::IsAdvancedDropdownEnabled);
-	if (bHasAdvancedLayouts)
+	for (const FDetailLayout& Layout : LayoutMap)
 	{
-		for (int32 LayoutIndex = 0; LayoutIndex < LayoutMap.Num(); ++LayoutIndex)
-		{
-			const FDetailLayout& Layout = LayoutMap[LayoutIndex];
-
-			const FName RequiredGroupName = Layout.GetInstanceName();
-
-			const bool bDefaultLayout = false;
-			const bool bShouldShowGroup = LayoutMap.ShouldShowGroup(RequiredGroupName);
-			bGeneratedAnyChildren |= GenerateChildrenForSingleLayout(RequiredGroupName, bDefaultLayout, bShouldShowGroup, Layout.GetCustomAdvancedLayouts(), AdvancedChildNodes, bLastItemHasMultipleColumns);
-		}
-
-		for (int32 LayoutIndex = 0; LayoutIndex < LayoutMap.Num(); ++LayoutIndex)
-		{
-			const FDetailLayout& Layout = LayoutMap[LayoutIndex];
-
-			const FName RequiredGroupName = Layout.GetInstanceName();
-
-			const bool bDefaultLayout = true;
-			const bool bShouldShowGroup = LayoutMap.ShouldShowGroup(RequiredGroupName);
-			bGeneratedAnyChildren |= GenerateChildrenForSingleLayout(RequiredGroupName, bDefaultLayout, bShouldShowGroup, Layout.GetDefaultAdvancedLayouts(), AdvancedChildNodes, bLastItemHasMultipleColumns);
-		}
+		GenerateNodesFromCustomizations(Layout.GetAdvancedLayouts(), AdvancedChildNodes);
 	}
 
 	// Generate nodes for advanced dropdowns
+	if (AdvancedChildNodes.Num() > 0)
 	{
-		if (AdvancedChildNodes.Num() > 0)
-		{
-			AdvancedDropdownNodeTop = MakeShareable(new FAdvancedDropdownNode(*this, true));
-		}
+		TAttribute<bool> ShowAdvanced(this, &FDetailCategoryImpl::ShouldShowAdvanced);
+		TAttribute<bool> IsEnabled(this, &FDetailCategoryImpl::IsAdvancedDropdownEnabled);
 
+		AdvancedDropdownNodeTop = MakeShareable(new FAdvancedDropdownNode(*this, true));
 		AdvancedDropdownNodeBottom = MakeShareable(new FAdvancedDropdownNode(*this, ShowAdvanced, IsEnabled, AdvancedChildNodes.Num() > 0, SimpleChildNodes.Num() == 0));
 	}
 }
diff --git a/Engine/Source/Editor/PropertyEditor/Private/DetailCategoryBuilderImpl.h b/Engine/Source/Editor/PropertyEditor/Private/DetailCategoryBuilderImpl.h
--- a/Engine/Source/Editor/PropertyEditor/Private/DetailCategoryBuilderImpl.h
+++ b/Engine/Source/Editor/PropertyEditor/Private/DetailCategoryBuilderImpl.h
@@ -19,12 +19,10 @@
 #include "DetailCategoryBuilder.h"
 #include "PropertyCustomizationHelpers.h"
 
-class FDetailCustomBuilderRow;
-class FDetailGroup;
-class FDetailPropertyRow;
-class FDetailWidgetRow;
 class IDetailGroup;
+class FDetailGroup;
 class IDetailPropertyRow;
+class FDetailPropertyRow;
 
 /**
  * Defines a customization for a specific detail
@@ -58,10 +56,12 @@ struct FDetailLayoutCustomization
 	TSharedPtr<FPropertyNode> GetPropertyNode() const;
 	/** @return The row to display from this customization */
 	FDetailWidgetRow GetWidgetRow() const;
+	/** Whether or not this customization is considered an advanced property. */
+	bool bAdvanced { false };
+	/** Whether or not this customization is custom or a default one. */
+	bool bCustom { false };
 };
 
-typedef TArray<FDetailLayoutCustomization> FCustomizationList;
-
 class FDetailLayout
 {
 public:
@@ -69,31 +69,26 @@ public:
 		: InstanceName(InInstanceName)
 	{}
 
-	void AddCustomLayout(const FDetailLayoutCustomization& Layout, bool bAdvanced);
-	void AddDefaultLayout(const FDetailLayoutCustomization& Layout, bool bAdvanced);
+	void AddLayout(const FDetailLayoutCustomization& Layout);
 
-	const FCustomizationList& GetCustomSimpleLayouts() const { return CustomSimpleLayouts; }
-	const FCustomizationList& GetCustomAdvancedLayouts() const { return CustomAdvancedLayouts; }
-	const FCustomizationList& GetDefaultSimpleLayouts() const { return DefaultSimpleLayouts; }
-	const FCustomizationList& GetDefaultAdvancedLayouts() const { return DefaultAdvancedLayouts; }
+	const TArray<FDetailLayoutCustomization>& GetSimpleLayouts() const { return SimpleLayouts; }
+	const TArray<FDetailLayoutCustomization>& GetAdvancedLayouts() const { return AdvancedLayouts; }
 
 	FDetailLayoutCustomization* GetDefaultLayout(const TSharedRef<FPropertyNode>& PropertyNode);
 
-	bool HasAdvancedLayouts() const { return CustomAdvancedLayouts.Num() > 0 || DefaultAdvancedLayouts.Num() > 0; }
+	bool HasAdvancedLayouts() const { return AdvancedLayouts.Num() > 0; }
 
+	/**
+	 * Get the instance name - this is usually the UObject's name when in a multi-selection.
+	 * Used to display a group beneath the category if multiple objects share some of the same-named properties.
+	 */
 	FName GetInstanceName() const { return InstanceName; }
 
 private:
-	void AddLayoutInternal(const FDetailLayoutCustomization& Layout, FCustomizationList& ListToUse);
-private:
-	/** Customized layouts that appear in the simple (visible by default) area of a category */
-	FCustomizationList CustomSimpleLayouts;
-	/** Customized layouts that appear in the advanced (hidden by default) details area of a category */
-	FCustomizationList CustomAdvancedLayouts;
-	/** Default layouts that appear in the simple (visible by default) details area of a category */
-	FCustomizationList DefaultSimpleLayouts;
-	/** Default layouts that appear in the advanced (visible by default) details area of a category */
-	FCustomizationList DefaultAdvancedLayouts;
+	/** Layouts that appear in the simple (visible by default) area of a category */
+	TArray<FDetailLayoutCustomization> SimpleLayouts;
+	/** Layouts that appear in the advanced (hidden by default) details area of a category */
+	TArray<FDetailLayoutCustomization> AdvancedLayouts;
 	/** The sort order in which this layout is displayed (lower numbers are displayed first) */
 	FName InstanceName;
 };
@@ -124,16 +119,13 @@ public:
 		return Layouts[Index];
 	}
 
-	/**
-	 * @return The number of layouts
-	 */
-	int32 Num() const { return Layouts.Num(); }
+	using RangedForIteratorType = TArray<FDetailLayout>::RangedForIteratorType;
+	using RangedForConstIteratorType = TArray<FDetailLayout>::RangedForConstIteratorType;
 
-	/**
-	 * @return Gets a layout at a specific instance
-	 */
-	const FDetailLayout& operator[](int32 Index) const { return Layouts[Index]; }
-	FDetailLayout& operator[](int32 Index) { return Layouts[Index]; }
+	FORCEINLINE RangedForIteratorType      begin()       { return Layouts.begin(); }
+	FORCEINLINE RangedForConstIteratorType begin() const { return Layouts.begin(); }
+	FORCEINLINE RangedForIteratorType      end()         { return Layouts.end(); }
+	FORCEINLINE RangedForConstIteratorType end()   const { return Layouts.end(); }
 
 	/**
 	 * @return Whether or not we need to display a group border around a list of details.
@@ -293,6 +285,9 @@ public:
 	/** @return true if this category only contains advanced properties */
 	bool ContainsOnlyAdvanced() const;
 
+	/**
+	 * Get the number of customizations in this category.
+	 */
 	int32 GetNumCustomizations() const;
 
 	/**
@@ -325,20 +320,8 @@ private:
 	 *
 	 * @param InCustomizationList	The list of customizations to generate nodes from
 	 * @param OutNodeList			The generated nodes
-	 * @param bDefaultLayouts		True if we are generating a default layout
-	 */
-	void GenerateNodesFromCustomizations(const FCustomizationList& InCustomizationList, bool bDefaultLayouts, FDetailNodeList& OutNodeList, bool &bOutHasMultipleColumns);
-
-	/**
-	 * Generates nodes from a list of customization in a single layout
-	 *
-	 * @param RequiredGroupName 	If valid the children will be surrounded by a group
-	 * @param bDefaultLayout	True if we are generating a default layout
-	 * @param bNeedsGroup		True if the children need to be grouped
-	 * @param LayoutList		The list of customizations to generate nodes from
-		 * @param OutChildren		The generated nodes
 	 */
-	bool GenerateChildrenForSingleLayout(const FName RequiredGroupName, bool bDefaultLayout, bool bNeedsGroup, const FCustomizationList& LayoutList, FDetailNodeList& OutChildren, bool& bOutHasMultipleColumns);
+	void GenerateNodesFromCustomizations(const TArray<FDetailLayoutCustomization>& InCustomizationList, FDetailNodeList& OutNodeList);
 
 	/**
 	 * @return Whether or not a customization should appear in the advanced section of the category by default
@@ -351,7 +334,7 @@ private:
 	 * @param LayoutInfo	The custom layout information
 	 * @param bForAdvanced	Whether or not the custom layout should appear in the advanced section of the category
 	 */
-	void AddCustomLayout(const FDetailLayoutCustomization& LayoutInfo, bool bForAdvanced);
+	void AddCustomLayout(const FDetailLayoutCustomization& LayoutInfo);
 
 	/**
 	 * Adds a default layout to this category
@@ -359,7 +342,7 @@ private:
 	 * @param DefaultLayoutInfo		The layout information
 	 * @param bForAdvanced			Whether or not the layout should appear in the advanced section of the category
 	 */
-	void AddDefaultLayout(const FDetailLayoutCustomization& DefaultLayoutInfo, bool bForAdvanced, FName InstanceName);
+	void AddDefaultLayout(const FDetailLayoutCustomization& DefaultLayoutInfo, FName InstanceName);
 
 	/**
 	 * Returns the layout for a given object instance name
@@ -407,7 +390,7 @@ private:
 	FString CategoryPathName;
 	/** Custom header content displayed to the right of the category name */
 	TSharedPtr<SWidget> HeaderContentWidget;
-
+	/** A property node that is displayed in the header row to the right of the category name. */
 	TSharedPtr<FDetailTreeNode> InlinePropertyNode;
 	/** The parent detail builder */
 	TWeakPtr<FDetailLayoutBuilderImpl> DetailLayoutBuilder;
@@ -429,7 +412,7 @@ private:
 	bool bHasVisibleDetails : 1;
 	/** true if the category is visible at all */
 	bool bIsCategoryVisible : 1;
-	/*true if the category is the special favorite category, all property in the layout will be display when we generate the roottree */
+	/*true if this category is the special favorite category, all property in the layout will be display when we generate the root tree */
 	bool bFavoriteCategory : 1;
 	bool bShowOnlyChildren : 1;
 };
diff --git a/Engine/Source/Editor/PropertyEditor/Private/DetailCustomBuilderRow.h b/Engine/Source/Editor/PropertyEditor/Private/DetailCustomBuilderRow.h
--- a/Engine/Source/Editor/PropertyEditor/Private/DetailCustomBuilderRow.h
+++ b/Engine/Source/Editor/PropertyEditor/Private/DetailCustomBuilderRow.h
@@ -17,12 +17,12 @@ class FDetailCustomBuilderRow : public TSharedFromThis<FDetailCustomBuilderRow>
 public:
 	FDetailCustomBuilderRow( TSharedRef<IDetailCustomNodeBuilder> CustomBuilder );
 
-
 	void Tick( float DeltaTime );
 	bool RequiresTick() const;
 	bool HasColumns() const;
 	bool ShowOnlyChildren() const;
 	void OnItemNodeInitialized( TSharedRef<FDetailItemNode> InTreeNode, TSharedRef<FDetailCategoryImpl> InParentCategory, const TAttribute<bool>& InIsParentEnabled );
+	TSharedRef<IDetailCustomNodeBuilder> GetCustomBuilder() const { return CustomNodeBuilder; }
 	FName GetCustomBuilderName() const;
 	TSharedPtr<IPropertyHandle> GetPropertyHandle() const;
 	void OnGenerateChildren( FDetailNodeList& OutChildren );
@@ -38,4 +38,3 @@ private:
 	TSharedPtr<class FCustomChildrenBuilder> ChildrenBuilder;
 	TWeakPtr<FDetailCategoryImpl> ParentCategory;
 };
-
diff --git a/Engine/Source/Editor/PropertyEditor/Private/IDetailsViewPrivate.h b/Engine/Source/Editor/PropertyEditor/Private/IDetailsViewPrivate.h
--- a/Engine/Source/Editor/PropertyEditor/Private/IDetailsViewPrivate.h
+++ b/Engine/Source/Editor/PropertyEditor/Private/IDetailsViewPrivate.h
@@ -145,4 +145,16 @@ public:
 
 	/** Does this details view allow favoriting? */
 	virtual bool IsFavoritingEnabled() const = 0;
+
+	/** Is the given group a favorite? */
+	virtual bool IsGroupFavorite(FStringView GroupPath) const = 0;
+
+	/** Set the given group's favorite status. */
+	virtual void SetGroupFavorite(FStringView GroupPath, bool IsFavorite) = 0;
+
+	/** Is the given custom builder a favorite? */
+	virtual bool IsCustomBuilderFavorite(FStringView Path) const = 0;
+
+	/** Set the given group's favorite status. */
+	virtual void SetCustomBuilderFavorite(FStringView GroupPath, bool IsFavorite) = 0;
 };
diff --git a/Engine/Source/Editor/PropertyEditor/Private/SDetailSingleItemRow.cpp b/Engine/Source/Editor/PropertyEditor/Private/SDetailSingleItemRow.cpp
--- a/Engine/Source/Editor/PropertyEditor/Private/SDetailSingleItemRow.cpp
+++ b/Engine/Source/Editor/PropertyEditor/Private/SDetailSingleItemRow.cpp
@@ -6,6 +6,7 @@
 #include "DetailPropertyRow.h"
 #include "DetailWidgetRow.h"
 #include "Editor.h"
+#include "EditorMetadataOverrides.h"
 #include "IDetailKeyframeHandler.h"
 #include "IDetailPropertyExtensionHandler.h"
 #include "ObjectPropertyNode.h"
@@ -691,17 +692,13 @@ bool SDetailSingleItemRow::OnContextMenuOpening(FMenuBuilder& MenuBuilder)
 	{
 		FUIAction FavoriteAction;
 		FavoriteAction.ExecuteAction = FExecuteAction::CreateSP(this, &SDetailSingleItemRow::OnFavoriteMenuToggle);
-		FavoriteAction.CanExecuteAction = FCanExecuteAction::CreateLambda([this]()
-			{
-				return Customization->GetPropertyNode().IsValid() && Customization->GetPropertyNode()->CanDisplayFavorite();
-			});
+		FavoriteAction.CanExecuteAction = FCanExecuteAction::CreateSP(this, &SDetailSingleItemRow::CanFavorite);
 
 		FText FavoriteText = NSLOCTEXT("PropertyView", "FavoriteProperty", "Add to Favorites");
 		FText FavoriteTooltipText = NSLOCTEXT("PropertyView", "FavoriteProperty_ToolTip", "Add this property to your favorites.");
 		FName FavoriteIcon = "DetailsView.PropertyIsFavorite";
 
-		bool IsFavorite = Customization->GetPropertyNode().IsValid() && Customization->GetPropertyNode()->IsFavorite();
-		if (IsFavorite)
+		if (IsFavorite())
 		{
 			FavoriteText = NSLOCTEXT("PropertyView", "RemoveFavoriteProperty", "Remove from Favorites");
 			FavoriteTooltipText = NSLOCTEXT("PropertyView", "RemoveFavoriteProperty_ToolTip", "Remove this property from your favorites.");
@@ -834,22 +831,89 @@ TSharedRef<SWidget> SDetailSingleItemRow::CreateExtensionWidget() const
 	return ExtensionWidget.ToSharedRef();
 }
 
+bool SDetailSingleItemRow::CanFavorite() const
+{
+	if (Customization->HasPropertyNode())
+	{
+		return Customization->GetPropertyNode()->CanDisplayFavorite();
+	}
+
+	if (Customization->HasCustomBuilder())
+	{
+		return Customization->CustomBuilderRow->GetCustomBuilderName() != NAME_None;
+	}
+
+	return false;
+}
+
+bool SDetailSingleItemRow::IsFavorite() const
+{
+	if (Customization->HasPropertyNode())
+	{
+		return Customization->GetPropertyNode()->IsFavorite();
+	}
+
+	if (Customization->HasCustomBuilder())
+	{
+		TSharedPtr<FDetailTreeNode> OwnerTreeNodePinned = OwnerTreeNode.Pin();
+		if (OwnerTreeNodePinned.IsValid())
+		{
+			TSharedPtr<FDetailCategoryImpl> ParentCategory = OwnerTreeNodePinned->GetParentCategory();
+			if (ParentCategory.IsValid())
+			{
+				FName BuilderName = Customization->CustomBuilderRow->GetCustomBuilderName();
+				if (BuilderName != NAME_None)
+				{
+					TStringBuilder<256> PathToNode;
+					PathToNode.Append(ParentCategory->GetCategoryPathName());
+					PathToNode.Append(TEXT("."));
+					PathToNode.Append(BuilderName.ToString());
+
+					return OwnerTreeNodePinned->GetDetailsView()->IsCustomBuilderFavorite(PathToNode);
+				}
+			}
+		}
+	}
+
+	return false;
+}
+
 void SDetailSingleItemRow::OnFavoriteMenuToggle()
 {
-	if (!Customization->GetPropertyNode().IsValid() || !Customization->GetPropertyNode()->CanDisplayFavorite())
+	if (!CanFavorite())
 	{
 		return; 
 	}
 
-	bool bToggled = !Customization->GetPropertyNode()->IsFavorite();
-	Customization->GetPropertyNode()->SetFavorite(bToggled);
-
 	TSharedPtr<FDetailTreeNode> OwnerTreeNodePinned = OwnerTreeNode.Pin();
 	if (!OwnerTreeNodePinned.IsValid())
 	{
 		return;
 	}
 
+	IDetailsViewPrivate* DetailsView = OwnerTreeNodePinned->GetDetailsView();
+
+	bool bNewValue = !IsFavorite();
+	if (Customization->HasPropertyNode())
+	{
+		TSharedPtr<FPropertyNode> PropertyNode = Customization->GetPropertyNode();
+		PropertyNode->SetFavorite(bNewValue);
+	}
+	else if (Customization->HasCustomBuilder())
+	{
+		TSharedPtr<FDetailCategoryImpl> ParentCategory = OwnerTreeNodePinned->GetParentCategory();
+		if (ParentCategory.IsValid())
+		{
+			TStringBuilder<256> PathToNode;
+			PathToNode.Append(ParentCategory->GetCategoryPathName());
+			PathToNode.Append(TEXT("."));
+			PathToNode.Append(Customization->CustomBuilderRow->GetCustomBuilderName().ToString());
+
+			bNewValue = !DetailsView->IsCustomBuilderFavorite(PathToNode);
+			DetailsView->SetCustomBuilderFavorite(PathToNode, bNewValue);
+		}
+	}
+
 	// Calculate the scrolling offset (by item) to make sure the mouse stay over the same property
 	int32 ExpandSize = 0;
 	if (OwnerTreeNodePinned->ShouldBeExpanded())
@@ -863,8 +927,7 @@ void SDetailSingleItemRow::OnFavoriteMenuToggle()
 	}
 
 	// Apply the calculated offset
-	IDetailsViewPrivate* DetailsView = OwnerTreeNodePinned->GetDetailsView();
-	DetailsView->MoveScrollOffset(bToggled ? ExpandSize : -ExpandSize);
+	DetailsView->MoveScrollOffset(bNewValue ? ExpandSize : -ExpandSize);
 
 	// Refresh the tree
 	DetailsView->ForceRefresh();
diff --git a/Engine/Source/Editor/PropertyEditor/Private/SDetailSingleItemRow.h b/Engine/Source/Editor/PropertyEditor/Private/SDetailSingleItemRow.h
--- a/Engine/Source/Editor/PropertyEditor/Private/SDetailSingleItemRow.h
+++ b/Engine/Source/Editor/PropertyEditor/Private/SDetailSingleItemRow.h
@@ -81,6 +81,8 @@ private:
 	bool IsHighlighted() const;
 
 	void OnFavoriteMenuToggle();
+	bool CanFavorite() const;
+	bool IsFavorite() const;
 
 	void OnArrayDragEnter(const FDragDropEvent& DragDropEvent);
 	void OnArrayDragLeave(const FDragDropEvent& DragDropEvent);
diff --git a/Engine/Source/Editor/PropertyEditor/Private/SDetailsView.cpp b/Engine/Source/Editor/PropertyEditor/Private/SDetailsView.cpp
--- a/Engine/Source/Editor/PropertyEditor/Private/SDetailsView.cpp
+++ b/Engine/Source/Editor/PropertyEditor/Private/SDetailsView.cpp
@@ -2,24 +2,27 @@
 
 
 #include "SDetailsView.h"
-#include "GameFramework/Actor.h"
+#include "DetailLayoutBuilderImpl.h"
+#include "DetailsViewGenericObjectFilter.h"
+#include "DetailsViewPropertyGenerationUtilities.h"
+#include "Editor.h"
+#include "EditorMetadataOverrides.h"
+#include "ObjectPropertyNode.h"
+#include "PropertyEditorHelpers.h"
+#include "SDetailNameArea.h"
+
+#include "Classes/EditorStyleSettings.h"
 #include "Framework/MultiBox/MultiBoxBuilder.h"
+#include "GameFramework/Actor.h"
 #include "Misc/ConfigCacheIni.h"
-#include "ObjectPropertyNode.h"
 #include "Modules/ModuleManager.h"
+#include "Styling/StyleColors.h"
+#include "UserInterface/PropertyDetails/PropertyDetailsUtilities.h"
+#include "Widgets/Colors/SColorPicker.h"
 #include "Widgets/Images/SImage.h"
 #include "Widgets/Input/SButton.h"
 #include "Widgets/Input/SComboButton.h"
-#include "SDetailNameArea.h"
-#include "PropertyEditorHelpers.h"
-#include "UserInterface/PropertyDetails/PropertyDetailsUtilities.h"
-#include "Widgets/Colors/SColorPicker.h"
 #include "Widgets/Input/SSearchBox.h"
-#include "Classes/EditorStyleSettings.h"
-#include "DetailLayoutBuilderImpl.h"
-#include "DetailsViewPropertyGenerationUtilities.h"
-#include "DetailsViewGenericObjectFilter.h"
-#include "Styling/StyleColors.h"
 
 #define LOCTEXT_NAMESPACE "SDetailsView"
 
@@ -984,4 +987,154 @@ FReply SDetailsView::OnToggleFavoritesClicked()
 	return FReply::Handled();
 }
 
+bool SDetailsView::IsGroupFavorite(FStringView GroupPath) const
+{
+	UEditorMetadataOverrides* EditorMetadata = GEditor->GetEditorSubsystem<UEditorMetadataOverrides>();
+	if (!EditorMetadata)
+	{
+		return false;
+	}
+
+	const FString GroupPathString(GroupPath);
+
+	for (const TSharedPtr<FComplexPropertyNode>& RootPropertyNode : RootPropertyNodes)
+	{
+		const UStruct* BaseStruct = RootPropertyNode->GetBaseStructure();
+		if (BaseStruct != nullptr)
+		{
+			static const FName FavoriteGroupsName("FavoriteGroups");
+
+			TArray<FString> FavoriteGroupsList;
+			if (EditorMetadata->GetArrayMetadata(BaseStruct, FavoriteGroupsName, FavoriteGroupsList))
+			{
+				if (FavoriteGroupsList.Contains(GroupPathString))
+				{
+					return true;
+				}
+			}
+		}
+	}
+
+	return false;
+}
+
+void SDetailsView::SetGroupFavorite(FStringView GroupPath, bool IsFavorite)
+{
+	UEditorMetadataOverrides* MetadataOverrides = GEditor->GetEditorSubsystem<UEditorMetadataOverrides>();
+	if (!MetadataOverrides)
+	{
+		return;
+	}
+
+	const FString GroupPathString(GroupPath);
+
+	for (const TSharedPtr<FComplexPropertyNode>& RootPropertyNode : RootPropertyNodes)
+	{
+		const UStruct* BaseStruct = RootPropertyNode->GetBaseStructure();
+		if (BaseStruct != nullptr)
+		{
+			static const FName FavoriteGroupsName("FavoriteGroups");
+
+			TArray<FString> FavoriteGroupsList;
+			if (MetadataOverrides->GetArrayMetadata(BaseStruct, FavoriteGroupsName, FavoriteGroupsList))
+			{
+				if (IsFavorite)
+				{
+					FavoriteGroupsList.AddUnique(GroupPathString);
+				}
+				else
+				{
+					FavoriteGroupsList.Remove(GroupPathString);
+				}
+
+				MetadataOverrides->SetArrayMetadata(BaseStruct, FavoriteGroupsName, FavoriteGroupsList);
+			}
+			else
+			{
+				// no favorite groups set yet
+				if (IsFavorite)
+				{
+					FavoriteGroupsList.Add(GroupPathString);
+					MetadataOverrides->SetArrayMetadata(BaseStruct, FavoriteGroupsName, FavoriteGroupsList);
+				}
+			}
+		}
+	}
+}
+
+bool SDetailsView::IsCustomBuilderFavorite(FStringView Path) const
+{
+	UEditorMetadataOverrides* EditorMetadata = GEditor->GetEditorSubsystem<UEditorMetadataOverrides>();
+	if (!EditorMetadata)
+	{
+		return false;
+	}
+
+	const FString PathString(Path);
+
+	for (const TSharedPtr<FComplexPropertyNode>& RootPropertyNode : RootPropertyNodes)
+	{
+		const UStruct* BaseStruct = RootPropertyNode->GetBaseStructure();
+		if (BaseStruct != nullptr)
+		{
+			static const FName FavoriteCustomBuildersName("FavoriteCustomBuilders");
+
+			TArray<FString> FavoriteBuildersList;
+			if (EditorMetadata->GetArrayMetadata(BaseStruct, FavoriteCustomBuildersName, FavoriteBuildersList))
+			{
+				if (FavoriteBuildersList.Contains(PathString))
+				{
+					return true;
+				}
+			}
+		}
+	}
+
+	return false;
+}
+
+void SDetailsView::SetCustomBuilderFavorite(FStringView Path, bool IsFavorite)
+{
+	UEditorMetadataOverrides* MetadataOverrides = GEditor->GetEditorSubsystem<UEditorMetadataOverrides>();
+	if (!MetadataOverrides)
+	{
+		return;
+	}
+
+	const FString PathString(Path);
+
+	for (const TSharedPtr<FComplexPropertyNode>& RootPropertyNode : RootPropertyNodes)
+	{
+		const UStruct* BaseStruct = RootPropertyNode->GetBaseStructure();
+		if (BaseStruct != nullptr)
+		{
+			static const FName FavoriteCustomBuildersName("FavoriteCustomBuilders");
+
+			TArray<FString> FavoriteBuildersList;
+			if (MetadataOverrides->GetArrayMetadata(BaseStruct, FavoriteCustomBuildersName, FavoriteBuildersList))
+			{
+				if (IsFavorite)
+				{
+					FavoriteBuildersList.AddUnique(PathString);
+				}
+				else
+				{
+					FavoriteBuildersList.Remove(PathString);
+				}
+
+				MetadataOverrides->SetArrayMetadata(BaseStruct, FavoriteCustomBuildersName, FavoriteBuildersList);
+			}
+			else
+			{
+				// no favorite groups set yet
+				if (IsFavorite)
+				{
+					FavoriteBuildersList.Add(PathString);
+					MetadataOverrides->SetArrayMetadata(BaseStruct, FavoriteCustomBuildersName, FavoriteBuildersList);
+				}
+			}
+		}
+	}
+}
+
 #undef LOCTEXT_NAMESPACE
diff --git a/Engine/Source/Editor/PropertyEditor/Private/SDetailsView.h b/Engine/Source/Editor/PropertyEditor/Private/SDetailsView.h
--- a/Engine/Source/Editor/PropertyEditor/Private/SDetailsView.h
+++ b/Engine/Source/Editor/PropertyEditor/Private/SDetailsView.h
@@ -42,6 +42,10 @@ public:
 	virtual void SetObjects(const TArray<UObject*>& InObjects, bool bForceRefresh = false, bool bOverrideLock = false) override;
 	virtual void SetObjects(const TArray<TWeakObjectPtr<UObject>>& InObjects, bool bForceRefresh = false, bool bOverrideLock = false) override;
 	virtual void SetObject(UObject* InObject, bool bForceRefresh = false) override;
+	virtual bool IsGroupFavorite(FStringView GroupPath) const override;
+	virtual void SetGroupFavorite(FStringView GroupPath, bool IsFavorite) override;
+	virtual bool IsCustomBuilderFavorite(FStringView Path) const override;
+	virtual void SetCustomBuilderFavorite(FStringView Path, bool IsFavorite) override;
 
 	virtual void RemoveInvalidObjects() override;
 	virtual void SetObjectPackageOverrides(const TMap<TWeakObjectPtr<UObject>, TWeakObjectPtr<UPackage>>& InMapping) override;
diff --git a/Engine/Source/Editor/PropertyEditor/Private/SStructureDetailsView.h b/Engine/Source/Editor/PropertyEditor/Private/SStructureDetailsView.h
--- a/Engine/Source/Editor/PropertyEditor/Private/SStructureDetailsView.h
+++ b/Engine/Source/Editor/PropertyEditor/Private/SStructureDetailsView.h
@@ -95,6 +95,10 @@ public:
 	virtual void SetRootObjectCustomizationInstance(TSharedPtr<IDetailRootObjectCustomization> InRootObjectCustomization) override {}
 	virtual TSharedPtr<class IDetailRootObjectCustomization> GetRootObjectCustomization() const override { return nullptr; }
 	virtual void SetObjectFilter(TSharedPtr<FDetailsViewObjectFilter> InFilter) override {}
+	virtual bool IsGroupFavorite(FStringView GroupPath) const { return false; }
+	virtual void SetGroupFavorite(FStringView GroupPath, bool IsFavorite) {}
+	virtual bool IsCustomBuilderFavorite(FStringView Path) const { return false; }
+	virtual void SetCustomBuilderFavorite(FStringView Path, bool IsFavorite) {}
 
 	/* This is required by the base class but there is only ever one root node in a structure details view */
 	virtual FRootPropertyNodeList& GetRootNodes() override;
diff --git a/Engine/Source/Editor/PropertyEditor/Public/IDetailCustomNodeBuilder.h b/Engine/Source/Editor/PropertyEditor/Public/IDetailCustomNodeBuilder.h
--- a/Engine/Source/Editor/PropertyEditor/Public/IDetailCustomNodeBuilder.h
+++ b/Engine/Source/Editor/PropertyEditor/Public/IDetailCustomNodeBuilder.h
@@ -42,12 +42,12 @@ public:
 	/**
 	 * Called to generate child content of this node
 	 *
-	 * @param OutChildRows An array of rows to add children to
+	 * @param ChildrenBuilder The builder to add child rows to.
 	 */
-	virtual void GenerateChildContent( IDetailChildrenBuilder& ChildrenBuilder ) = 0;
+	virtual void GenerateChildContent( IDetailChildrenBuilder& ChildrenBuilder ) {}
 
 	/**
-	 * Called each tick if ReqiresTick is true
+	 * Called each tick if RequiresTick is true
 	 */
 	virtual void Tick(float DeltaTime) {}
 
diff --git a/Engine/Source/Editor/PropertyEditor/Public/PropertyHandle.h b/Engine/Source/Editor/PropertyEditor/Public/PropertyHandle.h
--- a/Engine/Source/Editor/PropertyEditor/Public/PropertyHandle.h
+++ b/Engine/Source/Editor/PropertyEditor/Public/PropertyHandle.h
@@ -62,6 +62,11 @@ public:
 	 */
 	virtual FProperty* GetProperty() const = 0;
 
+	/**
+	 * Gets the property node being edited.
+	 */
+	virtual TSharedPtr<FPropertyNode> GetPropertyNode() const = 0;
+
 	/**
 	 * Gets the property we should use to read meta-data
 	 */
