Fix for custom builders with no children being collapsed

diff --git a/Engine/Source/Editor/PropertyEditor/Private/DetailItemNode.cpp b/Engine/Source/Editor/PropertyEditor/Private/DetailItemNode.cpp
--- a/Engine/Source/Editor/PropertyEditor/Private/DetailItemNode.cpp
+++ b/Engine/Source/Editor/PropertyEditor/Private/DetailItemNode.cpp
@@ -627,7 +627,7 @@ void FDetailItemNode::Tick( float DeltaTime )
 
 EVisibility FDetailItemNode::ComputeItemVisibility() const
 {
-	EVisibility NewVisibility;
+	EVisibility NewVisibility = EVisibility::Visible;
 	if (Customization.HasCustomWidget())
 	{	
 		NewVisibility = Customization.WidgetDecl->VisibilityAttr.Get();
@@ -646,7 +646,7 @@ EVisibility FDetailItemNode::ComputeItemVisibility() const
 	{
 		NewVisibility = Customization.DetailGroup->GetGroupVisibility();
 	}
-	else if (Customization.HasCustomBuilder())
+	else if (Customization.HasCustomBuilder() && Children.Num() > 0)
 	{
 		NewVisibility = EVisibility::Collapsed;
 
