Fixed crash when opening Niagara systems due to GetDetailsView() being null in FDetailCategoryImpl.

diff --git a/Engine/Source/Editor/PropertyEditor/Private/DetailCategoryBuilderImpl.cpp b/Engine/Source/Editor/PropertyEditor/Private/DetailCategoryBuilderImpl.cpp
--- a/Engine/Source/Editor/PropertyEditor/Private/DetailCategoryBuilderImpl.cpp
+++ b/Engine/Source/Editor/PropertyEditor/Private/DetailCategoryBuilderImpl.cpp
@@ -168,7 +168,7 @@ void FDetailCategoryImpl::AddCustomBuilder(TSharedRef<IDetailCustomNodeBuilder>
 	NewCustomization.bAdvanced = bForAdvanced;
 	NewCustomization.CustomBuilderRow = MakeShareable(new FDetailCustomBuilderRow(InCustomBuilder));
 
-	if (!bFavoriteCategory)
+	if (GetDetailsView() != nullptr && !bFavoriteCategory)
 	{
 		TStringBuilder<256> PathToNode;
 		PathToNode.Append(GetCategoryPathName());
