Speculative fix for crash bug when fetching a potentially invalid CachedReadAddress when the parent node is invalid.

diff --git a/Engine/Source/Editor/PropertyEditor/Private/PropertyNode.cpp b/Engine/Source/Editor/PropertyEditor/Private/PropertyNode.cpp
--- a/Engine/Source/Editor/PropertyEditor/Private/PropertyNode.cpp
+++ b/Engine/Source/Editor/PropertyEditor/Private/PropertyNode.cpp
@@ -1142,6 +1142,10 @@ bool FPropertyNode::GetReadAddress(bool InRequiresSingleSelection,
 								   bool bArrayPropertiesCanDifferInSize) const
 
 {
+	if (!ParentNodeWeakPtr.IsValid())
+	{
+		return false;
+	}
 
 	// @todo PropertyEditor Nodes which require validation cannot be cached
 	if( CachedReadAddresses.Num() && !CachedReadAddresses.bRequiresCache && !HasNodeFlags(EPropertyNodeFlags::RequiresValidation) )
@@ -1152,14 +1156,10 @@ bool FPropertyNode::GetReadAddress(bool InRequiresSingleSelection,
 
 	CachedReadAddresses.Reset();
 
-	bool bAllValuesTheSame = false;
-	if (ParentNodeWeakPtr.IsValid())
-	{
-		bAllValuesTheSame = GetReadAddressUncached( *this, InRequiresSingleSelection, &CachedReadAddresses, bComparePropertyContents, bObjectForceCompare, bArrayPropertiesCanDifferInSize );
-		OutAddresses.ReadAddressListData = &CachedReadAddresses;
-		CachedReadAddresses.bAllValuesTheSame = bAllValuesTheSame;
-		CachedReadAddresses.bRequiresCache = false;
-	}
+	bool bAllValuesTheSame = GetReadAddressUncached( *this, InRequiresSingleSelection, &CachedReadAddresses, bComparePropertyContents, bObjectForceCompare, bArrayPropertiesCanDifferInSize );
+	OutAddresses.ReadAddressListData = &CachedReadAddresses;
+	CachedReadAddresses.bAllValuesTheSame = bAllValuesTheSame;
+	CachedReadAddresses.bRequiresCache = false;
 
 	return bAllValuesTheSame;
 }
@@ -1171,6 +1171,11 @@ bool FPropertyNode::GetReadAddress(bool InRequiresSingleSelection,
  */
 bool FPropertyNode::GetReadAddress( FReadAddressList& OutAddresses ) const
 {
+	if (!ParentNodeWeakPtr.IsValid())
+	{
+		return false;
+	}
+
 	// @todo PropertyEditor Nodes which require validation cannot be cached
 	if( CachedReadAddresses.Num() && !HasNodeFlags(EPropertyNodeFlags::RequiresValidation) )
 	{
@@ -1180,17 +1185,14 @@ bool FPropertyNode::GetReadAddress( FReadAddressList& OutAddresses ) const
 
 	CachedReadAddresses.Reset();
 
-	bool bSuccess = false;
-	if (ParentNodeWeakPtr.IsValid())
+	bool bSuccess = GetReadAddressUncached( *this, CachedReadAddresses );
+	if( bSuccess )
 	{
-		bSuccess = GetReadAddressUncached( *this, CachedReadAddresses );
-		if( bSuccess )
-		{
-			OutAddresses.ReadAddressListData = &CachedReadAddresses;
-		}
-		CachedReadAddresses.bRequiresCache = false;
+		OutAddresses.ReadAddressListData = &CachedReadAddresses;
 	}
 
+	CachedReadAddresses.bRequiresCache = false;
+
 	return bSuccess;
 }
 
