Restored logic to hide the default property if a customization exists for the property.

The logic was mistakenly removed in CL 15802120

diff --git a/Engine/Source/Editor/PropertyEditor/Private/DetailCategoryBuilderImpl.cpp b/Engine/Source/Editor/PropertyEditor/Private/DetailCategoryBuilderImpl.cpp
--- a/Engine/Source/Editor/PropertyEditor/Private/DetailCategoryBuilderImpl.cpp
+++ b/Engine/Source/Editor/PropertyEditor/Private/DetailCategoryBuilderImpl.cpp
@@ -794,22 +794,26 @@ void FDetailCategoryImpl::GenerateNodesFromCustomizations(const TArray<FDetailLa
 
 	for (const FDetailLayoutCustomization& Customization : InCustomizationList)
 	{
-		if (bFavoriteCategory || Customization.IsValidCustomization())
+		if (Customization.IsValidCustomization())
 		{
-			TSharedRef<FDetailItemNode> NewNode = MakeShareable(new FDetailItemNode(Customization, AsShared(), IsParentEnabled));
-			NewNode->Initialize();
-
-			if (ShouldBeInlineNode(NewNode))
+			// if a property is customized, skip the default customization
+			if (!IsCustomProperty(Customization.GetPropertyNode()) || Customization.bCustom)
 			{
-				ensureMsgf(!InlinePropertyNode.IsValid(), TEXT("Multiple properties marked InlineCategoryProperty detected in category %s."), *DisplayName.ToString());
-				InlinePropertyNode = NewNode;
-				continue;
-			}
+				TSharedRef<FDetailItemNode> NewNode = MakeShareable(new FDetailItemNode(Customization, AsShared(), IsParentEnabled));
+				NewNode->Initialize();
 
-			// Add the node unless only its children should be visible or it didn't generate any children or if it is a custom builder which can generate children at any point
-			if (!NewNode->ShouldShowOnlyChildren() || NewNode->HasGeneratedChildren() || Customization.HasCustomBuilder())
-			{
-				OutNodeList.Add(NewNode);
+				if (ShouldBeInlineNode(NewNode))
+				{
+					ensureMsgf(!InlinePropertyNode.IsValid(), TEXT("Multiple properties marked InlineCategoryProperty detected in category %s."), *DisplayName.ToString());
+					InlinePropertyNode = NewNode;
+					continue;
+				}
+
+				// Add the node unless only its children should be visible or it didn't generate any children or if it is a custom builder which can generate children at any point
+				if (!NewNode->ShouldShowOnlyChildren() || NewNode->HasGeneratedChildren() || Customization.HasCustomBuilder())
+				{
+					OutNodeList.Add(NewNode);
+				}
 			}
 		}
 	}
