Reverted PR to treat UInterface-derived classes as actors.

diff --git a/Engine/Source/Editor/PropertyEditor/Private/UserInterface/PropertyEditor/SPropertyEditorAsset.cpp b/Engine/Source/Editor/PropertyEditor/Private/UserInterface/PropertyEditor/SPropertyEditorAsset.cpp
--- a/Engine/Source/Editor/PropertyEditor/Private/UserInterface/PropertyEditor/SPropertyEditorAsset.cpp
+++ b/Engine/Source/Editor/PropertyEditor/Private/UserInterface/PropertyEditor/SPropertyEditorAsset.cpp
@@ -84,7 +84,7 @@ bool SPropertyEditorAsset::ShouldDisplayThumbnail(const FArguments& InArgs, cons
 		return false;
 	}
 
-	bool bShowThumbnail = InObjectClass == nullptr || !(InObjectClass->IsChildOf(AActor::StaticClass()) || InObjectClass->IsChildOf(UInterface::StaticClass()));
+	bool bShowThumbnail = InObjectClass == nullptr || !InObjectClass->IsChildOf(AActor::StaticClass());
 
 	// also check metadata for thumbnail & text display
 	const FProperty* PropertyToCheck = nullptr;
@@ -344,7 +344,7 @@ void SPropertyEditorAsset::Construct(const FArguments& InArgs, const TSharedPtr<
 		ObjectClass = UClass::FindCommonBase(ConstCastClassArray(AllowedClassFilters));
 	}
 
-	bIsActor = ObjectClass->IsChildOf(AActor::StaticClass()) || ObjectClass->IsChildOf(UInterface::StaticClass());
+	bIsActor = ObjectClass->IsChildOf(AActor::StaticClass());
 
 	if (Property && !Property->HasMetaData("NoCreate"))
 	{
@@ -805,8 +805,7 @@ void SPropertyEditorAsset::OnMenuOpenChanged(bool bOpen)
 
 bool SPropertyEditorAsset::IsFilteredActor( const AActor* const Actor ) const
 {
-	bool bActorTypeValid = ObjectClass->IsChildOf(UInterface::StaticClass()) ? Actor->GetClass()->ImplementsInterface(ObjectClass) : Actor->IsA(ObjectClass);
-	bool IsAllowed = bActorTypeValid && !Actor->IsChildActor() && IsClassAllowed(Actor->GetClass());
+	bool IsAllowed = Actor->IsA(ObjectClass) && !Actor->IsChildActor() && IsClassAllowed(Actor->GetClass());
 	return IsAllowed;
 }
 
