Add "GetKeyOptions" and "GetValueOptions" uproperty metadata to allow combobox dropdowns on TMap FString/FName key/value pairs.

This makes it possible to declare a UPROPERTY like so:

    UPROPERTY(EditAnywhere, meta = (GetKeyOptions="GetValidItemKeys"))
    TMap<FString, FString> MapOfItems;

and have the dropdown appear only on the key field.

diff --git a/Engine/Source/Editor/PropertyEditor/Private/PropertyEditorHelpers.cpp b/Engine/Source/Editor/PropertyEditor/Private/PropertyEditorHelpers.cpp
--- a/Engine/Source/Editor/PropertyEditor/Private/PropertyEditorHelpers.cpp
+++ b/Engine/Source/Editor/PropertyEditor/Private/PropertyEditorHelpers.cpp
@@ -1228,6 +1228,36 @@ namespace PropertyEditorHelpers
 			}
 		}
 	}
+
+	const TCHAR* GetPropertyOptionsMetaDataKey(const FProperty* Property)
+	{
+		// Only string and name properties can have options
+		if (Property->IsA(FStrProperty::StaticClass()) || Property->IsA(FNameProperty::StaticClass()))
+		{
+			const FProperty* OwnerProperty = Property->GetOwnerProperty();
+			if (OwnerProperty->HasMetaData(TEXT("GetOptions")))
+			{
+				return TEXT("GetOptions");
+			}
+
+			// Map properties can have separate options for keys and values
+			const FMapProperty* MapProperty = CastField<FMapProperty>(OwnerProperty);
+			if (MapProperty)
+			{
+				if (MapProperty->HasMetaData(TEXT("GetKeyOptions")) && MapProperty->GetKeyProperty() == Property)
+				{
+					return TEXT("GetKeyOptions");
+				}
+
+				if (MapProperty->HasMetaData(TEXT("GetValueOptions")) && MapProperty->GetValueProperty() == Property)
+				{
+					return TEXT("GetValueOptions");
+				}
+			}
+		}
+
+		return nullptr;
+	}
 }
 
 #undef LOCTEXT_NAMESPACE
diff --git a/Engine/Source/Editor/PropertyEditor/Private/PropertyEditorHelpers.h b/Engine/Source/Editor/PropertyEditor/Private/PropertyEditorHelpers.h
--- a/Engine/Source/Editor/PropertyEditor/Private/PropertyEditorHelpers.h
+++ b/Engine/Source/Editor/PropertyEditor/Private/PropertyEditorHelpers.h
@@ -265,5 +265,11 @@ namespace PropertyEditorHelpers
 	bool IsVisibleStandaloneProperty(const FPropertyNode& PropertyNode, const FPropertyNode& ParentNode);
 
 	void OrderPropertiesFromMetadata(TArray<FProperty*>& Properties);
+
+	/**
+	* For properties that support options lists, returns the metadata key which holds the name of the UFunction to call.
+	* Returns nullptr if the property doesn't support, or doesn't have, options.
+	*/
+	const TCHAR* GetPropertyOptionsMetaDataKey(const FProperty* Property);
 }
 
diff --git a/Engine/Source/Editor/PropertyEditor/Private/PropertyHandleImpl.cpp b/Engine/Source/Editor/PropertyEditor/Private/PropertyHandleImpl.cpp
--- a/Engine/Source/Editor/PropertyEditor/Private/PropertyHandleImpl.cpp
+++ b/Engine/Source/Editor/PropertyEditor/Private/PropertyHandleImpl.cpp
@@ -2910,9 +2910,9 @@ bool FPropertyHandleBase::GeneratePossibleValues(TArray< TSharedPtr<FString> >&
 			}
 		}
 	}
-	else if ((Property->IsA(FStrProperty::StaticClass()) || Property->IsA(FNameProperty::StaticClass())) && Property->GetOwnerProperty()->HasMetaData(TEXT("GetOptions")))
+	else if (const TCHAR* MetaDataKey = PropertyEditorHelpers::GetPropertyOptionsMetaDataKey(Property))
 	{
-		const FString GetOptionsFunctionName = Property->GetOwnerProperty()->GetMetaData(TEXT("GetOptions"));
+		const FString GetOptionsFunctionName = Property->GetOwnerProperty()->GetMetaData(MetaDataKey);
 		if (!GetOptionsFunctionName.IsEmpty())
 		{
 			TArray<UObject*> OutObjects;
diff --git a/Engine/Source/Editor/PropertyEditor/Private/UserInterface/PropertyEditor/SPropertyEditorCombo.cpp b/Engine/Source/Editor/PropertyEditor/Private/UserInterface/PropertyEditor/SPropertyEditorCombo.cpp
--- a/Engine/Source/Editor/PropertyEditor/Private/UserInterface/PropertyEditor/SPropertyEditorCombo.cpp
+++ b/Engine/Source/Editor/PropertyEditor/Private/UserInterface/PropertyEditor/SPropertyEditorCombo.cpp
@@ -49,8 +49,8 @@ bool SPropertyEditorCombo::Supports( const TSharedRef< class FPropertyEditor >&
 
 	if(	((Property->IsA(FByteProperty::StaticClass()) && CastField<const FByteProperty>(Property)->Enum)
 		||	Property->IsA(FEnumProperty::StaticClass())
-		||	(Property->IsA(FStrProperty::StaticClass()) && Property->HasMetaData(TEXT("Enum")))
-		|| ((Property->IsA(FStrProperty::StaticClass()) || Property->IsA(FNameProperty::StaticClass())) && Property->GetOwnerProperty()->HasMetaData(TEXT("GetOptions")))
+		|| (Property->IsA(FStrProperty::StaticClass()) && Property->HasMetaData(TEXT("Enum")))
+		|| (PropertyEditorHelpers::GetPropertyOptionsMetaDataKey(Property) != nullptr)
 		)
 		&&	( ( ArrayIndex == -1 && Property->ArrayDim == 1 ) || ( ArrayIndex > -1 && Property->ArrayDim > 0 ) ) )
 	{
