Fixed bug with Advanced dropdowns not being visible in the Project Settings dialog.

The assignment of PropertyRow was transposed with the call to IsAdvancedLayout(), so the property reference was invalid leading to a default of false. This path isn't really used outside of settings and modes, which is why it wasn't caught earlier.

diff --git a/Engine/Source/Editor/PropertyEditor/Private/DetailCategoryBuilderImpl.cpp b/Engine/Source/Editor/PropertyEditor/Private/DetailCategoryBuilderImpl.cpp
--- a/Engine/Source/Editor/PropertyEditor/Private/DetailCategoryBuilderImpl.cpp
+++ b/Engine/Source/Editor/PropertyEditor/Private/DetailCategoryBuilderImpl.cpp
@@ -433,8 +433,8 @@ TArray<TSharedPtr<IPropertyHandle>> FDetailCategoryImpl::AddAllExternalStructure
 			if (FProperty* Property = PropertyNode->GetProperty())
 			{
 				FDetailLayoutCustomization NewCustomization;
-				NewCustomization.bAdvanced = Location == EPropertyLocation::Advanced;
 				NewCustomization.PropertyRow = MakeShared<FDetailPropertyRow>(PropertyNode, AsShared(), RootPropertyNode);
+				NewCustomization.bAdvanced = Location == EPropertyLocation::Advanced;
 				AddDefaultLayout(NewCustomization, NAME_None);
 
 				Handles.Add(DetailLayoutBuilderRef.GetPropertyHandle(PropertyNode));
@@ -448,20 +448,20 @@ TArray<TSharedPtr<IPropertyHandle>> FDetailCategoryImpl::AddAllExternalStructure
 void FDetailCategoryImpl::AddPropertyNode(TSharedRef<FPropertyNode> PropertyNode, FName InstanceName)
 {
 	FDetailLayoutCustomization NewCustomization;
+	NewCustomization.PropertyRow = MakeShared<FDetailPropertyRow>(PropertyNode, AsShared());
 	NewCustomization.bAdvanced = IsAdvancedLayout(NewCustomization);
-	NewCustomization.PropertyRow = MakeShareable(new FDetailPropertyRow(PropertyNode, AsShared()));
 	AddDefaultLayout(NewCustomization, InstanceName);
 }
 
 bool FDetailCategoryImpl::IsAdvancedLayout(const FDetailLayoutCustomization& LayoutInfo)
 {
-	bool bAdvanced = false;
-	if (LayoutInfo.PropertyRow.IsValid() && LayoutInfo.GetPropertyNode().IsValid() && LayoutInfo.GetPropertyNode()->HasNodeFlags(EPropertyNodeFlags::IsAdvanced))
+	TSharedPtr<FPropertyNode> PropertyNode = LayoutInfo.GetPropertyNode();
+	if (PropertyNode.IsValid() && PropertyNode->HasNodeFlags(EPropertyNodeFlags::IsAdvanced))
 	{
-		bAdvanced = true;
+		return true;
 	}
 
-	return bAdvanced;
+	return false;
 }
 
 void FDetailCategoryImpl::AddCustomLayout(const FDetailLayoutCustomization& LayoutInfo)
