Fix invalid failure message (regression of CL 13954062)

diff --git a/Engine/Source/Editor/PropertyEditor/Private/PropertyHandleImpl.cpp b/Engine/Source/Editor/PropertyEditor/Private/PropertyHandleImpl.cpp
--- a/Engine/Source/Editor/PropertyEditor/Private/PropertyHandleImpl.cpp
+++ b/Engine/Source/Editor/PropertyEditor/Private/PropertyHandleImpl.cpp
@@ -4004,8 +4004,9 @@ FPropertyAccess::Result FPropertyHandleObject::SetValueFromFormattedString(const
 			}
 		}
 
-		// Parsing passed but object is null, we want to set it to null explicitly
-		FPropertyHandleBase::SetValueFromFormattedString(InValue, Flags);
+		// Parsing passed and QualifiedObject !nullptr and bSupportedObject is true and so we should set the value or
+		// Parsing passed but QualifiedObject is nullptr and we want to set it to null explicitly
+		return FPropertyHandleBase::SetValueFromFormattedString(InValue, Flags);
 	}
 
 	// Failed parsing, it's either invalid format or a nonexistent object
