Containers that have "Multiple Values" when multi-selected now correctly refresh the details view after they have been cleared.

The core mistake here seems to have been assuming that GetReadAddress() returned false on failure, but in fact it returns false if the values differ. The original implementer may also have thought that setting CanBeExpanded to true would always show the expander arrow, but this is only the case if the node has children.

diff --git a/Engine/Source/Editor/PropertyEditor/Private/ItemPropertyNode.cpp b/Engine/Source/Editor/PropertyEditor/Private/ItemPropertyNode.cpp
--- a/Engine/Source/Editor/PropertyEditor/Private/ItemPropertyNode.cpp
+++ b/Engine/Source/Editor/PropertyEditor/Private/ItemPropertyNode.cpp
@@ -119,18 +119,15 @@ uint8* FItemPropertyNode::GetValueAddress(uint8* StartAddress, bool bIsSparseDat
  */
 void FItemPropertyNode::InitExpansionFlags (void)
 {
-	
 	FProperty* MyProperty = GetProperty();
 
-	FReadAddressList Addresses;
-
 	bool bExpandableType = CastField<FStructProperty>(MyProperty) 
-		|| ( ( CastField<FArrayProperty>(MyProperty) || CastField<FSetProperty>(MyProperty) || CastField<FMapProperty>(MyProperty) ) && GetReadAddress(false,Addresses) );
+		|| (CastField<FArrayProperty>(MyProperty) || CastField<FSetProperty>(MyProperty) || CastField<FMapProperty>(MyProperty));
 
 	if(	bExpandableType
 		|| HasNodeFlags(EPropertyNodeFlags::EditInlineNew)
 		|| HasNodeFlags(EPropertyNodeFlags::ShowInnerObjectProperties)
-		||	( MyProperty->ArrayDim > 1 && ArrayIndex == -1 ) )
+		|| ( MyProperty->ArrayDim > 1 && ArrayIndex == -1 ) )
 	{
 		SetNodeFlags(EPropertyNodeFlags::CanBeExpanded, true);
 	}
