FPropertyEditorHelpers::GetMetadataOptionsKey  now returns an FName. Internally, FProperty::GetMetaData deals with FName keys, so it makes sense to introduce them as early as possible.

diff --git a/Engine/Source/Editor/PropertyEditor/Private/PropertyEditorHelpers.cpp b/Engine/Source/Editor/PropertyEditor/Private/PropertyEditorHelpers.cpp
--- a/Engine/Source/Editor/PropertyEditor/Private/PropertyEditorHelpers.cpp
+++ b/Engine/Source/Editor/PropertyEditor/Private/PropertyEditorHelpers.cpp
@@ -1229,34 +1229,37 @@ namespace PropertyEditorHelpers
 		}
 	}
 
-	const TCHAR* GetPropertyOptionsMetaDataKey(const FProperty* Property)
+	FName GetPropertyOptionsMetaDataKey(const FProperty* Property)
 	{
 		// Only string and name properties can have options
 		if (Property->IsA(FStrProperty::StaticClass()) || Property->IsA(FNameProperty::StaticClass()))
 		{
 			const FProperty* OwnerProperty = Property->GetOwnerProperty();
-			if (OwnerProperty->HasMetaData(TEXT("GetOptions")))
+			static const FName GetOptionsName("GetOptions");
+			if (OwnerProperty->HasMetaData(GetOptionsName))
 			{
-				return TEXT("GetOptions");
+				return GetOptionsName;
 			}
 
 			// Map properties can have separate options for keys and values
 			const FMapProperty* MapProperty = CastField<FMapProperty>(OwnerProperty);
 			if (MapProperty)
 			{
-				if (MapProperty->HasMetaData(TEXT("GetKeyOptions")) && MapProperty->GetKeyProperty() == Property)
+				static const FName GetKeyOptionsName("GetKeyOptions");
+				if (MapProperty->HasMetaData(GetKeyOptionsName) && MapProperty->GetKeyProperty() == Property)
 				{
-					return TEXT("GetKeyOptions");
+					return GetKeyOptionsName;
 				}
 
-				if (MapProperty->HasMetaData(TEXT("GetValueOptions")) && MapProperty->GetValueProperty() == Property)
+				static const FName GetValueOptionsName("GetValueOptions");
+				if (MapProperty->HasMetaData(GetValueOptionsName) && MapProperty->GetValueProperty() == Property)
 				{
-					return TEXT("GetValueOptions");
+					return GetValueOptionsName;
 				}
 			}
 		}
 
-		return nullptr;
+		return NAME_None;
 	}
 }
 
diff --git a/Engine/Source/Editor/PropertyEditor/Private/PropertyEditorHelpers.h b/Engine/Source/Editor/PropertyEditor/Private/PropertyEditorHelpers.h
--- a/Engine/Source/Editor/PropertyEditor/Private/PropertyEditorHelpers.h
+++ b/Engine/Source/Editor/PropertyEditor/Private/PropertyEditorHelpers.h
@@ -268,8 +268,8 @@ namespace PropertyEditorHelpers
 
 	/**
 	* For properties that support options lists, returns the metadata key which holds the name of the UFunction to call.
-	* Returns nullptr if the property doesn't support, or doesn't have, options.
+	* Returns NAME_None if the property doesn't support, or doesn't have, options.
 	*/
-	const TCHAR* GetPropertyOptionsMetaDataKey(const FProperty* Property);
+	FName GetPropertyOptionsMetaDataKey(const FProperty* Property);
 }
 
diff --git a/Engine/Source/Editor/PropertyEditor/Private/PropertyHandleImpl.cpp b/Engine/Source/Editor/PropertyEditor/Private/PropertyHandleImpl.cpp
--- a/Engine/Source/Editor/PropertyEditor/Private/PropertyHandleImpl.cpp
+++ b/Engine/Source/Editor/PropertyEditor/Private/PropertyHandleImpl.cpp
@@ -2910,7 +2910,9 @@ bool FPropertyHandleBase::GeneratePossibleValues(TArray< TSharedPtr<FString> >&
 			}
 		}
 	}
-	else if (const TCHAR* MetaDataKey = PropertyEditorHelpers::GetPropertyOptionsMetaDataKey(Property))
+
+	FName MetaDataKey = PropertyEditorHelpers::GetPropertyOptionsMetaDataKey(Property);
+	if (!MetaDataKey.IsNone())
 	{
 		const FString GetOptionsFunctionName = Property->GetOwnerProperty()->GetMetaData(MetaDataKey);
 		if (!GetOptionsFunctionName.IsEmpty())
diff --git a/Engine/Source/Editor/PropertyEditor/Private/UserInterface/PropertyEditor/SPropertyEditorCombo.cpp b/Engine/Source/Editor/PropertyEditor/Private/UserInterface/PropertyEditor/SPropertyEditorCombo.cpp
--- a/Engine/Source/Editor/PropertyEditor/Private/UserInterface/PropertyEditor/SPropertyEditorCombo.cpp
+++ b/Engine/Source/Editor/PropertyEditor/Private/UserInterface/PropertyEditor/SPropertyEditorCombo.cpp
@@ -50,7 +50,7 @@ bool SPropertyEditorCombo::Supports( const TSharedRef< class FPropertyEditor >&
 	if(	((Property->IsA(FByteProperty::StaticClass()) && CastField<const FByteProperty>(Property)->Enum)
 		||	Property->IsA(FEnumProperty::StaticClass())
 		|| (Property->IsA(FStrProperty::StaticClass()) && Property->HasMetaData(TEXT("Enum")))
-		|| (PropertyEditorHelpers::GetPropertyOptionsMetaDataKey(Property) != nullptr)
+		|| !PropertyEditorHelpers::GetPropertyOptionsMetaDataKey(Property).IsNone()
 		)
 		&&	( ( ArrayIndex == -1 && Property->ArrayDim == 1 ) || ( ArrayIndex > -1 && Property->ArrayDim > 0 ) ) )
 	{
