When creating a CustomWidget from an FDetailPropertyRow ie. Category.AddProperty(MyProperty).CustomWidget(), we now forward the property handle to the widget row as well.

diff --git a/Engine/Source/Editor/PropertyEditor/Private/DetailPropertyRow.cpp b/Engine/Source/Editor/PropertyEditor/Private/DetailPropertyRow.cpp
--- a/Engine/Source/Editor/PropertyEditor/Private/DetailPropertyRow.cpp
+++ b/Engine/Source/Editor/PropertyEditor/Private/DetailPropertyRow.cpp
@@ -676,6 +676,7 @@ void FDetailPropertyRow::SetWidgetRowProperties(FDetailWidgetRow& Row) const
 
 	Row.EditCondition(EditConditionValue, OnEditConditionValueChanged);
 	Row.CustomResetToDefault = CustomResetToDefault;
+	Row.PropertyHandles.Add(GetPropertyHandle());
 
 	// set custom actions and reset to default
 	if (CustomPropertyWidget.IsValid())
diff --git a/Engine/Source/Editor/PropertyEditor/Private/DetailPropertyRow.h b/Engine/Source/Editor/PropertyEditor/Private/DetailPropertyRow.h
--- a/Engine/Source/Editor/PropertyEditor/Private/DetailPropertyRow.h
+++ b/Engine/Source/Editor/PropertyEditor/Private/DetailPropertyRow.h
@@ -26,7 +26,7 @@ public:
 	FDetailPropertyRow(TSharedPtr<FPropertyNode> InPropertyNode, TSharedRef<FDetailCategoryImpl> InParentCategory, TSharedPtr<FComplexPropertyNode> InExternalRootNode = nullptr);
 
 	/** IDetailPropertyRow interface */
-	virtual TSharedPtr<IPropertyHandle> GetPropertyHandle() override { return PropertyHandle; }
+	virtual TSharedPtr<IPropertyHandle> GetPropertyHandle() const override { return PropertyHandle; }
 	virtual IDetailPropertyRow& DisplayName( const FText& InDisplayName ) override;
 	virtual IDetailPropertyRow& ToolTip( const FText& InToolTip ) override;
 	virtual IDetailPropertyRow& ShowPropertyButtons( bool bInShowPropertyButtons ) override;
diff --git a/Engine/Source/Editor/PropertyEditor/Public/IDetailPropertyRow.h b/Engine/Source/Editor/PropertyEditor/Public/IDetailPropertyRow.h
--- a/Engine/Source/Editor/PropertyEditor/Public/IDetailPropertyRow.h
+++ b/Engine/Source/Editor/PropertyEditor/Public/IDetailPropertyRow.h
@@ -243,7 +243,7 @@ public:
 	virtual ~IDetailPropertyRow(){}
 
 	/** @return the property handle for the property on this row */
-	virtual TSharedPtr<IPropertyHandle> GetPropertyHandle() = 0;
+	virtual TSharedPtr<IPropertyHandle> GetPropertyHandle() const = 0;
 
 	/**
 	 * Sets the localized display name of the property
