FPropertyNode::GetQualifiedName() with bIgnoreCategories = true now doesn't add ".." to the path.

diff --git a/Engine/Source/Editor/PropertyEditor/Private/CategoryPropertyNode.cpp b/Engine/Source/Editor/PropertyEditor/Private/CategoryPropertyNode.cpp
--- a/Engine/Source/Editor/PropertyEditor/Private/CategoryPropertyNode.cpp
+++ b/Engine/Source/Editor/PropertyEditor/Private/CategoryPropertyNode.cpp
@@ -117,25 +117,25 @@ bool FCategoryPropertyNode::GetQualifiedName( FString& PathPlusIndex, const bool
 {
 	bool bAddedAnything = false;
 
-	if( ParentNode && StopParent != ParentNode )
+	if (ParentNode && StopParent != ParentNode)
 	{
 		bAddedAnything = ParentNode->GetQualifiedName(PathPlusIndex, bWithArrayIndex, StopParent, bIgnoreCategories );
-		if( bAddedAnything )
-		{
-			PathPlusIndex += TEXT(".");
-		}
 	}
 	
-	if( !bIgnoreCategories )
+	if (!bIgnoreCategories)
 	{
-		bAddedAnything = true;
+		if (bAddedAnything)
+		{
+			PathPlusIndex += TEXT(".");
+		}
+
 		GetCategoryName().AppendString(PathPlusIndex);
+		bAddedAnything = true;
 	}
 
 	return bAddedAnything;
 }
 
-
 FString FCategoryPropertyNode::GetSubcategoryName() const 
 {
 	FString SubcategoryName;
diff --git a/Engine/Source/Editor/PropertyEditor/Private/PropertyNode.cpp b/Engine/Source/Editor/PropertyEditor/Private/PropertyNode.cpp
--- a/Engine/Source/Editor/PropertyEditor/Private/PropertyNode.cpp
+++ b/Engine/Source/Editor/PropertyEditor/Private/PropertyNode.cpp
@@ -1085,27 +1085,28 @@ bool FPropertyNode::IsOnlyVisibleWhenEditConditionMet() const
 bool FPropertyNode::GetQualifiedName( FString& PathPlusIndex, const bool bWithArrayIndex, const FPropertyNode* StopParent, bool bIgnoreCategories ) const
 {
 	bool bAddedAnything = false;
-	if( ParentNodeWeakPtr.IsValid() && StopParent != ParentNode )
+	if (ParentNodeWeakPtr.IsValid() && StopParent != ParentNode)
 	{
 		bAddedAnything = ParentNode->GetQualifiedName(PathPlusIndex, bWithArrayIndex, StopParent, bIgnoreCategories);
-		if( bAddedAnything )
+	}
+
+	if (Property.IsValid())
+	{
+		if (bAddedAnything)
 		{
 			PathPlusIndex += TEXT(".");
 		}
-	}
 
-	if( Property.IsValid() )
-	{
-		bAddedAnything = true;
 		Property->AppendName(PathPlusIndex);
-	}
 
-	if ( bWithArrayIndex && (ArrayIndex != INDEX_NONE) )
-	{
+		if (bWithArrayIndex && (ArrayIndex != INDEX_NONE))
+		{
+			PathPlusIndex += TEXT("[");
+			PathPlusIndex.AppendInt(ArrayIndex);
+			PathPlusIndex += TEXT("]");
+		}
+
 		bAddedAnything = true;
-		PathPlusIndex += TEXT("[");
-		PathPlusIndex.AppendInt(ArrayIndex);
-		PathPlusIndex += TEXT("]");
 	}
 
 	return bAddedAnything;
@@ -1999,15 +2000,13 @@ bool FPropertyNode::GetDiffersFromDefault()
 		bDiffersFromDefault = false;
 
 		FProperty* Prop = GetProperty();
-
 		if (!Prop)
 		{
 			return bDiffersFromDefault;
 		}
 
-
 		FObjectPropertyNode* ObjectNode = FindObjectItemParent();
-		if(ObjectNode && Property.IsValid() && !IsEditConst())
+		if (ObjectNode && Property.IsValid() && !IsEditConst())
 		{
 			// Get an iterator for the enclosing objects.
 			for(int32 ObjIndex = 0; ObjIndex < ObjectNode->GetNumObjects(); ++ObjIndex)
@@ -2082,18 +2081,18 @@ FString FPropertyNode::GetDefaultValueAsString(bool bUseDisplayName)
 {
 	FObjectPropertyNode* ObjectNode = FindObjectItemParent();
 	FString DefaultValue;
-	if ( ObjectNode && Property.IsValid() )
+	if (ObjectNode && Property.IsValid())
 	{
 		// Get an iterator for the enclosing objects.
-		for ( int32 ObjIndex = 0; ObjIndex < ObjectNode->GetNumObjects(); ++ObjIndex )
+		for (int32 ObjIndex = 0; ObjIndex < ObjectNode->GetNumObjects(); ++ObjIndex)
 		{
 			UObject* Object = ObjectNode->GetUObject( ObjIndex );
 			TSharedPtr<FPropertyItemValueDataTrackerSlate> ValueTracker = GetValueTracker(Object, ObjIndex);
 
-			if( Object && ValueTracker.IsValid() )
+			if (Object && ValueTracker.IsValid())
 			{
 				FString NodeDefaultValue = GetDefaultValueAsStringForObject( *ValueTracker, Object, Property.Get(), bUseDisplayName );
-				if ( DefaultValue.Len() > 0 && NodeDefaultValue.Len() > 0)
+				if (DefaultValue.Len() > 0 && NodeDefaultValue.Len() > 0)
 				{
 					DefaultValue += TEXT(", ");
 				}
@@ -2109,11 +2108,11 @@ FText FPropertyNode::GetResetToDefaultLabel()
 {
 	FString DefaultValue = GetDefaultValueAsString();
 	FText OutLabel = GetDisplayName();
-	if ( DefaultValue.Len() )
+	if (DefaultValue.Len())
 	{
 		const int32 MaxValueLen = 60;
 
-		if ( DefaultValue.Len() > MaxValueLen )
+		if (DefaultValue.Len() > MaxValueLen)
 		{
 			DefaultValue.LeftInline( MaxValueLen, false );
 			DefaultValue += TEXT( "..." );
@@ -2155,11 +2154,11 @@ bool FPropertyNode::AdjustEnumPropDisplayName( UEnum *InEnum, FString& DisplayNa
 {
 	// see if we have alternate text to use for displaying the value
 	UMetaData* PackageMetaData = InEnum->GetOutermost()->GetMetaData();
-	if ( PackageMetaData )
+	if (PackageMetaData)
 	{
 		FName AltDisplayName = FName(*(DisplayName+TEXT(".DisplayName")));
 		FString ValueText = PackageMetaData->GetValue(InEnum, AltDisplayName);
-		if ( ValueText.Len() > 0 )
+		if (ValueText.Len() > 0)
 		{
 			// use the alternate text for this enum value
 			DisplayName = ValueText;
@@ -2255,13 +2254,12 @@ void FPropertyNode::FilterNodes( const TArray<FString>& InFilterStrings, const b
 		SetNodeFlags(EPropertyNodeFlags::IsParentSeenDueToFiltering, true);
 	}
 
-	//default to doing only one pass
-	//bool bCategoryOrObject = (GetObjectNode()) || (GetCategoryNode()!=NULL);
-	int32 StartRecusionPass = HasNodeFlags(EPropertyNodeFlags::IsSeenDueToFiltering) ? 1 : 0;
+	// default to doing only one pass
+	int32 StartRecursionPass = HasNodeFlags(EPropertyNodeFlags::IsSeenDueToFiltering) ? 1 : 0;
 	//Pass 1, if a pass 1 exists (object or category), is to see if there are any children that pass the filter, if any do, trim the tree to the leaves.
 	//	This will stop categories from showing ALL properties if they pass the filter AND a child passes the filter
 	//Pass 0, if no child exists that passes the filter OR this node didn't pass the filter
-	for (int32 RecursionPass = StartRecusionPass; RecursionPass >= 0; --RecursionPass)
+	for (int32 RecursionPass = StartRecursionPass; RecursionPass >= 0; --RecursionPass)
 	{
 		for (int32 scan = 0; scan < ChildNodes.Num(); ++scan)
 		{
@@ -2281,14 +2279,13 @@ void FPropertyNode::FilterNodes( const TArray<FString>& InFilterStrings, const b
 				SetNodeFlags(EPropertyNodeFlags::IsSeenDueToChildFiltering, true);
 			}
 		}
+
 		//now that we've tried a pass at our children, if any of them have been successfully seen due to filtering, just quit now
 		if (HasNodeFlags(EPropertyNodeFlags::IsSeenDueToChildFiltering))
 		{
 			break;
 		}
 	}
-
-	
 }
 
 void FPropertyNode::ProcessSeenFlags(const bool bParentAllowsVisible )
