SDetailsView now listens to the FEditorDelegates::PostUndoRedo event and invalidates any cached state on its property nodes. This causes the Reset to Default buttons to refresh.

diff --git a/Engine/Source/Editor/PropertyEditor/Private/SDetailsView.cpp b/Engine/Source/Editor/PropertyEditor/Private/SDetailsView.cpp
--- a/Engine/Source/Editor/PropertyEditor/Private/SDetailsView.cpp
+++ b/Engine/Source/Editor/PropertyEditor/Private/SDetailsView.cpp
@@ -33,6 +33,8 @@ SDetailsView::~SDetailsView()
 	{
 		SaveExpandedItems(RootNode.ToSharedRef());
 	}
+
+	FEditorDelegates::PostUndoRedo.Remove(PostUndoRedoDelegateHandle);
 };
 
 /**
@@ -54,6 +56,8 @@ void SDetailsView::Construct(const FArguments& InArgs, const FDetailsViewArgs& I
 	PropertyUtilities = MakeShareable( new FPropertyDetailsUtilities( *this ) );
 	PropertyGenerationUtilities = MakeShareable( new FDetailsViewPropertyGenerationUtilities(*this) );
 
+	PostUndoRedoDelegateHandle = FEditorDelegates::PostUndoRedo.AddSP(this, &SDetailsView::OnPostUndoRedo);
+
 	// We want the scrollbar to always be visible when objects are selected, but not when there is no selection - however:
 	//  - We can't use AlwaysShowScrollbar for this, as this will also show the scrollbar when nothing is selected
 	//  - We can't use the Visibility construction parameter, as it gets translated into user visibility and can hide the scrollbar even when objects are selected
@@ -936,6 +940,14 @@ void SDetailsView::SetOnObjectArrayChanged(FOnObjectArrayChanged OnObjectArrayCh
 	OnObjectArrayChanged = OnObjectArrayChangedDelegate;
 }
 
+void SDetailsView::OnPostUndoRedo()
+{
+	for (const TSharedPtr<FComplexPropertyNode>& RootNode : RootPropertyNodes)
+	{
+		RootNode->InvalidateCachedState();
+	}
+}
+
 bool SDetailsView::IsConnected() const
 {
 	return GetNumObjects() > 0;
diff --git a/Engine/Source/Editor/PropertyEditor/Private/SDetailsView.h b/Engine/Source/Editor/PropertyEditor/Private/SDetailsView.h
--- a/Engine/Source/Editor/PropertyEditor/Private/SDetailsView.h
+++ b/Engine/Source/Editor/PropertyEditor/Private/SDetailsView.h
@@ -167,6 +167,9 @@ private:
 	/** Called when the toggle favorites button is clicked. */
 	FReply OnToggleFavoritesClicked();
 
+	/** Called after an undo or redo operation occurs in the editor. */
+	void OnPostUndoRedo();
+
 private:
 	/** The filter for objects viewed by this details panel */
 	TSharedPtr<FDetailsViewObjectFilter> ObjectFilter;
@@ -193,4 +196,6 @@ private:
 	TSharedPtr<IDetailRootObjectCustomization> RootObjectCustomization;
 	/** True if at least one viewed object is a CDO (blueprint editing) */
 	bool bViewingClassDefaultObject;
+	/** Delegate handle for unregistering from the PostUndoRedo event. */
+	FDelegateHandle PostUndoRedoDelegateHandle;
 };
