SPropertyEditorAsset now checks for null when setting an actor. This is likely more fallout from the ill-fated UInterface-as-actor PR.

diff --git a/Engine/Source/Editor/PropertyEditor/Private/UserInterface/PropertyEditor/SPropertyEditorAsset.cpp b/Engine/Source/Editor/PropertyEditor/Private/UserInterface/PropertyEditor/SPropertyEditorAsset.cpp
--- a/Engine/Source/Editor/PropertyEditor/Private/UserInterface/PropertyEditor/SPropertyEditorAsset.cpp
+++ b/Engine/Source/Editor/PropertyEditor/Private/UserInterface/PropertyEditor/SPropertyEditorAsset.cpp
@@ -771,7 +771,7 @@ TSharedRef<SWidget> SPropertyEditorAsset::OnGetMenuContent()
 	FObjectOrAssetData Value;
 	GetValue(Value);
 
-	if(bIsActor)
+	if (bIsActor)
 	{
 		return PropertyCustomizationHelpers::MakeActorPickerWithMenu(Cast<AActor>(Value.Object),
 																	 bAllowClear,
@@ -797,7 +797,7 @@ TSharedRef<SWidget> SPropertyEditorAsset::OnGetMenuContent()
 
 void SPropertyEditorAsset::OnMenuOpenChanged(bool bOpen)
 {
-	if ( bOpen == false )
+	if (bOpen == false)
 	{
 		AssetComboButton->SetMenuContent(SNullWidget::NullWidget);
 	}
@@ -805,7 +805,7 @@ void SPropertyEditorAsset::OnMenuOpenChanged(bool bOpen)
 
 bool SPropertyEditorAsset::IsFilteredActor( const AActor* const Actor ) const
 {
-	bool IsAllowed = Actor->IsA(ObjectClass) && !Actor->IsChildActor() && IsClassAllowed(Actor->GetClass());
+	bool IsAllowed = Actor != nullptr && Actor->IsA(ObjectClass) && !Actor->IsChildActor() && IsClassAllowed(Actor->GetClass());
 	return IsAllowed;
 }
 
@@ -936,7 +936,7 @@ void SPropertyEditorAsset::SetValue( const FAssetData& AssetData )
 {
 	AssetComboButton->SetIsOpen(false);
 
-	if(CanSetBasedOnCustomClasses(AssetData))
+	if (CanSetBasedOnCustomClasses(AssetData))
 	{
 		FText AssetReferenceFilterFailureReason;
 		if (CanSetBasedOnAssetReferenceFilter(AssetData, &AssetReferenceFilterFailureReason))
