Restored Create Asset options to the static mesh's Materials section.

The culprit was the new NoCreate metadata, which was evaluating to true if no property was given - which is the case with the material list.

diff --git a/Engine/Source/Editor/PropertyEditor/Private/UserInterface/PropertyEditor/SPropertyEditorAsset.cpp b/Engine/Source/Editor/PropertyEditor/Private/UserInterface/PropertyEditor/SPropertyEditorAsset.cpp
--- a/Engine/Source/Editor/PropertyEditor/Private/UserInterface/PropertyEditor/SPropertyEditorAsset.cpp
+++ b/Engine/Source/Editor/PropertyEditor/Private/UserInterface/PropertyEditor/SPropertyEditorAsset.cpp
@@ -346,7 +346,8 @@ void SPropertyEditorAsset::Construct(const FArguments& InArgs, const TSharedPtr<
 
 	bIsActor = ObjectClass->IsChildOf(AActor::StaticClass());
 
-	if (Property && !Property->HasMetaData("NoCreate"))
+	const bool bAllowCreation = !Property || !Property->HasMetaData("NoCreate");
+	if (bAllowCreation) 
 	{
 		if (InArgs._NewAssetFactories.IsSet())
 		{
