Restored highlighting for changed properties in the diff view.

diff --git a/Engine/Source/Editor/EditorStyle/Private/StarshipStyle.cpp b/Engine/Source/Editor/EditorStyle/Private/StarshipStyle.cpp
--- a/Engine/Source/Editor/EditorStyle/Private/StarshipStyle.cpp
+++ b/Engine/Source/Editor/EditorStyle/Private/StarshipStyle.cpp
@@ -3168,6 +3168,7 @@ void FStarshipEditorStyle::FStyle::SetupPropertyEditorStyles()
 		
 		// these are not actually displayed as white, see PropertyEditorConstants::GetRowBackgroundColor
 		Set( "DetailsView.CategoryMiddle",            new FSlateColorBrush(FStyleColors::White));
+		Set( "DetailsView.Highlight",				  new FSlateRoundedBoxBrush(FStyleColors::Transparent, 0.0f, FStyleColors::AccentBlue, 1.0f));
 
 		Set( "DetailsView.PropertyIsFavorite", new IMAGE_BRUSH("PropertyView/Favorites_Enabled", Icon12x12));
 		Set( "DetailsView.PropertyIsNotFavorite", new IMAGE_BRUSH("PropertyView/Favorites_Disabled", Icon12x12));
diff --git a/Engine/Source/Editor/PropertyEditor/Private/SDetailSingleItemRow.cpp b/Engine/Source/Editor/PropertyEditor/Private/SDetailSingleItemRow.cpp
--- a/Engine/Source/Editor/PropertyEditor/Private/SDetailSingleItemRow.cpp
+++ b/Engine/Source/Editor/PropertyEditor/Private/SDetailSingleItemRow.cpp
@@ -521,14 +521,19 @@ void SDetailSingleItemRow::Construct( const FArguments& InArgs, FDetailLayoutCus
 	auto GetScrollbarWellTint = [this, OwnerTableViewWeak]()
 	{
 		return SDetailTableRowBase::IsScrollBarVisible(OwnerTableViewWeak) ?
-			FSlateColor(EStyleColor::White) : 
+			FStyleColors::White : 
 			this->GetOuterBackgroundColor();
 	};
 
+	auto GetHighlightBorderPadding = [this]()
+	{
+		return this->IsHighlighted() ? FMargin(1) : FMargin(0);
+	};
+
 	this->ChildSlot
 	[
 		SNew( SBorder )
-		.BorderImage(FAppStyle::Get().GetBrush( "DetailsView.GridLine"))
+		.BorderImage(FAppStyle::Get().GetBrush("DetailsView.GridLine"))
 		.Padding(FMargin(0,0,0,1))
 		[
 			SNew( SHorizontalBox )
@@ -536,12 +541,17 @@ void SDetailSingleItemRow::Construct( const FArguments& InArgs, FDetailLayoutCus
 			.HAlign(HAlign_Fill)
 			[
 				SNew( SBorder )
-				.BorderImage(FAppStyle::Get().GetBrush("DetailsView.CategoryMiddle"))
-				.BorderBackgroundColor(this, &SDetailSingleItemRow::GetOuterBackgroundColor)
-				.Padding(0)
-				.Clipping(EWidgetClipping::ClipToBounds)
+				.BorderImage(FAppStyle::Get().GetBrush("DetailsView.Highlight"))
+				.Padding_Lambda(GetHighlightBorderPadding)
 				[
-					Widget
+					SNew( SBorder )
+					.BorderImage(FAppStyle::Get().GetBrush("DetailsView.CategoryMiddle"))
+					.BorderBackgroundColor(this, &SDetailSingleItemRow::GetOuterBackgroundColor)
+					.Padding(0)
+					.Clipping(EWidgetClipping::ClipToBounds)
+					[
+						Widget
+					]
 				]
 			]
 			+ SHorizontalBox::Slot()
@@ -549,7 +559,7 @@ void SDetailSingleItemRow::Construct( const FArguments& InArgs, FDetailLayoutCus
 			.VAlign(VAlign_Fill)
 			.AutoWidth()
 			[
-				SNew(SBorder)
+				SNew( SBorder )
 				.BorderImage_Lambda(GetScrollbarWellBrush)
 				.BorderBackgroundColor_Lambda(GetScrollbarWellTint)
 				.Padding(FMargin(0, 0, SDetailTableRowBase::ScrollBarPadding, 0))
@@ -592,33 +602,28 @@ FSlateColor SDetailSingleItemRow::GetOuterBackgroundColor() const
 {
 	if (IsHighlighted())
 	{
-		return FAppStyle::Get().GetSlateColor("Colors.Panel");
+		return FAppStyle::Get().GetSlateColor("Colors.Hover");
 	}
-	else if (bIsDragDropObject)
+
+	if (bIsHoveredDragTarget)
 	{
-		return FAppStyle::Get().GetSlateColor("Colors.Panel");
+		return FAppStyle::Get().GetSlateColor("Colors.Hover2");
 	}
-	else if (IsHovered())
+
+	if (IsHovered())
 	{
-		if (bIsHoveredDragTarget)
-		{
-			return FAppStyle::Get().GetSlateColor("Colors.Panel");
-		}
-		else
-		{
-			return FAppStyle::Get().GetSlateColor("Colors.Header");
-		}
+		return FAppStyle::Get().GetSlateColor("Colors.Header");
 	}
-	
+
 	return FAppStyle::Get().GetSlateColor("Colors.Panel");
 }
 
 /** Get the background color of the inner part of the row, which contains the name and value widgets. */
 FSlateColor SDetailSingleItemRow::GetInnerBackgroundColor() const
 {
-	if (IsHovered() && !bIsHoveredDragTarget)
+	if (IsHighlighted())
 	{
-		return FAppStyle::Get().GetSlateColor("Colors.Header");
+		return FAppStyle::Get().GetSlateColor("Colors.Hover");
 	}
 
 	if (bIsHoveredDragTarget)
@@ -626,6 +631,11 @@ FSlateColor SDetailSingleItemRow::GetInnerBackgroundColor() const
 		return FAppStyle::Get().GetSlateColor("Colors.Hover2");
 	}
 
+	if (IsHovered())
+	{
+		return FAppStyle::Get().GetSlateColor("Colors.Header");
+	}
+	
 	if (bIsDragDropObject)
 	{
 		return FAppStyle::Get().GetSlateColor("Colors.Hover");
