Edit conditions for properties that don't support edit condition toggles no longer checkSlow when retrieving the check state, but rather only when setting it.

CustomEditConditionValueChanged handlers are now postfixed with a call to FPropertyNode::InvalidateCachedState(), which causes the IsEditConst state to refresh.

Removed unnecessary code to access the PropertyEditor->IsEditConditionMet(), which is stored in EditConditionValue now.

diff --git a/Engine/Source/Editor/PropertyEditor/Private/DetailPropertyRow.cpp b/Engine/Source/Editor/PropertyEditor/Private/DetailPropertyRow.cpp
--- a/Engine/Source/Editor/PropertyEditor/Private/DetailPropertyRow.cpp
+++ b/Engine/Source/Editor/PropertyEditor/Private/DetailPropertyRow.cpp
@@ -619,10 +619,6 @@ bool FDetailPropertyRow::GetEnabledState() const
 		{
 			Result = Result && CustomEditConditionValue.Get();
 		}
-		else
-		{
-			Result = Result && PropertyEditor->IsEditConditionMet();
-		}
 	}
 	
 	Result = Result && CustomIsEnabledAttrib.Get();
@@ -658,6 +654,17 @@ static void TogglePropertyEditorEditCondition(bool bValue, TWeakPtr<FPropertyEdi
 	}
 }
 
+static void ExecuteCustomEditConditionToggle(bool bValue, FOnBooleanValueChanged CustomEditConditionToggle, TWeakPtr<FPropertyEditor> PropertyEditorWeak)
+{
+	CustomEditConditionToggle.ExecuteIfBound(bValue);
+
+	TSharedPtr<FPropertyEditor> PropertyEditorPtr = PropertyEditorWeak.Pin();
+	if (PropertyEditorPtr.IsValid())
+	{
+		 PropertyEditorPtr->GetPropertyNode()->InvalidateCachedState();
+	}
+}
+
 void FDetailPropertyRow::SetWidgetRowProperties(FDetailWidgetRow& Row) const
 {
 	// set edit condition handlers - use customized if provided
@@ -667,8 +674,13 @@ void FDetailPropertyRow::SetWidgetRowProperties(FDetailWidgetRow& Row) const
 		EditConditionValue = TAttribute<bool>(PropertyEditor.ToSharedRef(), &FPropertyEditor::IsEditConditionMet);
 	}
 
-	FOnBooleanValueChanged OnEditConditionValueChanged = CustomEditConditionValueChanged;
-	if (!OnEditConditionValueChanged.IsBound() && PropertyEditor->SupportsEditConditionToggle())
+	FOnBooleanValueChanged OnEditConditionValueChanged;
+	if (CustomEditConditionValueChanged.IsBound())
+	{
+		TWeakPtr<FPropertyEditor> PropertyEditorWeak = PropertyEditor;
+		OnEditConditionValueChanged = FOnBooleanValueChanged::CreateStatic(&ExecuteCustomEditConditionToggle, CustomEditConditionValueChanged, PropertyEditorWeak);
+	}
+	else if (PropertyEditor->SupportsEditConditionToggle())
 	{
 		TWeakPtr<FPropertyEditor> PropertyEditorWeak = PropertyEditor;
 		OnEditConditionValueChanged = FOnBooleanValueChanged::CreateStatic(&TogglePropertyEditorEditCondition, PropertyEditorWeak);
