Allow FNames to be initialized with a string that starts with an underscores from the editor.

https://p4-swarm.epicgames.net/reviews/15934265

diff --git a/Engine/Source/Editor/PropertyEditor/Private/PropertyHandleImpl.cpp b/Engine/Source/Editor/PropertyEditor/Private/PropertyHandleImpl.cpp
--- a/Engine/Source/Editor/PropertyEditor/Private/PropertyHandleImpl.cpp
+++ b/Engine/Source/Editor/PropertyEditor/Private/PropertyHandleImpl.cpp
@@ -707,18 +707,10 @@ FPropertyAccess::Result FPropertyValueImpl::SetValueAsString( const FString& InV
 
 		FString Value = InValue;
 
-		// Strip any leading underscores, and trim the name
+		// Trim the name
 		if( NodeProperty && NodeProperty->IsA( FNameProperty::StaticClass() ) )
 		{
-			// trim whitespace and leading underscores
-			int32 Pos = 0;
-			while (Pos < Value.Len() && 
-				(FChar::IsWhitespace(Value[Pos]) || Value[Pos] == '_'))
-			{
-				++Pos;
-			}
-			Value.RemoveAt(0, Pos, false);
-			Value.TrimEndInline();
+			Value.TrimStartAndEndInline();
 		}
 
 		// If more than one object is selected, an empty field indicates their values for this property differ.
@@ -2894,14 +2886,14 @@ bool FPropertyHandleBase::GeneratePossibleValues(TArray< TSharedPtr<FString> >&
 
 						FCachedPropertyPath Path(GetOptionsFunctionName);
 						if (!PropertyPathHelpers::GetPropertyValue(Target, Path, StringOptions))
+					{
+						TArray<FName> NameOptions;
+						if (PropertyPathHelpers::GetPropertyValue(Target, Path, NameOptions))
 						{
-							TArray<FName> NameOptions;
-							if (PropertyPathHelpers::GetPropertyValue(Target, Path, NameOptions))
-							{
-								Algo::Transform(NameOptions, StringOptions, [](const FName& InName) { return InName.ToString(); });
-							}
+							Algo::Transform(NameOptions, StringOptions, [](const FName& InName) { return InName.ToString(); });
 						}
 					}
+					}
 
 					// If this is the first time there won't be any options.
 					if (OptionIntersection.Num() == 0)
@@ -3972,8 +3964,8 @@ FPropertyAccess::Result FPropertyHandleObject::SetValueFromFormattedString(const
 		return FPropertyHandleBase::SetValueFromFormattedString(InValue, Flags);
 	}
 
-	// Failed parsing, it's either invalid format or a nonexistent object
-	return FPropertyAccess::Fail;
+		// Failed parsing, it's either invalid format or a nonexistent object
+		return FPropertyAccess::Fail;
 }
 
 FPropertyAccess::Result FPropertyHandleObject::SetObjectValueFromSelection()
