Renamed bShowOnlyDiffering to bShowOnlyWhitelisted - that it's only used for diffing is a detail that the details view shouldn't know about.

diff --git a/Engine/Source/Editor/PropertyEditor/Private/DetailItemNode.cpp b/Engine/Source/Editor/PropertyEditor/Private/DetailItemNode.cpp
--- a/Engine/Source/Editor/PropertyEditor/Private/DetailItemNode.cpp
+++ b/Engine/Source/Editor/PropertyEditor/Private/DetailItemNode.cpp
@@ -540,8 +540,11 @@ static bool PassesAllFilters( FDetailItemNode* ItemNode, const FDetailLayoutCust
 
 	bool bPassesAllFilters = true;
 
-	if( InFilter.FilterStrings.Num() > 0 || InFilter.bShowOnlyModifiedProperties == true || InFilter.bShowOnlyDiffering == true 
-		|| InFilter.bShowKeyable == true || InFilter.bShowAnimated == true)
+	if( InFilter.FilterStrings.Num() > 0 || 
+		InFilter.bShowOnlyModifiedProperties == true || 
+		InFilter.bShowOnlyDiffering == true || 
+		InFilter.bShowKeyable == true || 
+		InFilter.bShowAnimated == true)
 	{
 		const bool bSearchFilterIsEmpty = InFilter.FilterStrings.Num() == 0;
 
@@ -559,7 +562,7 @@ static bool PassesAllFilters( FDetailItemNode* ItemNode, const FDetailLayoutCust
 
 			const bool bPassesSearchFilter = bPassesCategoryFilter || bPassesValueFilter || bSearchFilterIsEmpty || ( bIsNotBeingFiltered || bIsSeenDueToFiltering || bIsParentSeenDueToFiltering );
 			const bool bPassesModifiedFilter = bPassesSearchFilter && ( InFilter.bShowOnlyModifiedProperties == false || PropertyNodePin->GetDiffersFromDefault() == true );
-			const bool bPassesDifferingFilter = InFilter.bShowOnlyDiffering ? InFilter.WhitelistedProperties.Find(*FPropertyNode::CreatePropertyPath(PropertyNodePin.ToSharedRef())) != nullptr : true;
+			const bool bPassesWhitelistedFilter = InFilter.bShowOnlyWhitelisted ? InFilter.WhitelistedProperties.Contains(*FPropertyNode::CreatePropertyPath(PropertyNodePin.ToSharedRef())) : true;
 
 			bool bPassesKeyableFilter = true;
 			if (InFilter.bShowKeyable)
@@ -578,13 +581,13 @@ static bool PassesAllFilters( FDetailItemNode* ItemNode, const FDetailLayoutCust
 			const bool bPassesAnimatedFilter = (InFilter.bShowAnimated == false || Local::ItemIsAnimated(ItemNode, PropertyNodePin));
 
 			// The property node is visible (note categories are never visible unless they have a child that is visible )
-			bPassesAllFilters = bPassesSearchFilter && bPassesModifiedFilter && bPassesDifferingFilter && bPassesKeyableFilter && bPassesAnimatedFilter;
+			bPassesAllFilters = bPassesSearchFilter && bPassesModifiedFilter && bPassesWhitelistedFilter && bPassesKeyableFilter && bPassesAnimatedFilter;
 		}
 		else if (InCustomization.HasCustomWidget())
 		{
 			const bool bPassesTextFilter = bPassesCategoryFilter || bPassesValueFilter || Local::StringPassesFilter(InFilter, InCustomization.WidgetDecl->FilterTextString.ToString());
 			const bool bPassesModifiedFilter = InFilter.bShowOnlyModifiedProperties == false || InCustomization.WidgetDecl->EditConditionValue.Get(false);
-			//@todo we need to support custom widgets for keyable,animated, in particularly for transforms(ComponentTransformDetails).
+			//@todo we need to support custom widgets for keyable, animated, in particular for transforms(ComponentTransformDetails).
 			const bool bPassesKeyableFilter = (InFilter.bShowKeyable == false);
 			const bool bPassesAnimatedFilter = (InFilter.bShowAnimated == false);
 			bPassesAllFilters = bPassesTextFilter && bPassesModifiedFilter && bPassesKeyableFilter && bPassesAnimatedFilter;
@@ -592,7 +595,7 @@ static bool PassesAllFilters( FDetailItemNode* ItemNode, const FDetailLayoutCust
 		else if (InCustomization.HasCustomBuilder())
 		{
 			const bool bPassesTextFilter = bPassesCategoryFilter || bPassesValueFilter || Local::StringPassesFilter(InFilter, InCustomization.CustomBuilderRow->GetWidgetRow().FilterTextString.ToString());
-			//@todo we need to support custom builders for modified, keyable, animated, in particularly for transforms(ComponentTransformDetails).
+			//@todo we need to support custom builders for modified, keyable, animated, in particular for transforms(ComponentTransformDetails).
 			const bool bPassesModifiedFilter = (InFilter.bShowOnlyModifiedProperties == false);
 			const bool bPassesKeyableFilter = (InFilter.bShowKeyable == false);
 			const bool bPassesAnimatedFilter = (InFilter.bShowAnimated == false);
diff --git a/Engine/Source/Editor/PropertyEditor/Private/SDetailsView.cpp b/Engine/Source/Editor/PropertyEditor/Private/SDetailsView.cpp
--- a/Engine/Source/Editor/PropertyEditor/Private/SDetailsView.cpp
+++ b/Engine/Source/Editor/PropertyEditor/Private/SDetailsView.cpp
@@ -108,9 +108,9 @@ void SDetailsView::Construct(const FArguments& InArgs, const FDetailsViewArgs& I
 				LOCTEXT("ShowOnlyDiffering_ToolTip", "Displays only properties in this instance which have been changed or added from the instance being compared"),
 				FSlateIcon(),
 				FUIAction(
-					FExecuteAction::CreateSP(this, &SDetailsView::OnShowOnlyDifferingClicked),
+					FExecuteAction::CreateSP(this, &SDetailsView::OnShowOnlyWhitelistedClicked),
 					FCanExecuteAction(),
-					FIsActionChecked::CreateSP(this, &SDetailsView::IsShowOnlyDifferingChecked)
+					FIsActionChecked::CreateSP(this, &SDetailsView::IsShowOnlyWhitelistedChecked)
 				),
 				NAME_None,
 				EUserInterfaceActionType::ToggleButton
diff --git a/Engine/Source/Editor/PropertyEditor/Private/SDetailsViewBase.cpp b/Engine/Source/Editor/PropertyEditor/Private/SDetailsViewBase.cpp
--- a/Engine/Source/Editor/PropertyEditor/Private/SDetailsViewBase.cpp
+++ b/Engine/Source/Editor/PropertyEditor/Private/SDetailsViewBase.cpp
@@ -606,9 +606,9 @@ void SDetailsViewBase::OnShowAllAdvancedClicked()
 	UpdateFilteredDetails();
 }
 
-void SDetailsViewBase::OnShowOnlyDifferingClicked()
+void SDetailsViewBase::OnShowOnlyWhitelistedClicked()
 {
-	CurrentFilter.bShowOnlyDiffering = !CurrentFilter.bShowOnlyDiffering;
+	CurrentFilter.bShowOnlyWhitelisted = !CurrentFilter.bShowOnlyWhitelisted;
 
 	UpdateFilteredDetails();
 }
diff --git a/Engine/Source/Editor/PropertyEditor/Private/SDetailsViewBase.h b/Engine/Source/Editor/PropertyEditor/Private/SDetailsViewBase.h
--- a/Engine/Source/Editor/PropertyEditor/Private/SDetailsViewBase.h
+++ b/Engine/Source/Editor/PropertyEditor/Private/SDetailsViewBase.h
@@ -34,25 +34,31 @@ class IPropertyGenerationUtilities;
 class SSearchBox;
 struct FDetailsViewObjectRoot;
 
-
 typedef STreeView< TSharedRef<class FDetailTreeNode> > SDetailTree;
 
-
-
 /** Represents a filter which controls the visibility of items in the details view */
 struct FDetailFilter
 {
 	FDetailFilter()
 		: bShowOnlyModifiedProperties(false)
 		, bShowAllAdvanced(false)
-		, bShowOnlyDiffering(false)
 		, bShowAllChildrenIfCategoryMatches(true)
 		, bShowKeyable(false)
 		, bShowAnimated(false)
 		, bShowFavoritesCategory(false)
+		, bShowOnlyWhitelisted(false)
 	{}
 
-	bool IsEmptyFilter() const { return FilterStrings.Num() == 0 && bShowOnlyModifiedProperties == false && bShowAllAdvanced == false && bShowOnlyDiffering == false && bShowAllChildrenIfCategoryMatches == false && bShowKeyable == false && bShowAnimated == false; }
+	bool IsEmptyFilter() const 
+	{ 
+		return FilterStrings.Num() == 0 
+			&& bShowOnlyModifiedProperties == false 
+			&& bShowAllAdvanced == false 
+			&& bShowOnlyWhitelisted == false 
+			&& bShowAllChildrenIfCategoryMatches == false 
+			&& bShowKeyable == false 
+			&& bShowAnimated == false;
+	}
 
 	/** Any user search terms that items must match */
 	TArray<FString> FilterStrings;
@@ -60,8 +66,6 @@ struct FDetailFilter
 	bool bShowOnlyModifiedProperties;
 	/** If we should show all advanced properties */
 	bool bShowAllAdvanced;
-	/** If we should only show differing properties */
-	bool bShowOnlyDiffering;
 	/** If we should show all the children if their category name matches the search */
 	bool bShowAllChildrenIfCategoryMatches;
 	/** If we should only show keyable properties */
@@ -70,6 +74,9 @@ struct FDetailFilter
 	bool bShowAnimated;
 	/** If we should show the favorites category. */
 	bool bShowFavoritesCategory;
+	/** If we should only show whitelisted properties */
+	bool bShowOnlyWhitelisted;
+	/** The set of whitelisted properties to show. */
 	TSet<FPropertyPath> WhitelistedProperties;
 };
 
@@ -270,7 +277,7 @@ protected:
 	bool IsShowAllAdvancedChecked() const { return CurrentFilter.bShowAllAdvanced; }
 
 	/** @return true if show only differing is checked */
-	bool IsShowOnlyDifferingChecked() const { return CurrentFilter.bShowOnlyDiffering; }
+	bool IsShowOnlyWhitelistedChecked() const { return CurrentFilter.bShowOnlyWhitelisted; }
 
 	/** @return true if show all advanced is checked */
 	bool IsShowAllChildrenIfCategoryMatchesChecked() const { return CurrentFilter.bShowAllChildrenIfCategoryMatches; }
@@ -291,7 +298,7 @@ protected:
 	void OnShowAllAdvancedClicked();
 
 	/** Called when show only differing is clicked */
-	void OnShowOnlyDifferingClicked();
+	void OnShowOnlyWhitelistedClicked();
 
 	/** Called when show all children if category matches is clicked */
 	void OnShowAllChildrenIfCategoryMatchesClicked();
