Missed change.

diff --git a/Engine/Source/Editor/PropertyEditor/Private/DetailItemNode.cpp b/Engine/Source/Editor/PropertyEditor/Private/DetailItemNode.cpp
--- a/Engine/Source/Editor/PropertyEditor/Private/DetailItemNode.cpp
+++ b/Engine/Source/Editor/PropertyEditor/Private/DetailItemNode.cpp
@@ -542,7 +542,7 @@ static bool PassesAllFilters( FDetailItemNode* ItemNode, const FDetailLayoutCust
 
 	if( InFilter.FilterStrings.Num() > 0 || 
 		InFilter.bShowOnlyModifiedProperties == true || 
-		InFilter.bShowOnlyDiffering == true || 
+		InFilter.bShowOnlyWhitelisted == true || 
 		InFilter.bShowKeyable == true || 
 		InFilter.bShowAnimated == true)
 	{
