Fixed crash when unfavoriting favorited custom builder rows from the Favorites category.

Custom builders now store the original path they were added from. This crash happened because the favorited custom builder row was resolving to "Favorites.BuilderName", rather than the "Category.BuilderName" that was stored. They were still unfavoriteable from the original location.

diff --git a/Engine/Source/Editor/PropertyEditor/Private/DetailCategoryBuilderImpl.cpp b/Engine/Source/Editor/PropertyEditor/Private/DetailCategoryBuilderImpl.cpp
--- a/Engine/Source/Editor/PropertyEditor/Private/DetailCategoryBuilderImpl.cpp
+++ b/Engine/Source/Editor/PropertyEditor/Private/DetailCategoryBuilderImpl.cpp
@@ -168,21 +168,31 @@ void FDetailCategoryImpl::AddCustomBuilder(TSharedRef<IDetailCustomNodeBuilder>
 	NewCustomization.bAdvanced = bForAdvanced;
 	NewCustomization.CustomBuilderRow = MakeShareable(new FDetailCustomBuilderRow(InCustomBuilder));
 
-	if (GetDetailsView() != nullptr && !bFavoriteCategory)
+	if (!InCustomBuilder->GetName().IsNone())
 	{
 		TStringBuilder<256> PathToNode;
 		PathToNode.Append(GetCategoryPathName());
 		PathToNode.Append(TEXT("."));
 		PathToNode.Append(InCustomBuilder->GetName().ToString());
+		NewCustomization.CustomBuilderRow->SetOriginalPath(PathToNode.ToString());
+	}
 
-		if (GetDetailsView()->IsCustomBuilderFavorite(PathToNode))
+	if (GetDetailsView() != nullptr && !bFavoriteCategory)
+	{
+		if (GetDetailsView()->IsCustomBuilderFavorite(NewCustomization.CustomBuilderRow->GetOriginalPath()))
 		{
 			TSharedPtr<FDetailLayoutBuilderImpl> LayoutBuilder = DetailLayoutBuilder.Pin();
 			if (LayoutBuilder.IsValid())
 			{
 				static const FName FavoritesCategoryName(TEXT("Favorites"));
 				FDetailCategoryImpl& FavoritesCategory = LayoutBuilder->DefaultCategory(FavoritesCategoryName);
-				FavoritesCategory.AddCustomBuilder(InCustomBuilder, false);
+
+				FDetailLayoutCustomization FavoritesCustomization;
+				FavoritesCustomization.bCustom = true;
+				FavoritesCustomization.bAdvanced = false;
+				FavoritesCustomization.CustomBuilderRow = MakeShareable(new FDetailCustomBuilderRow(InCustomBuilder));
+				FavoritesCustomization.CustomBuilderRow->SetOriginalPath(NewCustomization.CustomBuilderRow->GetOriginalPath());
+				FavoritesCategory.AddCustomLayout(FavoritesCustomization);
 			}
 		}
 	}
diff --git a/Engine/Source/Editor/PropertyEditor/Private/DetailCategoryBuilderImpl.h b/Engine/Source/Editor/PropertyEditor/Private/DetailCategoryBuilderImpl.h
--- a/Engine/Source/Editor/PropertyEditor/Private/DetailCategoryBuilderImpl.h
+++ b/Engine/Source/Editor/PropertyEditor/Private/DetailCategoryBuilderImpl.h
@@ -300,6 +300,9 @@ public:
 	 */
 	void SetCategoryAsSpecialFavorite() { bFavoriteCategory = true; bForceAdvanced = true; }
 
+	/** Is this the Favorites category? */
+	bool IsFavoriteCategory() const { return bFavoriteCategory; }
+
 	FDetailLayoutCustomization* GetDefaultCustomization(TSharedRef<FPropertyNode> PropertyNode);
 
 private:
diff --git a/Engine/Source/Editor/PropertyEditor/Private/DetailCustomBuilderRow.h b/Engine/Source/Editor/PropertyEditor/Private/DetailCustomBuilderRow.h
--- a/Engine/Source/Editor/PropertyEditor/Private/DetailCustomBuilderRow.h
+++ b/Engine/Source/Editor/PropertyEditor/Private/DetailCustomBuilderRow.h
@@ -29,6 +29,8 @@ public:
 	bool IsInitiallyCollapsed() const;
 	FDetailWidgetRow GetWidgetRow();
 	bool AreChildCustomizationsHidden() const;
+	void SetOriginalPath(FStringView Path) { OriginalPath = Path; }
+	const FString& GetOriginalPath() const { return OriginalPath; }
 
 private:
 	/** Whether or not our parent is enabled */
@@ -37,4 +39,5 @@ private:
 	TSharedRef<class IDetailCustomNodeBuilder> CustomNodeBuilder;
 	TSharedPtr<class FCustomChildrenBuilder> ChildrenBuilder;
 	TWeakPtr<FDetailCategoryImpl> ParentCategory;
+	FString OriginalPath;
 };
diff --git a/Engine/Source/Editor/PropertyEditor/Private/SDetailSingleItemRow.cpp b/Engine/Source/Editor/PropertyEditor/Private/SDetailSingleItemRow.cpp
--- a/Engine/Source/Editor/PropertyEditor/Private/SDetailSingleItemRow.cpp
+++ b/Engine/Source/Editor/PropertyEditor/Private/SDetailSingleItemRow.cpp
@@ -850,7 +850,7 @@ bool SDetailSingleItemRow::CanFavorite() const
 
 	if (Customization->HasCustomBuilder())
 	{
-		return Customization->CustomBuilderRow->GetCustomBuilderName() != NAME_None;
+		return Customization->CustomBuilderRow->GetOriginalPath().IsEmpty() == false;
 	}
 
 	return false;
@@ -871,15 +871,14 @@ bool SDetailSingleItemRow::IsFavorite() const
 			TSharedPtr<FDetailCategoryImpl> ParentCategory = OwnerTreeNodePinned->GetParentCategory();
 			if (ParentCategory.IsValid())
 			{
-				FName BuilderName = Customization->CustomBuilderRow->GetCustomBuilderName();
-				if (BuilderName != NAME_None)
+				if (ParentCategory->IsFavoriteCategory())
 				{
-					TStringBuilder<256> PathToNode;
-					PathToNode.Append(ParentCategory->GetCategoryPathName());
-					PathToNode.Append(TEXT("."));
-					PathToNode.Append(BuilderName.ToString());
+					return true;
+				}
 
-					return OwnerTreeNodePinned->GetDetailsView()->IsCustomBuilderFavorite(PathToNode);
+				if (Customization->CustomBuilderRow->GetOriginalPath().IsEmpty() == false)
+				{
+					return OwnerTreeNodePinned->GetDetailsView()->IsCustomBuilderFavorite(Customization->CustomBuilderRow->GetOriginalPath());
 				}
 			}
 		}
@@ -914,13 +913,10 @@ void SDetailSingleItemRow::OnFavoriteMenuToggle()
 		TSharedPtr<FDetailCategoryImpl> ParentCategory = OwnerTreeNodePinned->GetParentCategory();
 		if (ParentCategory.IsValid())
 		{
-			TStringBuilder<256> PathToNode;
-			PathToNode.Append(ParentCategory->GetCategoryPathName());
-			PathToNode.Append(TEXT("."));
-			PathToNode.Append(Customization->CustomBuilderRow->GetCustomBuilderName().ToString());
+			const FString& OriginalPath = Customization->CustomBuilderRow->GetOriginalPath();
 
-			bNewValue = !DetailsView->IsCustomBuilderFavorite(PathToNode);
-			DetailsView->SetCustomBuilderFavorite(PathToNode, bNewValue);
+			bNewValue = !DetailsView->IsCustomBuilderFavorite(OriginalPath);
+			DetailsView->SetCustomBuilderFavorite(OriginalPath, bNewValue);
 		}
 	}
 
diff --git a/Engine/Source/Editor/PropertyEditor/Private/SDetailsView.cpp b/Engine/Source/Editor/PropertyEditor/Private/SDetailsView.cpp
--- a/Engine/Source/Editor/PropertyEditor/Private/SDetailsView.cpp
+++ b/Engine/Source/Editor/PropertyEditor/Private/SDetailsView.cpp
@@ -1076,6 +1076,11 @@ void SDetailsView::SetGroupFavorite(FStringView GroupPath, bool IsFavorite)
 
 bool SDetailsView::IsCustomBuilderFavorite(FStringView Path) const
 {
+	if (Path.IsEmpty())
+	{
+		return false;
+	}
+
 	UEditorMetadataOverrides* EditorMetadata = GEditor->GetEditorSubsystem<UEditorMetadataOverrides>();
 	if (!EditorMetadata)
 	{
@@ -1107,6 +1112,11 @@ bool SDetailsView::IsCustomBuilderFavorite(FStringView Path) const
 
 void SDetailsView::SetCustomBuilderFavorite(FStringView Path, bool IsFavorite)
 {
+	if (Path.IsEmpty())
+	{
+		return;
+	}
+
 	UEditorMetadataOverrides* MetadataOverrides = GEditor->GetEditorSubsystem<UEditorMetadataOverrides>();
 	if (!MetadataOverrides)
 	{
