Custom builders that have a property handle associated with them now can use that as a key to store their favorite status.

diff --git a/Engine/Source/Editor/PropertyEditor/Private/SDetailSingleItemRow.cpp b/Engine/Source/Editor/PropertyEditor/Private/SDetailSingleItemRow.cpp
--- a/Engine/Source/Editor/PropertyEditor/Private/SDetailSingleItemRow.cpp
+++ b/Engine/Source/Editor/PropertyEditor/Private/SDetailSingleItemRow.cpp
@@ -850,7 +850,10 @@ bool SDetailSingleItemRow::CanFavorite() const
 
 	if (Customization->HasCustomBuilder())
 	{
-		return Customization->CustomBuilderRow->GetOriginalPath().IsEmpty() == false;
+		TSharedPtr<IPropertyHandle> PropertyHandle = Customization->CustomBuilderRow->GetPropertyHandle();
+		const FString& OriginalPath = Customization->CustomBuilderRow->GetOriginalPath();
+		
+		return PropertyHandle.IsValid() || !OriginalPath.IsEmpty();
 	}
 
 	return false;
@@ -876,9 +879,16 @@ bool SDetailSingleItemRow::IsFavorite() const
 					return true;
 				}
 
-				if (Customization->CustomBuilderRow->GetOriginalPath().IsEmpty() == false)
+				TSharedPtr<IPropertyHandle> PropertyHandle = Customization->CustomBuilderRow->GetPropertyHandle();
+				if (PropertyHandle.IsValid())
+				{
+					return PropertyHandle->GetPropertyNode()->IsFavorite();
+				}
+
+				const FString& OriginalPath = Customization->CustomBuilderRow->GetOriginalPath();
+				if (!OriginalPath.IsEmpty())
 				{
-					return OwnerTreeNodePinned->GetDetailsView()->IsCustomBuilderFavorite(Customization->CustomBuilderRow->GetOriginalPath());
+					return OwnerTreeNodePinned->GetDetailsView()->IsCustomBuilderFavorite(OriginalPath);
 				}
 			}
 		}
@@ -910,11 +920,15 @@ void SDetailSingleItemRow::OnFavoriteMenuToggle()
 	}
 	else if (Customization->HasCustomBuilder())
 	{
-		TSharedPtr<FDetailCategoryImpl> ParentCategory = OwnerTreeNodePinned->GetParentCategory();
-		if (ParentCategory.IsValid())
-		{
-			const FString& OriginalPath = Customization->CustomBuilderRow->GetOriginalPath();
+		TSharedPtr<IPropertyHandle> PropertyHandle = Customization->CustomBuilderRow->GetPropertyHandle();
+		const FString& OriginalPath = Customization->CustomBuilderRow->GetOriginalPath();
 
+		if (PropertyHandle.IsValid())
+		{
+			PropertyHandle->GetPropertyNode()->SetFavorite(bNewValue); 
+		}
+		else if (!OriginalPath.IsEmpty())
+		{		
 			bNewValue = !DetailsView->IsCustomBuilderFavorite(OriginalPath);
 			DetailsView->SetCustomBuilderFavorite(OriginalPath, bNewValue);
 		}
