Improved behavior of detail item node when the parent category is either invalid or there is a derived class

diff --git a/Engine/Source/Editor/PropertyEditor/Private/DetailItemNode.h b/Engine/Source/Editor/PropertyEditor/Private/DetailItemNode.h
--- a/Engine/Source/Editor/PropertyEditor/Private/DetailItemNode.h
+++ b/Engine/Source/Editor/PropertyEditor/Private/DetailItemNode.h
@@ -58,8 +58,8 @@ public:
 	EVisibility ComputeItemVisibility() const;
 
 	/** FDetailTreeNode interface */
-	virtual IDetailsView* GetNodeDetailsView() const override { return ParentCategory.Pin()->GetNodeDetailsView(); }
-	virtual IDetailsViewPrivate* GetDetailsView() const override{ return ParentCategory.Pin()->GetDetailsView(); }
+	virtual IDetailsView* GetNodeDetailsView() const override { TSharedPtr<FDetailCategoryImpl> PC = GetParentCategory(); return PC.IsValid() ? PC->GetNodeDetailsView() : nullptr; }
+	virtual IDetailsViewPrivate* GetDetailsView() const override{ TSharedPtr<FDetailCategoryImpl> PC = GetParentCategory(); return PC.IsValid() ? PC->GetDetailsView() : nullptr; }
 	virtual TSharedRef< ITableRow > GenerateWidgetForTableView( const TSharedRef<STableViewBase>& OwnerTable, bool bAllowFavoriteSystem) override;
 	virtual bool GenerateStandaloneWidget(FDetailWidgetRow& OutRow) const override;
 	virtual void GetChildren( FDetailNodeList& OutChildren )  override;
