The SDetailsViewBase::IsPropertyVisible and IsCustomRowVisible delegates are now applied correctly. IsCustomRowVisible is now used for all things that are not properties.

This fixes the bugs of CategoriesToShow and PropertiesToShow not being respected in UMG details views.

diff --git a/Engine/Source/Editor/PropertyEditor/Private/DetailItemNode.cpp b/Engine/Source/Editor/PropertyEditor/Private/DetailItemNode.cpp
--- a/Engine/Source/Editor/PropertyEditor/Private/DetailItemNode.cpp
+++ b/Engine/Source/Editor/PropertyEditor/Private/DetailItemNode.cpp
@@ -97,15 +97,9 @@ TSharedPtr<IPropertyHandle> FDetailItemNode::CreatePropertyHandle() const
 		{
 			return Handles[0];
 		}
-		else
-		{
-			return nullptr;
-		}
-	}
-	else
-	{
-		return nullptr;
 	}
+
+	return nullptr;
 }
 
 void FDetailItemNode::GetFilterStrings(TArray<FString>& OutFilterStrings) const
@@ -631,19 +625,25 @@ void FDetailItemNode::Tick( float DeltaTime )
 EVisibility FDetailItemNode::ComputeItemVisibility() const
 {
 	EVisibility NewVisibility = EVisibility::Visible;
-	if (Customization.HasCustomWidget())
-	{	
-		NewVisibility = Customization.WidgetDecl->VisibilityAttr.Get();
+	if (Customization.HasPropertyNode())
+	{
+		NewVisibility = Customization.PropertyRow->GetPropertyVisibility();
 
-		IDetailsViewPrivate* DetailsView = GetDetailsView();
-		if (DetailsView && !DetailsView->IsCustomRowVisible(FName(*Customization.WidgetDecl->FilterTextString.ToString()), FName(*GetParentCategory()->GetDisplayName().ToString())))
+		if (NewVisibility != EVisibility::Collapsed)
 		{
-			NewVisibility = EVisibility::Collapsed;
+			TSharedPtr<FDetailCategoryImpl> ParentCategoryPtr = GetParentCategory();
+			if (ParentCategoryPtr.IsValid() && ParentCategoryPtr->IsParentLayoutValid())
+			{
+				if (!ParentCategoryPtr->GetParentLayout().IsPropertyVisible(CreatePropertyHandle().ToSharedRef()))
+				{
+					NewVisibility = EVisibility::Collapsed;
+				}
+			}
 		}
 	}
-	else if (Customization.HasPropertyNode())
-	{
-		NewVisibility = Customization.PropertyRow->GetPropertyVisibility();
+	else if (Customization.HasCustomWidget())
+	{	
+		NewVisibility = Customization.WidgetDecl->VisibilityAttr.Get();
 	}
 	else if (Customization.HasGroup())
 	{
@@ -663,6 +663,32 @@ EVisibility FDetailItemNode::ComputeItemVisibility() const
 		}
 	}
 
+	// check the details view's IsCustomRowVisible delegate if this isn't a property row
+	if (NewVisibility != EVisibility::Collapsed && 
+		GetDetailsView() != nullptr && 
+		!Customization.HasPropertyNode())
+	{
+		const FName CategoryName = GetParentCategory()->GetCategoryName();
+		FName RowName;
+		if (Customization.HasCustomWidget())
+		{
+			RowName = FName(*Customization.WidgetDecl->FilterTextString.ToString());
+		}
+		else if (Customization.HasCustomBuilder())
+		{
+			RowName = Customization.CustomBuilderRow->GetCustomBuilderName();
+		}
+		else if (Customization.HasGroup())
+		{
+			RowName = Customization.DetailGroup->GetGroupName();
+		}
+
+		if (!GetDetailsView()->IsCustomRowVisible(RowName, CategoryName))
+		{
+			NewVisibility = EVisibility::Collapsed;
+		}
+	}
+
 	return NewVisibility;
 }
 
