Reverting changes in CL 16122954 due to crash in Frosty.

diff --git a/Engine/Source/Editor/PropertyEditor/Private/DetailItemNode.cpp b/Engine/Source/Editor/PropertyEditor/Private/DetailItemNode.cpp
--- a/Engine/Source/Editor/PropertyEditor/Private/DetailItemNode.cpp
+++ b/Engine/Source/Editor/PropertyEditor/Private/DetailItemNode.cpp
@@ -97,9 +97,15 @@ TSharedPtr<IPropertyHandle> FDetailItemNode::CreatePropertyHandle() const
 		{
 			return Handles[0];
 		}
+		else
+		{
+			return nullptr;
+		}
+	}
+	else
+	{
+		return nullptr;
 	}
-
-	return nullptr;
 }
 
 void FDetailItemNode::GetFilterStrings(TArray<FString>& OutFilterStrings) const
@@ -625,25 +631,19 @@ void FDetailItemNode::Tick( float DeltaTime )
 EVisibility FDetailItemNode::ComputeItemVisibility() const
 {
 	EVisibility NewVisibility = EVisibility::Visible;
-	if (Customization.HasPropertyNode())
-	{
-		NewVisibility = Customization.PropertyRow->GetPropertyVisibility();
+	if (Customization.HasCustomWidget())
+	{	
+		NewVisibility = Customization.WidgetDecl->VisibilityAttr.Get();
 
-		if (NewVisibility != EVisibility::Collapsed)
+		IDetailsViewPrivate* DetailsView = GetDetailsView();
+		if (DetailsView && !DetailsView->IsCustomRowVisible(FName(*Customization.WidgetDecl->FilterTextString.ToString()), FName(*GetParentCategory()->GetDisplayName().ToString())))
 		{
-			TSharedPtr<FDetailCategoryImpl> ParentCategoryPtr = GetParentCategory();
-			if (ParentCategoryPtr.IsValid() && ParentCategoryPtr->IsParentLayoutValid())
-			{
-				if (!ParentCategoryPtr->GetParentLayout().IsPropertyVisible(CreatePropertyHandle().ToSharedRef()))
-				{
-					NewVisibility = EVisibility::Collapsed;
-				}
-			}
+			NewVisibility = EVisibility::Collapsed;
 		}
 	}
-	else if (Customization.HasCustomWidget())
-	{	
-		NewVisibility = Customization.WidgetDecl->VisibilityAttr.Get();
+	else if (Customization.HasPropertyNode())
+	{
+		NewVisibility = Customization.PropertyRow->GetPropertyVisibility();
 	}
 	else if (Customization.HasGroup())
 	{
@@ -663,32 +663,6 @@ EVisibility FDetailItemNode::ComputeItemVisibility() const
 		}
 	}
 
-	// check the details view's IsCustomRowVisible delegate if this isn't a property row
-	if (NewVisibility != EVisibility::Collapsed && 
-		GetDetailsView() != nullptr && 
-		!Customization.HasPropertyNode())
-	{
-		const FName CategoryName = GetParentCategory()->GetCategoryName();
-		FName RowName;
-		if (Customization.HasCustomWidget())
-		{
-			RowName = FName(*Customization.WidgetDecl->FilterTextString.ToString());
-		}
-		else if (Customization.HasCustomBuilder())
-		{
-			RowName = Customization.CustomBuilderRow->GetCustomBuilderName();
-		}
-		else if (Customization.HasGroup())
-		{
-			RowName = Customization.DetailGroup->GetGroupName();
-		}
-
-		if (!GetDetailsView()->IsCustomRowVisible(RowName, CategoryName))
-		{
-			NewVisibility = EVisibility::Collapsed;
-		}
-	}
-
 	return NewVisibility;
 }
 
