Remove outdated/incorrect comment about when EnsureDataIsValid should be called.

diff --git a/Engine/Source/Editor/PropertyEditor/Private/PropertyNode.cpp b/Engine/Source/Editor/PropertyEditor/Private/PropertyNode.cpp
--- a/Engine/Source/Editor/PropertyEditor/Private/PropertyNode.cpp
+++ b/Engine/Source/Editor/PropertyEditor/Private/PropertyNode.cpp
@@ -422,7 +422,7 @@ bool FPropertyNode::DoesChildPropertyRequireValidation(FProperty* InChildProp)
 }
 
 /** 
- * Used to see if any data has been destroyed from under the property tree.  Should only be called by PropertyWindow::OnIdle
+ * Used to see if any data has been destroyed from under the property tree.
  */
 EPropertyDataValidationResult FPropertyNode::EnsureDataIsValid()
 {
