Resubmitting IsPropertyVisible and IsCustomRowVisible (originally submitted in 16122954) with additional changes to ensure FPropertyAndParent cannot contain a null pointer inside an FProperty&.

diff --git a/Engine/Source/Editor/PropertyEditor/Private/DetailItemNode.cpp b/Engine/Source/Editor/PropertyEditor/Private/DetailItemNode.cpp
--- a/Engine/Source/Editor/PropertyEditor/Private/DetailItemNode.cpp
+++ b/Engine/Source/Editor/PropertyEditor/Private/DetailItemNode.cpp
@@ -97,15 +97,9 @@ TSharedPtr<IPropertyHandle> FDetailItemNode::CreatePropertyHandle() const
 		{
 			return Handles[0];
 		}
-		else
-		{
-			return nullptr;
-		}
-	}
-	else
-	{
-		return nullptr;
 	}
+
+	return nullptr;
 }
 
 void FDetailItemNode::GetFilterStrings(TArray<FString>& OutFilterStrings) const
@@ -631,19 +625,25 @@ void FDetailItemNode::Tick( float DeltaTime )
 EVisibility FDetailItemNode::ComputeItemVisibility() const
 {
 	EVisibility NewVisibility = EVisibility::Visible;
-	if (Customization.HasCustomWidget())
-	{	
-		NewVisibility = Customization.WidgetDecl->VisibilityAttr.Get();
+	if (Customization.HasPropertyNode())
+	{
+		NewVisibility = Customization.PropertyRow->GetPropertyVisibility();
 
-		IDetailsViewPrivate* DetailsView = GetDetailsView();
-		if (DetailsView && !DetailsView->IsCustomRowVisible(FName(*Customization.WidgetDecl->FilterTextString.ToString()), FName(*GetParentCategory()->GetDisplayName().ToString())))
+		if (NewVisibility != EVisibility::Collapsed)
 		{
-			NewVisibility = EVisibility::Collapsed;
+			TSharedPtr<FDetailCategoryImpl> ParentCategoryPtr = GetParentCategory();
+			if (ParentCategoryPtr.IsValid() && ParentCategoryPtr->IsParentLayoutValid())
+			{
+				if (!ParentCategoryPtr->GetParentLayout().IsPropertyVisible(CreatePropertyHandle().ToSharedRef()))
+				{
+					NewVisibility = EVisibility::Collapsed;
+				}
+			}
 		}
 	}
-	else if (Customization.HasPropertyNode())
-	{
-		NewVisibility = Customization.PropertyRow->GetPropertyVisibility();
+	else if (Customization.HasCustomWidget())
+	{	
+		NewVisibility = Customization.WidgetDecl->VisibilityAttr.Get();
 	}
 	else if (Customization.HasGroup())
 	{
@@ -663,6 +663,32 @@ EVisibility FDetailItemNode::ComputeItemVisibility() const
 		}
 	}
 
+	// check the details view's IsCustomRowVisible delegate if this isn't a property row
+	if (NewVisibility != EVisibility::Collapsed && 
+		GetDetailsView() != nullptr && 
+		!Customization.HasPropertyNode())
+	{
+		const FName CategoryName = GetParentCategory()->GetCategoryName();
+		FName RowName;
+		if (Customization.HasCustomWidget())
+		{
+			RowName = FName(*Customization.WidgetDecl->FilterTextString.ToString());
+		}
+		else if (Customization.HasCustomBuilder())
+		{
+			RowName = Customization.CustomBuilderRow->GetCustomBuilderName();
+		}
+		else if (Customization.HasGroup())
+		{
+			RowName = Customization.DetailGroup->GetGroupName();
+		}
+
+		if (!GetDetailsView()->IsCustomRowVisible(RowName, CategoryName))
+		{
+			NewVisibility = EVisibility::Collapsed;
+		}
+	}
+
 	return NewVisibility;
 }
 
diff --git a/Engine/Source/Editor/PropertyEditor/Private/DetailLayoutBuilderImpl.cpp b/Engine/Source/Editor/PropertyEditor/Private/DetailLayoutBuilderImpl.cpp
--- a/Engine/Source/Editor/PropertyEditor/Private/DetailLayoutBuilderImpl.cpp
+++ b/Engine/Source/Editor/PropertyEditor/Private/DetailLayoutBuilderImpl.cpp
@@ -727,7 +727,7 @@ TSharedPtr<FAssetThumbnailPool> FDetailLayoutBuilderImpl::GetThumbnailPool() con
 
 bool FDetailLayoutBuilderImpl::IsPropertyVisible( TSharedRef<IPropertyHandle> PropertyHandle ) const
 {
-	if( PropertyHandle->IsValidHandle() )
+	if( PropertyHandle->IsValidHandle() && PropertyHandle->GetProperty() != nullptr )
 	{
 		TArray<UObject*> OuterObjects;
 		PropertyHandle->GetOuterObjects(OuterObjects);
diff --git a/Engine/Source/Editor/PropertyEditor/Public/PropertyEditorDelegates.cpp b/Engine/Source/Editor/PropertyEditor/Public/PropertyEditorDelegates.cpp
--- a/Engine/Source/Editor/PropertyEditor/Public/PropertyEditorDelegates.cpp
+++ b/Engine/Source/Editor/PropertyEditor/Public/PropertyEditorDelegates.cpp
@@ -7,6 +7,8 @@ FPropertyAndParent::FPropertyAndParent(const TSharedRef<IPropertyHandle>& InProp
 	Property(*InPropertyHandle->GetProperty()),
 	Objects(InObjects)
 {
+	checkf(InPropertyHandle->GetProperty() != nullptr, "Creating an FPropertyAndParent with a null property!");
+
 	TSharedPtr<IPropertyHandle> ParentHandle = InPropertyHandle->GetParentHandle();
 	while (ParentHandle.IsValid())
 	{
