Added missing TEXT() macro.

diff --git a/Engine/Source/Editor/PropertyEditor/Public/PropertyEditorDelegates.cpp b/Engine/Source/Editor/PropertyEditor/Public/PropertyEditorDelegates.cpp
--- a/Engine/Source/Editor/PropertyEditor/Public/PropertyEditorDelegates.cpp
+++ b/Engine/Source/Editor/PropertyEditor/Public/PropertyEditorDelegates.cpp
@@ -7,7 +7,7 @@ FPropertyAndParent::FPropertyAndParent(const TSharedRef<IPropertyHandle>& InProp
 	Property(*InPropertyHandle->GetProperty()),
 	Objects(InObjects)
 {
-	checkf(InPropertyHandle->GetProperty() != nullptr, "Creating an FPropertyAndParent with a null property!");
+	checkf(InPropertyHandle->GetProperty() != nullptr, TEXT("Creating an FPropertyAndParent with a null property!"));
 
 	TSharedPtr<IPropertyHandle> ParentHandle = InPropertyHandle->GetParentHandle();
 	while (ParentHandle.IsValid())
