FDetailLayoutBuidlerImpl::IsPropertyVisible had the wrong default, which was causing properties in subcategories to be hidden under all circumstances.

diff --git a/Engine/Source/Editor/PropertyEditor/Private/DetailLayoutBuilderImpl.cpp b/Engine/Source/Editor/PropertyEditor/Private/DetailLayoutBuilderImpl.cpp
--- a/Engine/Source/Editor/PropertyEditor/Private/DetailLayoutBuilderImpl.cpp
+++ b/Engine/Source/Editor/PropertyEditor/Private/DetailLayoutBuilderImpl.cpp
@@ -743,7 +743,7 @@ bool FDetailLayoutBuilderImpl::IsPropertyVisible( TSharedRef<IPropertyHandle> Pr
 		return IsPropertyVisible(PropertyAndParent);
 	}
 	
-	return false;
+	return true;
 }
 
 bool FDetailLayoutBuilderImpl::IsPropertyVisible( const struct FPropertyAndParent& PropertyAndParent ) const
