Subcategory nodes are no longer always visible in editor utility widgets with categories filtered out.

diff --git a/Engine/Source/Editor/PropertyEditor/Private/DetailLayoutBuilderImpl.cpp b/Engine/Source/Editor/PropertyEditor/Private/DetailLayoutBuilderImpl.cpp
--- a/Engine/Source/Editor/PropertyEditor/Private/DetailLayoutBuilderImpl.cpp
+++ b/Engine/Source/Editor/PropertyEditor/Private/DetailLayoutBuilderImpl.cpp
@@ -1,15 +1,16 @@
 // Copyright Epic Games, Inc. All Rights Reserved.
 
 #include "DetailLayoutBuilderImpl.h"
-#include "ObjectPropertyNode.h"
+#include "CategoryPropertyNode.h"
 #include "DetailCategoryBuilderImpl.h"
-#include "PropertyHandleImpl.h"
-#include "PropertyEditorHelpers.h"
-#include "StructurePropertyNode.h"
 #include "DetailMultiTopLevelObjectRootNode.h"
-#include "ObjectEditorUtils.h"
 #include "DetailPropertyRow.h"
 #include "IPropertyGenerationUtilities.h"
+#include "ObjectEditorUtils.h"
+#include "ObjectPropertyNode.h"
+#include "PropertyEditorHelpers.h"
+#include "PropertyHandleImpl.h"
+#include "StructurePropertyNode.h"
 
 FDetailLayoutBuilderImpl::FDetailLayoutBuilderImpl(TSharedPtr<FComplexPropertyNode>& InRootNode, FClassToPropertyMap& InPropertyMap, const TSharedRef<IPropertyUtilities>& InPropertyUtilities, const TSharedRef<IPropertyGenerationUtilities>& InPropertyGenerationUtilities, const TSharedPtr< IDetailsViewPrivate >& InDetailsView, bool bIsExternal)
 	: RootNode( InRootNode )
@@ -727,22 +728,32 @@ TSharedPtr<FAssetThumbnailPool> FDetailLayoutBuilderImpl::GetThumbnailPool() con
 
 bool FDetailLayoutBuilderImpl::IsPropertyVisible( TSharedRef<IPropertyHandle> PropertyHandle ) const
 {
-	if( PropertyHandle->IsValidHandle() && PropertyHandle->GetProperty() != nullptr )
+	if (PropertyHandle->IsValidHandle() && DetailsView != nullptr)
 	{
-		TArray<UObject*> OuterObjects;
-		PropertyHandle->GetOuterObjects(OuterObjects);
-		
-		TArray<TWeakObjectPtr<UObject>> Objects;
-		for (auto OuterObject : OuterObjects)
+		TSharedPtr<FPropertyNode> PropertyNode = PropertyHandle->GetPropertyNode();
+		const FCategoryPropertyNode* CategoryNode = PropertyNode.IsValid() ? PropertyNode->AsCategoryNode() : nullptr;
+		if (CategoryNode != nullptr)
 		{
-			Objects.Add(OuterObject);
+			// this is a subcategory
+			FName CategoryName = CategoryNode->GetCategoryName();
+			return DetailsView->IsCustomRowVisible(FName(), CategoryName);
 		}
+		else if (PropertyHandle->GetProperty() != nullptr)
+		{
+			TArray<UObject*> OuterObjects;
+			PropertyHandle->GetOuterObjects(OuterObjects);
+		
+			TArray<TWeakObjectPtr<UObject>> Objects;
+			for (UObject* OuterObject : OuterObjects)
+			{
+				Objects.Add(OuterObject);
+			}
 
-		FPropertyAndParent PropertyAndParent(PropertyHandle, Objects);
-
-		return IsPropertyVisible(PropertyAndParent);
+			FPropertyAndParent PropertyAndParent(PropertyHandle, Objects);
+			return DetailsView->IsPropertyVisible(PropertyAndParent);
+		}
 	}
-	
+
 	return true;
 }
 
