Property rows with custom widgets now default to the standard IsEnabled check if a custom one is not set.

diff --git a/Engine/Source/Editor/PropertyEditor/Private/DetailPropertyRow.cpp b/Engine/Source/Editor/PropertyEditor/Private/DetailPropertyRow.cpp
--- a/Engine/Source/Editor/PropertyEditor/Private/DetailPropertyRow.cpp
+++ b/Engine/Source/Editor/PropertyEditor/Private/DetailPropertyRow.cpp
@@ -636,15 +636,20 @@ bool FDetailPropertyRow::GetEnabledState() const
 {
 	bool Result = IsParentEnabled.Get();
 
-	if (HasEditCondition())
+	if (PropertyEditor.IsValid())
 	{
-		if (CustomEditConditionValue.IsSet())
-		{
-			Result = Result && CustomEditConditionValue.Get();
-		}
+		Result &= !PropertyEditor->IsEditConst();
+	}
+
+	if (CustomEditConditionValue.IsSet())
+	{
+		Result &= CustomEditConditionValue.Get();
 	}
 	
-	Result = Result && CustomIsEnabledAttrib.Get();
+	if (CustomIsEnabledAttrib.IsSet())
+	{
+		Result &= CustomIsEnabledAttrib.Get();
+	}
 
 	return Result;
 }
@@ -742,9 +747,12 @@ void FDetailPropertyRow::MakeNameOrKeyWidget( FDetailWidgetRow& Row, const TShar
 		HorizontalAlignment = InCustomRow->NameWidget.HorizontalAlignment;
 	}
 
-	TAttribute<bool> IsEnabledAttrib = CustomIsEnabledAttrib;
-	
-	if (HasEditCondition())
+	TAttribute<bool> IsEnabledAttrib; 
+	if (CustomIsEnabledAttrib.IsBound() || CustomIsEnabledAttrib.IsSet())
+	{
+		IsEnabledAttrib = CustomIsEnabledAttrib;
+	}
+	else
 	{
 		IsEnabledAttrib.Bind( this, &FDetailPropertyRow::GetEnabledState );
 	}
@@ -835,8 +843,12 @@ void FDetailPropertyRow::MakeValueWidget( FDetailWidgetRow& Row, const TSharedPt
 		HorizontalAlignment = InCustomRow->ValueWidget.HorizontalAlignment;
 	}
 
-	TAttribute<bool> IsEnabledAttrib = CustomIsEnabledAttrib;
-	if( HasEditCondition() )
+	TAttribute<bool> IsEnabledAttrib; 
+	if (CustomIsEnabledAttrib.IsBound() || CustomIsEnabledAttrib.IsSet())
+	{
+		IsEnabledAttrib = CustomIsEnabledAttrib;
+	}
+	else
 	{
 		IsEnabledAttrib.Bind( this, &FDetailPropertyRow::GetEnabledState );
 	}
