Fixed warnings related to &= in DetailPropertyRow.cpp on old VS compiler versions.

diff --git a/Engine/Source/Editor/PropertyEditor/Private/DetailPropertyRow.cpp b/Engine/Source/Editor/PropertyEditor/Private/DetailPropertyRow.cpp
--- a/Engine/Source/Editor/PropertyEditor/Private/DetailPropertyRow.cpp
+++ b/Engine/Source/Editor/PropertyEditor/Private/DetailPropertyRow.cpp
@@ -638,17 +638,17 @@ bool FDetailPropertyRow::GetEnabledState() const
 
 	if (PropertyEditor.IsValid())
 	{
-		Result &= !PropertyEditor->IsEditConst();
+		Result = Result && !PropertyEditor->IsEditConst();
 	}
 
 	if (CustomEditConditionValue.IsSet())
 	{
-		Result &= CustomEditConditionValue.Get();
+		Result = Result && CustomEditConditionValue.Get();
 	}
 	
 	if (CustomIsEnabledAttrib.IsSet())
 	{
-		Result &= CustomIsEnabledAttrib.Get();
+		Result = Result && CustomIsEnabledAttrib.Get();
 	}
 
 	return Result;
