Moved items in a details view array now have their expansion state restored after being moved.

diff --git a/Engine/Source/Editor/PropertyEditor/Private/SDetailsViewBase.cpp b/Engine/Source/Editor/PropertyEditor/Private/SDetailsViewBase.cpp
--- a/Engine/Source/Editor/PropertyEditor/Private/SDetailsViewBase.cpp
+++ b/Engine/Source/Editor/PropertyEditor/Private/SDetailsViewBase.cpp
@@ -884,7 +884,7 @@ void SDetailsViewBase::Tick( const FGeometry& AllottedGeometry, const double InC
 	bool bUpdateFilteredDetails = false;
 	if (FoundIndex != INDEX_NONE)
 	{ 
-		// Reaquire the root property nodes.  It may have been changed by the deferred actions if something like a blueprint editor forcefully resets a details panel during a posteditchange
+		// Reacquire the root property nodes.  It may have been changed by the deferred actions if something like a blueprint editor forcefully resets a details panel during a posteditchange
 		ForceRefresh();
 
 		// All objects are being reset, no need to validate external nodes
@@ -906,11 +906,10 @@ void SDetailsViewBase::Tick( const FGeometry& AllottedGeometry, const double InC
 			}
 			else if(Result == EPropertyDataValidationResult::ObjectInvalid)
 			{
+				bValidateExternalNodes = false;
+
 				ForceRefresh();
 				break;
-
-				// All objects are being reset, no need to validate external nodes
-				bValidateExternalNodes = false;
 			}
 		}
 	}
@@ -948,6 +947,14 @@ void SDetailsViewBase::Tick( const FGeometry& AllottedGeometry, const double InC
 	if (bUpdateFilteredDetails)
 	{
 		UpdateFilteredDetails();
+
+		for (const TSharedPtr<FComplexPropertyNode>& RootPropertyNode : RootPropertyNodes)
+		{
+			if (RootPropertyNode.IsValid())
+			{
+				RestoreExpandedItems(RootPropertyNode.ToSharedRef());
+			}
+		}
 	}
 
 	for(FDetailLayoutData& LayoutData : DetailLayouts)
@@ -965,7 +972,6 @@ void SDetailsViewBase::Tick( const FGeometry& AllottedGeometry, const double InC
 		bHasOpenColorPicker = false;
 	}
 
-
 	if (FilteredNodesRequestingExpansionState.Num() > 0)
 	{
 		// change expansion state on the nodes that request it
