Added FDetailsViewConfig to store various details view settings, using the new EditorConfig system.

Renamed FDetailFilter properties to be consistent with one another.

diff --git a/Engine/Source/Editor/PropertyEditor/Private/DetailItemNode.cpp b/Engine/Source/Editor/PropertyEditor/Private/DetailItemNode.cpp
--- a/Engine/Source/Editor/PropertyEditor/Private/DetailItemNode.cpp
+++ b/Engine/Source/Editor/PropertyEditor/Private/DetailItemNode.cpp
@@ -562,10 +562,10 @@ static bool PassesAllFilters( FDetailItemNode* ItemNode, const FDetailLayoutCust
 	bool bPassesAllFilters = true;
 
 	if( InFilter.FilterStrings.Num() > 0 || 
-		InFilter.bShowOnlyModifiedProperties == true || 
+		InFilter.bShowOnlyModified == true || 
 		InFilter.bShowOnlyWhitelisted == true || 
-		InFilter.bShowKeyable == true || 
-		InFilter.bShowAnimated == true)
+		InFilter.bShowOnlyKeyable == true || 
+		InFilter.bShowOnlyAnimated == true)
 	{
 		const bool bSearchFilterIsEmpty = InFilter.FilterStrings.Num() == 0;
 
@@ -582,11 +582,11 @@ static bool PassesAllFilters( FDetailItemNode* ItemNode, const FDetailLayoutCust
 			const bool bIsParentSeenDueToFiltering = PropertyNodePin->HasNodeFlags(EPropertyNodeFlags::IsParentSeenDueToFiltering) != 0;
 
 			const bool bPassesSearchFilter = bPassesCategoryFilter || bPassesValueFilter || bSearchFilterIsEmpty || ( bIsNotBeingFiltered || bIsSeenDueToFiltering || bIsParentSeenDueToFiltering );
-			const bool bPassesModifiedFilter = bPassesSearchFilter && ( InFilter.bShowOnlyModifiedProperties == false || PropertyNodePin->GetDiffersFromDefault() == true );
+			const bool bPassesModifiedFilter = bPassesSearchFilter && ( InFilter.bShowOnlyModified == false || PropertyNodePin->GetDiffersFromDefault() == true );
 			const bool bPassesWhitelistedFilter = InFilter.bShowOnlyWhitelisted ? InFilter.WhitelistedProperties.Contains(*FPropertyNode::CreatePropertyPath(PropertyNodePin.ToSharedRef())) : true;
 
 			bool bPassesKeyableFilter = true;
-			if (InFilter.bShowKeyable)
+			if (InFilter.bShowOnlyKeyable)
 			{
 				FObjectPropertyNode* ParentPropertyNode = PropertyNodePin->FindObjectItemParent();
 				if (ParentPropertyNode != nullptr)
@@ -599,7 +599,7 @@ static bool PassesAllFilters( FDetailItemNode* ItemNode, const FDetailLayoutCust
 					bPassesKeyableFilter = false;
 				}
 			}
-			const bool bPassesAnimatedFilter = (InFilter.bShowAnimated == false || Local::ItemIsAnimated(ItemNode, PropertyNodePin));
+			const bool bPassesAnimatedFilter = (InFilter.bShowOnlyAnimated == false || Local::ItemIsAnimated(ItemNode, PropertyNodePin));
 
 			// The property node is visible (note categories are never visible unless they have a child that is visible )
 			bPassesAllFilters = bPassesSearchFilter && bPassesModifiedFilter && bPassesWhitelistedFilter && bPassesKeyableFilter && bPassesAnimatedFilter;
@@ -607,19 +607,19 @@ static bool PassesAllFilters( FDetailItemNode* ItemNode, const FDetailLayoutCust
 		else if (InCustomization.HasCustomWidget())
 		{
 			const bool bPassesTextFilter = bPassesCategoryFilter || bPassesValueFilter || Local::StringPassesFilter(InFilter, InCustomization.WidgetDecl->FilterTextString.ToString());
-			const bool bPassesModifiedFilter = InFilter.bShowOnlyModifiedProperties == false || InCustomization.WidgetDecl->EditConditionValue.Get(false);
+			const bool bPassesModifiedFilter = InFilter.bShowOnlyModified == false || InCustomization.WidgetDecl->EditConditionValue.Get(false);
 			//@todo we need to support custom widgets for keyable, animated, in particular for transforms(ComponentTransformDetails).
-			const bool bPassesKeyableFilter = (InFilter.bShowKeyable == false);
-			const bool bPassesAnimatedFilter = (InFilter.bShowAnimated == false);
+			const bool bPassesKeyableFilter = (InFilter.bShowOnlyKeyable == false);
+			const bool bPassesAnimatedFilter = (InFilter.bShowOnlyAnimated == false);
 			bPassesAllFilters = bPassesTextFilter && bPassesModifiedFilter && bPassesKeyableFilter && bPassesAnimatedFilter;
 		}
 		else if (InCustomization.HasCustomBuilder())
 		{
 			const bool bPassesTextFilter = bPassesCategoryFilter || bPassesValueFilter || Local::StringPassesFilter(InFilter, InCustomization.CustomBuilderRow->GetWidgetRow().FilterTextString.ToString());
 			//@todo we need to support custom builders for modified, keyable, animated, in particular for transforms(ComponentTransformDetails).
-			const bool bPassesModifiedFilter = (InFilter.bShowOnlyModifiedProperties == false);
-			const bool bPassesKeyableFilter = (InFilter.bShowKeyable == false);
-			const bool bPassesAnimatedFilter = (InFilter.bShowAnimated == false);
+			const bool bPassesModifiedFilter = (InFilter.bShowOnlyModified == false);
+			const bool bPassesKeyableFilter = (InFilter.bShowOnlyKeyable == false);
+			const bool bPassesAnimatedFilter = (InFilter.bShowOnlyAnimated == false);
 			bPassesAllFilters = bPassesTextFilter && bPassesModifiedFilter && bPassesKeyableFilter && bPassesAnimatedFilter;
 		}
 	}
@@ -744,7 +744,7 @@ TAttribute<bool> FDetailItemNode::IsPropertyEditingEnabled() const
 
 bool FDetailItemNode::IsPropertyEditingEnabledImpl() const
 {
-	bool IsParentEnabledValue = IsParentEnabled.Get(true);
+	const bool IsParentEnabledValue = IsParentEnabled.Get(true);
 	if (Customization.HasCustomWidget())
 	{
 		IDetailsViewPrivate* DetailsView = GetDetailsView();
diff --git a/Engine/Source/Editor/PropertyEditor/Private/DetailsViewConfig.h b/Engine/Source/Editor/PropertyEditor/Private/DetailsViewConfig.h
new file mode 100644
--- /dev/null
+++ b/Engine/Source/Editor/PropertyEditor/Private/DetailsViewConfig.h
@@ -0,0 +1,64 @@
+// Copyright Epic Games, Inc. All Rights Reserved.
+
+#pragma once
+
+#include "EditorConfigBase.h"
+
+#include "DetailsViewConfig.generated.h"
+
+USTRUCT()
+struct FDetailsViewConfig
+{
+	GENERATED_BODY()
+
+public:
+	
+	/** If we should show the favorites category. */
+	UPROPERTY()
+	bool bShowFavoritesCategory { false };
+
+	/** When enabled, the Advanced Details will always auto expand. */
+	UPROPERTY()
+	bool bShowAllAdvanced { false };
+
+	/** When Playing or Simulating, shows all properties (even non-visible and non-editable properties), if the object belongs to a simulating world.  This is useful for debugging. */
+	UPROPERTY()
+	bool bShowHiddenPropertiesWhilePlaying { false };
+
+	/** Show only whitelisted properties. */
+	UPROPERTY()
+	bool bShowOnlyWhitelisted { false };
+
+	/** Show all category children if the category matches the filter. */
+	UPROPERTY()
+	bool bShowAllChildrenIfCategoryMatches { false };
+
+	/** Show only keyable properties. */
+	UPROPERTY()
+	bool bShowOnlyKeyable { false };
+
+	/** Show only animated properties. */
+	UPROPERTY()
+	bool bShowOnlyAnimated { false };
+
+	/** Show only modified properties. */
+	UPROPERTY()
+	bool bShowOnlyModified { false };
+
+	UPROPERTY()
+	float ValueColumnWidth { 0 };
+
+	UPROPERTY()
+	float RightColumnWidth { 0 };
+};
+
+UCLASS(EditorConfig="DetailsView")
+class UDetailsConfig : public UEditorConfigBase
+{
+	GENERATED_BODY()
+	
+public:
+
+	UPROPERTY(meta=(EditorConfig))
+	TMap<FName, FDetailsViewConfig> Views;
+};
\ No newline at end of file
diff --git a/Engine/Source/Editor/PropertyEditor/Private/SDetailsView.cpp b/Engine/Source/Editor/PropertyEditor/Private/SDetailsView.cpp
--- a/Engine/Source/Editor/PropertyEditor/Private/SDetailsView.cpp
+++ b/Engine/Source/Editor/PropertyEditor/Private/SDetailsView.cpp
@@ -931,13 +931,23 @@ const FSlateBrush* SDetailsView::OnGetLockButtonImageResource() const
 
 bool SDetailsView::IsShowHiddenPropertiesWhilePlayingChecked() const
 {
-	return GetDefault<UEditorStyleSettings>()->bShowHiddenPropertiesWhilePlaying;
+	const FDetailsViewConfig* ViewConfig = GetConstViewConfig();
+	if (ViewConfig != nullptr)
+	{
+		return ViewConfig->bShowHiddenPropertiesWhilePlaying;
+	}
+
+	return false;
 }
 
 void SDetailsView::OnShowHiddenPropertiesWhilePlayingClicked()
 {
-	GetMutableDefault<UEditorStyleSettings>()->bShowHiddenPropertiesWhilePlaying = !GetDefault<UEditorStyleSettings>()->bShowHiddenPropertiesWhilePlaying;
-	GConfig->SetBool(TEXT("/Script/EditorStyle.EditorStyleSettings"), TEXT("bShowHiddenPropertiesWhilePlaying"), GetMutableDefault<UEditorStyleSettings>()->bShowHiddenPropertiesWhilePlaying, GEditorPerProjectIni);
+	const bool bNewValue = !IsShowHiddenPropertiesWhilePlayingChecked();
+
+	GetMutableViewConfig().bShowAllAdvanced = bNewValue;
+	SaveViewConfig();
+
+	GConfig->SetBool(TEXT("/Script/EditorStyle.EditorStyleSettings"), TEXT("bShowHiddenPropertiesWhilePlaying"), bNewValue, GEditorPerProjectIni);
 
 	// Force a refresh of the whole details panel, as the entire set of visible properties may be different
 	ForceRefresh();
@@ -958,6 +968,10 @@ FSlateColor SDetailsView::GetToggleFavoritesColor() const
 FReply SDetailsView::OnToggleFavoritesClicked()
 {
 	CurrentFilter.bShowFavoritesCategory = !CurrentFilter.bShowFavoritesCategory;
+
+	GetMutableViewConfig().bShowFavoritesCategory = CurrentFilter.bShowFavoritesCategory;
+	SaveViewConfig();
+
 	RerunCurrentFilter();
 
 	return FReply::Handled();
diff --git a/Engine/Source/Editor/PropertyEditor/Private/SDetailsViewBase.cpp b/Engine/Source/Editor/PropertyEditor/Private/SDetailsViewBase.cpp
--- a/Engine/Source/Editor/PropertyEditor/Private/SDetailsViewBase.cpp
+++ b/Engine/Source/Editor/PropertyEditor/Private/SDetailsViewBase.cpp
@@ -1,32 +1,59 @@
 // Copyright Epic Games, Inc. All Rights Reserved.
 
 #include "SDetailsViewBase.h"
-#include "GameFramework/Actor.h"
-#include "EngineGlobals.h"
+
+#include "Classes/EditorStyleSettings.h"
 #include "Engine/Engine.h"
-#include "Presentation/PropertyEditor/PropertyEditor.h"
-#include "ObjectPropertyNode.h"
+#include "GameFramework/Actor.h"
+#include "Misc/ConfigCacheIni.h"
 #include "Modules/ModuleManager.h"
+#include "Presentation/PropertyEditor/PropertyEditor.h"
 #include "Settings/EditorExperimentalSettings.h"
-#include "IDetailCustomization.h"
-#include "SDetailsView.h"
-#include "DetailLayoutBuilderImpl.h"
-#include "DetailCategoryBuilderImpl.h"
-#include "CategoryPropertyNode.h"
-#include "ScopedTransaction.h"
-#include "SDetailNameArea.h"
 #include "UserInterface/PropertyEditor/SPropertyEditorEditInline.h"
-#include "ObjectEditorUtils.h"
-#include "Misc/ConfigCacheIni.h"
 #include "Widgets/Colors/SColorPicker.h"
-#include "DetailPropertyRow.h"
 #include "Widgets/Input/SSearchBox.h"
-#include "Classes/EditorStyleSettings.h"
+
+#include "CategoryPropertyNode.h"
+#include "DetailCategoryBuilderImpl.h"
+#include "DetailLayoutBuilderImpl.h"
 #include "DetailLayoutHelpers.h"
+#include "DetailPropertyRow.h"
 #include "EditConditionParser.h"
 #include "Editor.h"
+#include "EditorConfigSubsystem.h"
+#include "IDetailCustomization.h"
+#include "ObjectEditorUtils.h"
+#include "ObjectPropertyNode.h"
+#include "ScopedTransaction.h"
+#include "SDetailNameArea.h"
+#include "SDetailsView.h"
+
 #include "ThumbnailRendering/ThumbnailManager.h"
 
+SDetailsViewBase::SDetailsViewBase() :
+	bHasActiveFilter(false)
+	, bIsLocked(false)
+	, bHasOpenColorPicker(false)
+	, bDisableCustomDetailLayouts( false )
+	, NumVisibleTopLevelObjectNodes(0)
+	, bPendingCleanupTimerSet(false)
+{
+	UDetailsConfig* DetailsConfig = GetMutableDefault<UDetailsConfig>();
+	DetailsConfig->LoadEditorConfig();
+
+	FDetailsViewConfig* ViewConfig = DetailsConfig->Views.Find(DetailsViewArgs.ViewIdentifier);
+	if (ViewConfig != nullptr)
+	{
+		CurrentFilter.bShowAllAdvanced = ViewConfig->bShowAllAdvanced;
+		CurrentFilter.bShowAllChildrenIfCategoryMatches = ViewConfig->bShowAllChildrenIfCategoryMatches;
+		CurrentFilter.bShowOnlyAnimated = ViewConfig->bShowOnlyAnimated;
+		CurrentFilter.bShowOnlyKeyable = ViewConfig->bShowOnlyKeyable;
+		CurrentFilter.bShowOnlyModified = ViewConfig->bShowOnlyModified;
+		CurrentFilter.bShowOnlyWhitelisted = ViewConfig->bShowOnlyWhitelisted;
+
+	}
+}
+
 SDetailsViewBase::~SDetailsViewBase()
 {
 }
@@ -581,9 +608,28 @@ TSharedPtr<IPropertyUtilities> SDetailsViewBase::GetPropertyUtilities()
 	return PropertyUtilities;
 }
 
+const FDetailsViewConfig* SDetailsViewBase::GetConstViewConfig() const 
+{
+	return GetDefault<UDetailsConfig>()->Views.Find(DetailsViewArgs.ViewIdentifier);
+}
+
+FDetailsViewConfig& SDetailsViewBase::GetMutableViewConfig()
+{
+	UDetailsConfig* DetailsConfig = GetMutableDefault<UDetailsConfig>();
+	return DetailsConfig->Views.FindOrAdd(DetailsViewArgs.ViewIdentifier);
+}
+
+void SDetailsViewBase::SaveViewConfig()
+{
+	GetMutableDefault<UDetailsConfig>()->SaveEditorConfig();
+}
+
 void SDetailsViewBase::OnShowOnlyModifiedClicked()
 {
-	CurrentFilter.bShowOnlyModifiedProperties = !CurrentFilter.bShowOnlyModifiedProperties;
+	CurrentFilter.bShowOnlyModified = !CurrentFilter.bShowOnlyModified;
+
+	GetMutableViewConfig().bShowOnlyModified = CurrentFilter.bShowOnlyModified;
+	SaveViewConfig();
 
 	UpdateFilteredDetails();
 }
@@ -600,8 +646,9 @@ void SDetailsViewBase::OnCustomFilterClicked()
 void SDetailsViewBase::OnShowAllAdvancedClicked()
 {
 	CurrentFilter.bShowAllAdvanced = !CurrentFilter.bShowAllAdvanced;
-	GetMutableDefault<UEditorStyleSettings>()->bShowAllAdvancedDetails = CurrentFilter.bShowAllAdvanced;
-	GConfig->SetBool(TEXT("/Script/EditorStyle.EditorStyleSettings"), TEXT("bShowAllAdvancedDetails"), GetMutableDefault<UEditorStyleSettings>()->bShowAllAdvancedDetails, GEditorPerProjectIni);
+
+	GetMutableViewConfig().bShowAllAdvanced = CurrentFilter.bShowAllAdvanced;
+	SaveViewConfig();
 
 	UpdateFilteredDetails();
 }
@@ -610,6 +657,9 @@ void SDetailsViewBase::OnShowOnlyWhitelistedClicked()
 {
 	CurrentFilter.bShowOnlyWhitelisted = !CurrentFilter.bShowOnlyWhitelisted;
 
+	GetMutableViewConfig().bShowOnlyWhitelisted = CurrentFilter.bShowOnlyWhitelisted;
+	SaveViewConfig();
+
 	UpdateFilteredDetails();
 }
 
@@ -617,22 +667,32 @@ void SDetailsViewBase::OnShowAllChildrenIfCategoryMatchesClicked()
 {
 	CurrentFilter.bShowAllChildrenIfCategoryMatches = !CurrentFilter.bShowAllChildrenIfCategoryMatches;
 
+	GetMutableViewConfig().bShowAllChildrenIfCategoryMatches = CurrentFilter.bShowAllChildrenIfCategoryMatches;
+	SaveViewConfig();
+
 	UpdateFilteredDetails();
 }
 
 void SDetailsViewBase::OnShowKeyableClicked()
 {
-	CurrentFilter.bShowKeyable = !CurrentFilter.bShowKeyable;
+	CurrentFilter.bShowOnlyKeyable = !CurrentFilter.bShowOnlyKeyable;
+
+	GetMutableViewConfig().bShowOnlyKeyable = CurrentFilter.bShowOnlyKeyable;
+	SaveViewConfig();
 
 	UpdateFilteredDetails();
 }
 
 void SDetailsViewBase::OnShowAnimatedClicked()
 {
-	CurrentFilter.bShowAnimated = !CurrentFilter.bShowAnimated;
+	CurrentFilter.bShowOnlyAnimated = !CurrentFilter.bShowOnlyAnimated;
+
+	GetMutableViewConfig().bShowOnlyAnimated = CurrentFilter.bShowOnlyAnimated;
+	SaveViewConfig();
 
 	UpdateFilteredDetails();
 }
+
 /** Called when the filter text changes.  This filters specific property nodes out of view */
 void SDetailsViewBase::OnFilterTextChanged(const FText& InFilterText)
 {
@@ -1206,22 +1266,27 @@ void SDetailsViewBase::UpdateFilteredDetails()
 	FDetailNodeList InitialRootNodeList;
 	
 	NumVisibleTopLevelObjectNodes = 0;
-	FRootPropertyNodeList& RootPropertyNodes = GetRootNodes();
 
-	if( GetDefault<UEditorStyleSettings>()->bShowAllAdvancedDetails )
+	CurrentFilter.bShowAllAdvanced = false;
+
+	const FDetailsViewConfig* DetailsViewConfig = GetMutableDefault<UDetailsConfig>()->Views.Find(DetailsViewArgs.ViewIdentifier);
+	if (DetailsViewConfig != nullptr)
 	{
-		CurrentFilter.bShowAllAdvanced = true;
+		CurrentFilter.bShowAllAdvanced = DetailsViewConfig->bShowAllAdvanced;
 	}
 	
+	FRootPropertyNodeList& RootPropertyNodes = GetRootNodes();
 	for(int32 RootNodeIndex = 0; RootNodeIndex < RootPropertyNodes.Num(); ++RootNodeIndex)
 	{
-		TSharedPtr<FComplexPropertyNode>& RootPropertyNode = RootPropertyNodes[RootNodeIndex];
+		const TSharedPtr<FComplexPropertyNode>& RootPropertyNode = RootPropertyNodes[RootNodeIndex];
 		if(RootPropertyNode.IsValid())
 		{
 			RootPropertyNode->FilterNodes(CurrentFilter.FilterStrings);
 			RootPropertyNode->ProcessSeenFlags(true);
 
-			TSharedPtr<FDetailLayoutBuilderImpl>& DetailLayout = DetailLayouts[RootNodeIndex].DetailLayout;
+			RestoreExpandedItems(RootPropertyNode.ToSharedRef());
+
+			const TSharedPtr<FDetailLayoutBuilderImpl>& DetailLayout = DetailLayouts[RootNodeIndex].DetailLayout;
 			if(DetailLayout.IsValid())
 			{
 				FRootPropertyNodeList& ExternalRootPropertyNodes = DetailLayout->GetExternalRootPropertyNodes();
diff --git a/Engine/Source/Editor/PropertyEditor/Private/SDetailsViewBase.h b/Engine/Source/Editor/PropertyEditor/Private/SDetailsViewBase.h
--- a/Engine/Source/Editor/PropertyEditor/Private/SDetailsViewBase.h
+++ b/Engine/Source/Editor/PropertyEditor/Private/SDetailsViewBase.h
@@ -22,6 +22,7 @@
 #include "IDetailsViewPrivate.h"
 #include "PropertyRowGenerator.h"
 #include "PropertyCustomizationHelpers.h"
+#include "DetailsViewConfig.h"
 
 class FDetailCategoryImpl;
 class FDetailLayoutBuilderImpl;
@@ -39,12 +40,12 @@ typedef STreeView< TSharedRef<class FDetailTreeNode> > SDetailTree;
 /** Represents a filter which controls the visibility of items in the details view */
 struct FDetailFilter
 {
-	FDetailFilter()
-		: bShowOnlyModifiedProperties(false)
+	FDetailFilter() :
+		bShowOnlyModified(false)
 		, bShowAllAdvanced(false)
 		, bShowAllChildrenIfCategoryMatches(true)
-		, bShowKeyable(false)
-		, bShowAnimated(false)
+		, bShowOnlyKeyable(false)
+		, bShowOnlyAnimated(false)
 		, bShowFavoritesCategory(false)
 		, bShowOnlyWhitelisted(false)
 	{}
@@ -52,26 +53,25 @@ struct FDetailFilter
 	bool IsEmptyFilter() const 
 	{ 
 		return FilterStrings.Num() == 0 
-			&& bShowOnlyModifiedProperties == false 
 			&& bShowAllAdvanced == false 
 			&& bShowOnlyWhitelisted == false 
 			&& bShowAllChildrenIfCategoryMatches == false 
-			&& bShowKeyable == false 
-			&& bShowAnimated == false;
+			&& bShowOnlyKeyable == false 
+			&& bShowOnlyAnimated == false;
 	}
 
 	/** Any user search terms that items must match */
 	TArray<FString> FilterStrings;
 	/** If we should only show modified properties */
-	bool bShowOnlyModifiedProperties;
+	bool bShowOnlyModified;
 	/** If we should show all advanced properties */
 	bool bShowAllAdvanced;
 	/** If we should show all the children if their category name matches the search */
 	bool bShowAllChildrenIfCategoryMatches;
 	/** If we should only show keyable properties */
-	bool bShowKeyable;
+	bool bShowOnlyKeyable;
 	/** If we should only show animated properties */
-	bool bShowAnimated;
+	bool bShowOnlyAnimated;
 	/** If we should show the favorites category. */
 	bool bShowFavoritesCategory;
 	/** If we should only show whitelisted properties */
@@ -83,16 +83,7 @@ struct FDetailFilter
 class SDetailsViewBase : public IDetailsViewPrivate
 {
 public:
-	SDetailsViewBase() : 
-		bHasActiveFilter(false)
-		, bIsLocked(false)
-		, bHasOpenColorPicker(false)
-		, bDisableCustomDetailLayouts( false )
-		, NumVisibleTopLevelObjectNodes(0)
-		, bPendingCleanupTimerSet(false)
-	{
-	}
-
+	SDetailsViewBase();
 	virtual ~SDetailsViewBase();
 
 	/**
@@ -268,7 +259,7 @@ protected:
 	TSharedRef<ITableRow> OnGenerateRowForDetailTree( TSharedRef<FDetailTreeNode> InTreeNode, const TSharedRef<STableViewBase>& OwnerTable );
 
 	/** @return true if show only modified is checked */
-	bool IsShowOnlyModifiedChecked() const { return CurrentFilter.bShowOnlyModifiedProperties; }
+	bool IsShowOnlyModifiedChecked() const { return CurrentFilter.bShowOnlyModified; }
 
 	/** @return true if custom filter is checked */
 	bool IsCustomFilterChecked() const { return bCustomFilterActive; }
@@ -283,10 +274,10 @@ protected:
 	bool IsShowAllChildrenIfCategoryMatchesChecked() const { return CurrentFilter.bShowAllChildrenIfCategoryMatches; }
 	
 	/** @return true if show keyable is checked */
-	bool IsShowKeyableChecked() const { return CurrentFilter.bShowKeyable; }
+	bool IsShowKeyableChecked() const { return CurrentFilter.bShowOnlyKeyable; }
 	
 	/** @return true if show animated is checked */
-	bool IsShowAnimatedChecked() const { return CurrentFilter.bShowAnimated; }
+	bool IsShowAnimatedChecked() const { return CurrentFilter.bShowOnlyAnimated; }
 
 	/** Called when show only modified is clicked */
 	void OnShowOnlyModifiedClicked();
@@ -366,6 +357,10 @@ protected:
 	void SavePreSearchExpandedItems();
 	void RestorePreSearchExpandedItems();
 
+	struct FDetailsViewConfig& GetMutableViewConfig();
+	const FDetailsViewConfig* GetConstViewConfig() const;
+	void SaveViewConfig();
+
 protected:
 	/** The user defined args for the details view */
 	FDetailsViewArgs DetailsViewArgs;
@@ -458,4 +453,5 @@ protected:
 	FText CustomFilterLabel;
 
 	mutable TSharedPtr<FEditConditionParser> EditConditionParser;
+
 };
