Removed check of IsEditConst for DetailPropertyRows, since this was causing elements in VisibleAnywhere containers to not be editable.

diff --git a/Engine/Source/Editor/PropertyEditor/Private/DetailPropertyRow.cpp b/Engine/Source/Editor/PropertyEditor/Private/DetailPropertyRow.cpp
--- a/Engine/Source/Editor/PropertyEditor/Private/DetailPropertyRow.cpp
+++ b/Engine/Source/Editor/PropertyEditor/Private/DetailPropertyRow.cpp
@@ -636,11 +636,6 @@ bool FDetailPropertyRow::GetEnabledState() const
 {
 	bool Result = IsParentEnabled.Get();
 
-	if (PropertyEditor.IsValid())
-	{
-		Result = Result && !PropertyEditor->IsEditConst();
-	}
-
 	if (CustomEditConditionValue.IsSet())
 	{
 		Result = Result && CustomEditConditionValue.Get();
