Removed bShowOnlyWhitelistedProperties from details config. If the value was serialized, it would likely be nonsensical without the list of whitelisted properties, which are currently always set programmatically (by the diff tool, for example) and not by the user.

diff --git a/Engine/Source/Editor/PropertyEditor/Private/DetailsViewConfig.h b/Engine/Source/Editor/PropertyEditor/Private/DetailsViewConfig.h
--- a/Engine/Source/Editor/PropertyEditor/Private/DetailsViewConfig.h
+++ b/Engine/Source/Editor/PropertyEditor/Private/DetailsViewConfig.h
@@ -25,10 +25,6 @@ public:
 	UPROPERTY()
 	bool bShowHiddenPropertiesWhilePlaying { false };
 
-	/** Show only whitelisted properties. */
-	UPROPERTY()
-	bool bShowOnlyWhitelisted { false };
-
 	/** Show all category children if the category matches the filter. */
 	UPROPERTY()
 	bool bShowAllChildrenIfCategoryMatches { false };
diff --git a/Engine/Source/Editor/PropertyEditor/Private/SDetailsViewBase.cpp b/Engine/Source/Editor/PropertyEditor/Private/SDetailsViewBase.cpp
--- a/Engine/Source/Editor/PropertyEditor/Private/SDetailsViewBase.cpp
+++ b/Engine/Source/Editor/PropertyEditor/Private/SDetailsViewBase.cpp
@@ -41,7 +41,7 @@ SDetailsViewBase::SDetailsViewBase() :
 	UDetailsConfig* DetailsConfig = GetMutableDefault<UDetailsConfig>();
 	DetailsConfig->LoadEditorConfig();
 
-	FDetailsViewConfig* ViewConfig = DetailsConfig->Views.Find(DetailsViewArgs.ViewIdentifier);
+	const FDetailsViewConfig* ViewConfig = DetailsConfig->Views.Find(DetailsViewArgs.ViewIdentifier);
 	if (ViewConfig != nullptr)
 	{
 		CurrentFilter.bShowAllAdvanced = ViewConfig->bShowAllAdvanced;
@@ -49,8 +49,6 @@ SDetailsViewBase::SDetailsViewBase() :
 		CurrentFilter.bShowOnlyAnimated = ViewConfig->bShowOnlyAnimated;
 		CurrentFilter.bShowOnlyKeyable = ViewConfig->bShowOnlyKeyable;
 		CurrentFilter.bShowOnlyModified = ViewConfig->bShowOnlyModified;
-		CurrentFilter.bShowOnlyWhitelisted = ViewConfig->bShowOnlyWhitelisted;
-
 	}
 }
 
@@ -657,9 +655,6 @@ void SDetailsViewBase::OnShowOnlyWhitelistedClicked()
 {
 	CurrentFilter.bShowOnlyWhitelisted = !CurrentFilter.bShowOnlyWhitelisted;
 
-	GetMutableViewConfig().bShowOnlyWhitelisted = CurrentFilter.bShowOnlyWhitelisted;
-	SaveViewConfig();
-
 	UpdateFilteredDetails();
 }
 
