FPropertyNode::HasNodeFlags now returns a bool instead of an int32.

diff --git a/Engine/Source/Editor/PropertyEditor/Private/PropertyNode.cpp b/Engine/Source/Editor/PropertyEditor/Private/PropertyNode.cpp
--- a/Engine/Source/Editor/PropertyEditor/Private/PropertyNode.cpp
+++ b/Engine/Source/Editor/PropertyEditor/Private/PropertyNode.cpp
@@ -130,7 +130,7 @@ void FPropertyNode::InitNode(const FPropertyNodeInitParams& InitParams)
 	//default to copying from the parent
 	if (ParentNode)
 	{
-		if (ParentNode->HasNodeFlags(EPropertyNodeFlags::ShowCategories) != 0)
+		if (ParentNode->HasNodeFlags(EPropertyNodeFlags::ShowCategories))
 		{
 			SetNodeFlags(EPropertyNodeFlags::ShowCategories, true);
 		}
@@ -232,7 +232,7 @@ void FPropertyNode::InitNode(const FPropertyNodeInitParams& InitParams)
 		bRequiresValidation |= MyProperty->IsA<FArrayProperty>() || MyProperty->IsA<FSetProperty>() || MyProperty->IsA<FMapProperty>();
 
 		// We require validation if our parent also needs validation (if an array parent was resized all the addresses of children are invalid)
-		bRequiresValidation |= (GetParentNode() && GetParentNode()->HasNodeFlags( EPropertyNodeFlags::RequiresValidation ) != 0);
+		bRequiresValidation |= (GetParentNode() && GetParentNode()->HasNodeFlags(EPropertyNodeFlags::RequiresValidation));
 
 		SetNodeFlags( EPropertyNodeFlags::RequiresValidation, bRequiresValidation );
 	}
@@ -437,7 +437,7 @@ EPropertyDataValidationResult FPropertyNode::EnsureDataIsValid()
 	}
 
 	// The root must always be validated
-	if( GetParentNode() == NULL || HasNodeFlags(EPropertyNodeFlags::RequiresValidation) != 0 )
+	if (GetParentNode() == nullptr || HasNodeFlags(EPropertyNodeFlags::RequiresValidation))
 	{
 		CachedReadAddresses.Reset();
 
@@ -638,7 +638,7 @@ EPropertyDataValidationResult FPropertyNode::EnsureDataIsValid()
 				return EPropertyDataValidationResult::EditInlineNewValueChanged;
 			}
 
-			const bool bHasChildren = (GetNumChildNodes() != 0);
+			const bool bHasChildren = (GetNumChildNodes() > 0);
 			// If the object property is not null and has no children, its children need to be rebuilt
 			// If the object property is null and this node has children, the node needs to be rebuilt
 			if (!HasNodeFlags(EPropertyNodeFlags::ShowInnerObjectProperties) && ObjectProperty && ((!bObjectPropertyNull && !bHasChildren) || (bObjectPropertyNull && bHasChildren)))
@@ -877,7 +877,6 @@ TSharedPtr<FPropertyNode> FPropertyNode::FindChildPropertyNode( const FName InPr
 		}
 	}
 
-	// Return nullptr if not found...
 	return nullptr;
 }
 
@@ -886,13 +885,17 @@ TSharedPtr<FPropertyNode> FPropertyNode::FindChildPropertyNode( const FName InPr
  */
 bool FPropertyNode::IsPropertyConst() const
 {
-	bool bIsPropertyConst = (HasNodeFlags(EPropertyNodeFlags::IsReadOnly) != 0);
-	if (!bIsPropertyConst && Property != nullptr)
+	if (HasNodeFlags(EPropertyNodeFlags::IsReadOnly))
 	{
-		bIsPropertyConst = (Property->PropertyFlags & CPF_EditConst) ? true : false;	
+		return true;
 	}
 
-	return bIsPropertyConst;
+	if (Property != nullptr)
+	{
+		return Property->HasAllPropertyFlags(CPF_EditConst);
+	}
+
+	return false;
 }
 
 /** @return whether this window's property is constant (can't be edited by the user) */
@@ -1022,7 +1025,7 @@ void FPropertyNode::ToggleEditConditionState()
 	check(MyParentNode != nullptr);
 
 	bool OldValue = true; 
-	bool IsSparseClassData = HasNodeFlags(EPropertyNodeFlags::IsSparseClassData) != 0;
+	bool IsSparseClassData = HasNodeFlags(EPropertyNodeFlags::IsSparseClassData);
 
 	FComplexPropertyNode* ComplexParentNode = FindComplexParent();
 	for (int32 Index = 0; Index < ComplexParentNode->GetInstancesNum(); ++Index)
@@ -1231,12 +1234,12 @@ uint8* FPropertyNode::GetStartAddress(const UObject* Obj) const
 
 uint8* FPropertyNode::GetValueBaseAddressFromObject(const UObject* Obj) const
 {
-	return GetValueBaseAddress(GetStartAddress(Obj), HasNodeFlags(EPropertyNodeFlags::IsSparseClassData) != 0);
+	return GetValueBaseAddress(GetStartAddress(Obj), HasNodeFlags(EPropertyNodeFlags::IsSparseClassData));
 }
 
 uint8* FPropertyNode::GetValueAddressFromObject(const UObject* Obj) const
 {
-	return GetValueAddress(GetStartAddress(Obj), HasNodeFlags(EPropertyNodeFlags::IsSparseClassData) != 0);
+	return GetValueAddress(GetStartAddress(Obj), HasNodeFlags(EPropertyNodeFlags::IsSparseClassData));
 }
 
 
@@ -1367,7 +1370,7 @@ public:
 	 */
 	bool IsValidTracker() const
 	{
-		return PropertyValueBaseAddress != 0 && OwnerObject.IsValid();
+		return PropertyValueBaseAddress != nullptr && OwnerObject.IsValid();
 	}
 
 	/**
@@ -2679,7 +2682,7 @@ void FPropertyNode::GetExpandedChildPropertyPaths(TSet<FString>& OutExpandedChil
 	do
 	{
 		const FPropertyNode* SearchNode = RecursiveStack.Pop();
-		if (SearchNode->HasNodeFlags(EPropertyNodeFlags::Expanded) != 0)
+		if (SearchNode->HasNodeFlags(EPropertyNodeFlags::Expanded))
 		{
 			OutExpandedChildPropertyPaths.Add(SearchNode->PropertyPath);
 
diff --git a/Engine/Source/Editor/PropertyEditor/Private/PropertyNode.h b/Engine/Source/Editor/PropertyEditor/Private/PropertyNode.h
--- a/Engine/Source/Editor/PropertyEditor/Private/PropertyNode.h
+++ b/Engine/Source/Editor/PropertyEditor/Private/PropertyNode.h
@@ -371,7 +371,7 @@ public:
 
 	//////////////////////////////////////////////////////////////////////////
 	//Flags
-	uint32 HasNodeFlags(const EPropertyNodeFlags::Type InTestFlags) const { return PropertyNodeFlags & InTestFlags; }
+	bool HasNodeFlags(const EPropertyNodeFlags::Type InTestFlags) const { return (PropertyNodeFlags & InTestFlags) != 0; }
 	/**
 	 * Sets the flags used by the window and the root node
 	 * @param InFlags - flags to turn on or off
