Fix comment on IPropertyHandle::GetParent

diff --git a/Engine/Source/Editor/PropertyEditor/Public/PropertyHandle.h b/Engine/Source/Editor/PropertyEditor/Public/PropertyHandle.h
--- a/Engine/Source/Editor/PropertyEditor/Public/PropertyHandle.h
+++ b/Engine/Source/Editor/PropertyEditor/Public/PropertyHandle.h
@@ -396,7 +396,8 @@ public:
 	virtual TSharedPtr<IPropertyHandle> GetChildHandle( uint32 Index ) const = 0;
 	
 	/**
-	 * @return a handle to the parent array if this handle is an array element
+	 * @return a handle to the parent property 
+	 * This parent handle may not contain a valid FProperty if the parent is the uobject.
 	 */
 	virtual TSharedPtr<IPropertyHandle> GetParentHandle() const = 0;
 
