Add ability to copy property display name, in order to copy variable names from material instances

diff --git a/Engine/Source/Editor/PropertyEditor/Private/SDetailSingleItemRow.cpp b/Engine/Source/Editor/PropertyEditor/Private/SDetailSingleItemRow.cpp
--- a/Engine/Source/Editor/PropertyEditor/Private/SDetailSingleItemRow.cpp
+++ b/Engine/Source/Editor/PropertyEditor/Private/SDetailSingleItemRow.cpp
@@ -652,6 +652,7 @@ bool SDetailSingleItemRow::OnContextMenuOpening(FMenuBuilder& MenuBuilder)
 
 	FUIAction CopyAction;
 	FUIAction PasteAction;
+	FUIAction CopyDisplayNameAction = FExecuteAction::CreateSP(this, &SDetailSingleItemRow::OnCopyPropertyDisplayName);
 
 	if (bIsCopyPasteBound)
 	{
@@ -697,6 +698,12 @@ bool SDetailSingleItemRow::OnContextMenuOpening(FMenuBuilder& MenuBuilder)
 			NSLOCTEXT("PropertyView", "PasteProperty_ToolTip", "Paste the copied value here"),
 			FSlateIcon(FCoreStyle::Get().GetStyleSetName(), "GenericCommands.Paste"),
 			PasteAction);
+
+		MenuBuilder.AddMenuEntry(
+			NSLOCTEXT("PropertyView", "CopyPropertyDisplayName", "Copy Display Name"),
+			NSLOCTEXT("PropertyView", "CopyPropertyDisplayName_ToolTip", "Copy this property display name"),
+			FSlateIcon(FCoreStyle::Get().GetStyleSetName(), "GenericCommands.Copy"),
+			CopyDisplayNameAction);
 	}
 
 	if (OwnerTreeNode.Pin()->GetDetailsView()->IsFavoritingEnabled())
@@ -799,6 +806,18 @@ void SDetailSingleItemRow::OnCopyProperty()
 	}
 }
 
+void SDetailSingleItemRow::OnCopyPropertyDisplayName()
+{
+	if (OwnerTreeNode.IsValid())
+	{
+		TSharedPtr<FPropertyNode> PropertyNode = GetPropertyNode();
+		if (PropertyNode.IsValid())
+		{
+			FPlatformApplicationMisc::ClipboardCopy(*PropertyNode->GetDisplayName().ToString());
+		}
+	}
+}
+
 void SDetailSingleItemRow::OnPasteProperty()
 {
 	FString ClipboardContent;
diff --git a/Engine/Source/Editor/PropertyEditor/Private/SDetailSingleItemRow.h b/Engine/Source/Editor/PropertyEditor/Private/SDetailSingleItemRow.h
--- a/Engine/Source/Editor/PropertyEditor/Private/SDetailSingleItemRow.h
+++ b/Engine/Source/Editor/PropertyEditor/Private/SDetailSingleItemRow.h
@@ -68,6 +68,7 @@ protected:
 	virtual bool OnContextMenuOpening( FMenuBuilder& MenuBuilder ) override;
 private:
 	void OnCopyProperty();
+	void OnCopyPropertyDisplayName();
 	void OnPasteProperty();
 	bool CanPasteProperty() const;
 	FSlateColor GetOuterBackgroundColor() const;
