Fix typo preventing changing the bShowHiddenPropertiesWhilePlaying from being changed in the UI

diff --git a/Engine/Source/Editor/PropertyEditor/Private/SDetailsView.cpp b/Engine/Source/Editor/PropertyEditor/Private/SDetailsView.cpp
--- a/Engine/Source/Editor/PropertyEditor/Private/SDetailsView.cpp
+++ b/Engine/Source/Editor/PropertyEditor/Private/SDetailsView.cpp
@@ -955,7 +955,7 @@ void SDetailsView::OnShowHiddenPropertiesWhilePlayingClicked()
 {
 	const bool bNewValue = !IsShowHiddenPropertiesWhilePlayingChecked();
 
-	GetMutableViewConfig().bShowAllAdvanced = bNewValue;
+	GetMutableViewConfig().bShowHiddenPropertiesWhilePlaying = bNewValue;
 	SaveViewConfig();
 
 	GConfig->SetBool(TEXT("/Script/EditorStyle.EditorStyleSettings"), TEXT("bShowHiddenPropertiesWhilePlaying"), bNewValue, GEditorPerProjectIni);
