Fixed crash in class picker if we dont have previous values to check

diff --git a/Engine/Source/Editor/PropertyEditor/Private/UserInterface/PropertyEditor/SPropertyEditorEditInline.cpp b/Engine/Source/Editor/PropertyEditor/Private/UserInterface/PropertyEditor/SPropertyEditorEditInline.cpp
--- a/Engine/Source/Editor/PropertyEditor/Private/UserInterface/PropertyEditor/SPropertyEditorEditInline.cpp
+++ b/Engine/Source/Editor/PropertyEditor/Private/UserInterface/PropertyEditor/SPropertyEditorEditInline.cpp
@@ -266,7 +266,7 @@ void SPropertyEditorEditInline::OnClassPicked(UClass* InClass)
 			if (InClass)
 			{
 				FStringView CurClassName, CurObjectName;
-				if (PrevPerObjectValues[NewValues.Num()] != FName(NAME_None).ToString())
+				if (PrevPerObjectValues.IsValidIndex(NewValues.Num()) && PrevPerObjectValues[NewValues.Num()] != FName(NAME_None).ToString())
 				{
 					ExtractClassAndObjectNames(PrevPerObjectValues[NewValues.Num()], CurClassName, CurObjectName);
 				}
