PR #8266: Fix crash in EditConditionParser when using invalid enum/value. (Contributed by [REDACTED_PERSON])

diff --git a/Engine/Source/Editor/PropertyEditor/Private/EditConditionParser.cpp b/Engine/Source/Editor/PropertyEditor/Private/EditConditionParser.cpp
--- a/Engine/Source/Editor/PropertyEditor/Private/EditConditionParser.cpp
+++ b/Engine/Source/Editor/PropertyEditor/Private/EditConditionParser.cpp
@@ -559,6 +559,11 @@ static FExpressionResult EnumPropertyEquals(const EditConditionParserTokens::FEn
 {
 	TOptional<FString> TypeName = Context.GetTypeName(Property.PropertyName);
 
+	if (!TypeName.IsSet())
+	{
+		return MakeError(FText::Format(LOCTEXT("InvalidOperand", "EditCondition attempted to use an invalid operand \"{0}\"."), FText::FromString(Property.PropertyName)));
+	}
+
 	if (TypeName.GetValue() != Enum.Type)
 	{
 		return MakeError(FText::Format(LOCTEXT("OperandTypeMismatch", "EditCondition attempted to compare operands of different types: \"{0}\" and \"{1}\"."), FText::FromString(Property.PropertyName), FText::FromString(Enum.Type + TEXT("::") + Enum.Value)));
