Fixed crash in class picker if we dont have previous values to check

diff --git a/Engine/Source/Editor/PropertyEditor/Private/UserInterface/PropertyEditor/SPropertyEditorEditInline.cpp b/Engine/Source/Editor/PropertyEditor/Private/UserInterface/PropertyEditor/SPropertyEditorEditInline.cpp
--- a/Engine/Source/Editor/PropertyEditor/Private/UserInterface/PropertyEditor/SPropertyEditorEditInline.cpp
+++ b/Engine/Source/Editor/PropertyEditor/Private/UserInterface/PropertyEditor/SPropertyEditorEditInline.cpp
@@ -293,7 +293,10 @@ void SPropertyEditorEditInline::OnClassPicked(UClass* InClass)
 							SubObject->Rename(*MakeUniqueObjectName(GetTransientPackage(), SubObject->GetClass()).ToString(), GetTransientPackage(), REN_DontCreateRedirectors);
 
 							// If we've renamed the object out of the way here, we don't need to do it again below
-							PrevPerObjectValues[NewValues.Num()].Reset();
+							if (PrevPerObjectValues.IsValidIndex(NewValues.Num()))
+							{
+								PrevPerObjectValues[NewValues.Num()].Reset();
+							}
 						}
 					}
 
