Fixed an issue that caused hovering to not work on some details panel buttons

diff --git a/Engine/Source/Editor/PropertyEditor/Private/PropertyCustomizationHelpers.cpp b/Engine/Source/Editor/PropertyEditor/Private/PropertyCustomizationHelpers.cpp
--- a/Engine/Source/Editor/PropertyEditor/Private/PropertyCustomizationHelpers.cpp
+++ b/Engine/Source/Editor/PropertyEditor/Private/PropertyCustomizationHelpers.cpp
@@ -74,7 +74,6 @@ namespace PropertyCustomizationHelpers
 					.ButtonStyle( FAppStyle::Get(), "SimpleButton" )
 					.OnClicked( this, &SPropertyEditorButton::OnClick )
 					.ToolTipText( InArgs._Text )
-					.ForegroundColor( FSlateColor::UseForeground() )
 					.ContentPadding(0)
 					.IsFocusable(InArgs._IsFocusable)
 					[ 
