Added ellipsis effect for clipped labels in details panels

diff --git a/Engine/Source/Editor/PropertyEditor/Private/MaterialList.cpp b/Engine/Source/Editor/PropertyEditor/Private/MaterialList.cpp
--- a/Engine/Source/Editor/PropertyEditor/Private/MaterialList.cpp
+++ b/Engine/Source/Editor/PropertyEditor/Private/MaterialList.cpp
@@ -132,6 +132,7 @@ public:
 				SNew( STextBlock )
 				.Font( IDetailLayoutBuilder::GetDetailFont() )
 				.Text( FText::Format(LOCTEXT("ElementIndex", "Element {ElementIndex}"), Arguments ) )
+				.OverflowPolicy(ETextOverflowPolicy::Ellipsis)
 			]
 			+SVerticalBox::Slot()
 			.Padding(0.0f,4.0f)
diff --git a/Engine/Source/Editor/PropertyEditor/Private/UserInterface/PropertyEditor/SPropertyEditorTitle.h b/Engine/Source/Editor/PropertyEditor/Private/UserInterface/PropertyEditor/SPropertyEditorTitle.h
--- a/Engine/Source/Editor/PropertyEditor/Private/UserInterface/PropertyEditor/SPropertyEditorTitle.h
+++ b/Engine/Source/Editor/PropertyEditor/Private/UserInterface/PropertyEditor/SPropertyEditorTitle.h
@@ -44,18 +44,20 @@ public:
 		if (bHasTitleProperty)
 		{
 			NameTextBlock =
-				SNew( STextBlock )
-				.Text( InPropertyEditor, &FPropertyEditor::GetDisplayName )
-				.Font( NameFont );
+				SNew(STextBlock)
+				.Text(InPropertyEditor, &FPropertyEditor::GetDisplayName)
+				.Font(NameFont);
 		}		
 		else
 		{
 			NameTextBlock =
-				SNew( STextBlock )
-				.Text( InPropertyEditor->GetDisplayName() )
-				.Font( NameFont );
+				SNew(STextBlock)
+				.Text(InPropertyEditor->GetDisplayName())
+				.Font(NameFont);
 		}
 
+		NameTextBlock->SetOverflowPolicy(ETextOverflowPolicy::Ellipsis);
+
 		TSharedPtr<SWidget> NameWidget = NameTextBlock;
 		const bool bInArray = PropertyNode->GetProperty() != nullptr && PropertyNode->GetArrayIndex() != INDEX_NONE;
 		if (bInArray && !bHasTitleProperty)
