LWC: Temporary support added to PropertyHandleVector for vectors of either float or double type.

diff --git a/Engine/Source/Editor/PropertyEditor/Private/PropertyHandleImpl.cpp b/Engine/Source/Editor/PropertyEditor/Private/PropertyHandleImpl.cpp
--- a/Engine/Source/Editor/PropertyEditor/Private/PropertyHandleImpl.cpp
+++ b/Engine/Source/Editor/PropertyEditor/Private/PropertyHandleImpl.cpp
@@ -4097,6 +4097,68 @@ FPropertyAccess::Result FPropertyHandleObject::SetObjectValueFromSelection()
 	return Res;
 }
 
+
+// Temporary mixed float/double property handle to support the various default vector types having differing component types. LWC_TODO: Remove once all types support double.
+FPropertyHandleVector::FPropertyHandleMixed::FPropertyHandleMixed(TSharedRef<FPropertyNode> PropertyNode, FNotifyHook* NotifyHook, TSharedPtr<IPropertyUtilities> PropertyUtilities)
+	: FPropertyHandleBase(PropertyNode, NotifyHook, PropertyUtilities) {}
+
+bool FPropertyHandleVector::FPropertyHandleMixed::Supports(TSharedRef<FPropertyNode> PropertyNode)
+{
+	FProperty* Property = PropertyNode->GetProperty();
+
+	if (Property == nullptr)
+	{
+		return false;
+	}
+
+	return Property->IsA(FFloatProperty::StaticClass()) || Property->IsA(FDoubleProperty::StaticClass());
+}
+
+FPropertyAccess::Result FPropertyHandleVector::FPropertyHandleMixed::GetValue(double& OutValue) const
+{
+	void* PropValue = nullptr;
+	FPropertyAccess::Result Res = Implementation->GetValueData(PropValue);
+
+	if (Res == FPropertyAccess::Success)
+	{
+		if(Implementation->IsPropertyTypeOf(FFloatProperty::StaticClass()))
+		{
+			OutValue = Implementation->GetPropertyValue<FFloatProperty>(PropValue);
+		}
+		else
+		{
+			OutValue = Implementation->GetPropertyValue<FDoubleProperty>(PropValue);
+		}
+	}
+
+	return Res;
+}
+
+FPropertyAccess::Result FPropertyHandleVector::FPropertyHandleMixed::SetValue(const double& NewValue, EPropertyValueSetFlags::Type Flags)
+{
+	FPropertyAccess::Result Res;
+	// Clamp the value from any meta data ranges stored on the property value
+	double FinalValue = ClampValueFromMetaData<double>(NewValue, *Implementation->GetPropertyNode());
+
+	const FString ValueStr = FString::Printf(TEXT("%f"), FinalValue);
+	Res = Implementation->ImportText(ValueStr, Flags);
+
+	return Res;
+}
+
+FPropertyAccess::Result FPropertyHandleVector::FPropertyHandleMixed::GetValue(float& OutValue) const
+{
+	double AsDouble;
+	FPropertyAccess::Result Res = GetValue(AsDouble);
+	OutValue = AsDouble;
+	return Res;
+}
+
+FPropertyAccess::Result FPropertyHandleVector::FPropertyHandleMixed::SetValue(const float& NewValue, EPropertyValueSetFlags::Type Flags)
+{
+	return SetValue((double)NewValue);
+}
+
 // Vector
 bool FPropertyHandleVector::Supports( TSharedRef<FPropertyNode> PropertyNode )
 {
@@ -4128,19 +4190,19 @@ FPropertyHandleVector::FPropertyHandleVector( TSharedRef<class FPropertyNode> Pr
 {
 	const bool bRecurse = false;
 	// A vector is a struct property that has 3 children.  We get/set the values from the children
-	VectorComponents.Add( MakeShareable( new FPropertyHandleFloat( Implementation->GetChildNode("X", bRecurse).ToSharedRef(), NotifyHook, PropertyUtilities ) ) );
+	VectorComponents.Add( MakeShareable( new FPropertyHandleMixed( Implementation->GetChildNode("X", bRecurse).ToSharedRef(), NotifyHook, PropertyUtilities ) ) );
 
-	VectorComponents.Add( MakeShareable( new FPropertyHandleFloat( Implementation->GetChildNode("Y", bRecurse).ToSharedRef(), NotifyHook, PropertyUtilities ) ) );
+	VectorComponents.Add( MakeShareable( new FPropertyHandleMixed( Implementation->GetChildNode("Y", bRecurse).ToSharedRef(), NotifyHook, PropertyUtilities ) ) );
 
 	if( Implementation->GetNumChildren() > 2 )
 	{
 		// at least a 3 component vector
-		VectorComponents.Add( MakeShareable( new FPropertyHandleFloat( Implementation->GetChildNode("Z",bRecurse).ToSharedRef(), NotifyHook, PropertyUtilities ) ) );
+		VectorComponents.Add( MakeShareable( new FPropertyHandleMixed( Implementation->GetChildNode("Z",bRecurse).ToSharedRef(), NotifyHook, PropertyUtilities ) ) );
 	}
 	if( Implementation->GetNumChildren() > 3 )
 	{
 		// a 4 component vector
-		VectorComponents.Add( MakeShareable( new FPropertyHandleFloat( Implementation->GetChildNode("W",bRecurse).ToSharedRef(), NotifyHook, PropertyUtilities ) ) );
+		VectorComponents.Add( MakeShareable( new FPropertyHandleMixed( Implementation->GetChildNode("W",bRecurse).ToSharedRef(), NotifyHook, PropertyUtilities ) ) );
 	}
 }
 
@@ -4193,11 +4255,9 @@ FPropertyAccess::Result FPropertyHandleVector::GetValue( FVector& OutValue ) con
 	if( VectorComponents.Num() == 3 )
 	{
 		// To get the value from the vector we read each child.  If reading a child fails, the value for that component is not set
-		FVector3f OutValue32;	// LWC_TODO: Forced reading of properties as float. Precision loss?
-		FPropertyAccess::Result ResX = VectorComponents[0]->GetValue( OutValue32.X );
-		FPropertyAccess::Result ResY = VectorComponents[1]->GetValue( OutValue32.Y );
-		FPropertyAccess::Result ResZ = VectorComponents[2]->GetValue( OutValue32.Z );
-		OutValue = OutValue32;	// LWC_TODO: Perf pessimization
+		FPropertyAccess::Result ResX = VectorComponents[0]->GetValue( OutValue.X );
+		FPropertyAccess::Result ResY = VectorComponents[1]->GetValue( OutValue.Y );
+		FPropertyAccess::Result ResZ = VectorComponents[2]->GetValue( OutValue.Z );
 
 		if( ResX == FPropertyAccess::Fail || ResY == FPropertyAccess::Fail || ResZ == FPropertyAccess::Fail )
 		{
@@ -4316,21 +4376,21 @@ FPropertyAccess::Result FPropertyHandleVector::SetValue( const FQuat& NewValue,
 	return SetValue(VectorProxy);
 }
 
-FPropertyAccess::Result FPropertyHandleVector::SetX( float InValue, EPropertyValueSetFlags::Type Flags )
+FPropertyAccess::Result FPropertyHandleVector::SetX( double InValue, EPropertyValueSetFlags::Type Flags )
 {
 	FPropertyAccess::Result Res = VectorComponents[0]->SetValue( InValue, Flags );
 
 	return Res;
 }
 
-FPropertyAccess::Result FPropertyHandleVector::SetY( float InValue, EPropertyValueSetFlags::Type Flags )
+FPropertyAccess::Result FPropertyHandleVector::SetY(double InValue, EPropertyValueSetFlags::Type Flags )
 {
 	FPropertyAccess::Result Res = VectorComponents[1]->SetValue( InValue, Flags );
 
 	return Res;
 }
 
-FPropertyAccess::Result FPropertyHandleVector::SetZ( float InValue, EPropertyValueSetFlags::Type Flags )
+FPropertyAccess::Result FPropertyHandleVector::SetZ(double InValue, EPropertyValueSetFlags::Type Flags )
 {
 	if( VectorComponents.Num() > 2 )
 	{
@@ -4342,7 +4402,7 @@ FPropertyAccess::Result FPropertyHandleVector::SetZ( float InValue, EPropertyVal
 	return FPropertyAccess::Fail;
 }
 
-FPropertyAccess::Result FPropertyHandleVector::SetW( float InValue, EPropertyValueSetFlags::Type Flags )
+FPropertyAccess::Result FPropertyHandleVector::SetW(double InValue, EPropertyValueSetFlags::Type Flags )
 {
 	if( VectorComponents.Num() == 4 )
 	{
diff --git a/Engine/Source/Editor/PropertyEditor/Private/PropertyHandleImpl.h b/Engine/Source/Editor/PropertyEditor/Private/PropertyHandleImpl.h
--- a/Engine/Source/Editor/PropertyEditor/Private/PropertyHandleImpl.h
+++ b/Engine/Source/Editor/PropertyEditor/Private/PropertyHandleImpl.h
@@ -618,12 +618,23 @@ public:
 	virtual FPropertyAccess::Result GetValue( FQuat& OutValue ) const override;
 	virtual FPropertyAccess::Result SetValue( const FQuat& InValue, EPropertyValueSetFlags::Type Flags = EPropertyValueSetFlags::DefaultFlags ) override;
 
-	virtual FPropertyAccess::Result SetX( float InValue, EPropertyValueSetFlags::Type Flags = EPropertyValueSetFlags::DefaultFlags );
-	virtual FPropertyAccess::Result SetY( float InValue, EPropertyValueSetFlags::Type Flags = EPropertyValueSetFlags::DefaultFlags );
-	virtual FPropertyAccess::Result SetZ( float InValue, EPropertyValueSetFlags::Type Flags = EPropertyValueSetFlags::DefaultFlags );
-	virtual FPropertyAccess::Result SetW( float InValue, EPropertyValueSetFlags::Type Flags = EPropertyValueSetFlags::DefaultFlags );
+	virtual FPropertyAccess::Result SetX(double InValue, EPropertyValueSetFlags::Type Flags = EPropertyValueSetFlags::DefaultFlags);
+	virtual FPropertyAccess::Result SetY(double InValue, EPropertyValueSetFlags::Type Flags = EPropertyValueSetFlags::DefaultFlags);
+	virtual FPropertyAccess::Result SetZ(double InValue, EPropertyValueSetFlags::Type Flags = EPropertyValueSetFlags::DefaultFlags);
+	virtual FPropertyAccess::Result SetW(double InValue, EPropertyValueSetFlags::Type Flags = EPropertyValueSetFlags::DefaultFlags);
 private:
-	TArray< TSharedPtr<FPropertyHandleFloat> > VectorComponents;
+	// LWC_TODO: Replace with FPropertyHandleDouble once all types support it
+	class FPropertyHandleMixed: public FPropertyHandleBase
+	{
+	public:
+		FPropertyHandleMixed(TSharedRef<FPropertyNode> PropertyNode, FNotifyHook* NotifyHook, TSharedPtr<IPropertyUtilities> PropertyUtilities);
+		static bool Supports(TSharedRef<FPropertyNode> PropertyNode);
+		virtual FPropertyAccess::Result GetValue(float& OutValue) const override;
+		virtual FPropertyAccess::Result GetValue(double& OutValue) const override;
+		virtual FPropertyAccess::Result SetValue(const float& InValue, EPropertyValueSetFlags::Type Flags = EPropertyValueSetFlags::DefaultFlags) override;
+		virtual FPropertyAccess::Result SetValue(const double& InValue, EPropertyValueSetFlags::Type Flags = EPropertyValueSetFlags::DefaultFlags) override;
+	};
+	TArray< TSharedPtr<FPropertyHandleMixed> > VectorComponents;	
 };
 
 class FPropertyHandleRotator : public FPropertyHandleBase
