Add a helper function to create a save button on object selector widget.

diff --git a/Engine/Source/Editor/PropertyEditor/Private/PropertyCustomizationHelpers.cpp b/Engine/Source/Editor/PropertyEditor/Private/PropertyCustomizationHelpers.cpp
--- a/Engine/Source/Editor/PropertyEditor/Private/PropertyCustomizationHelpers.cpp
+++ b/Engine/Source/Editor/PropertyEditor/Private/PropertyCustomizationHelpers.cpp
@@ -393,6 +393,18 @@ namespace PropertyCustomizationHelpers
 		return IDocumentation::Get()->CreateAnchor(DocLink, FString(), DocExcerptName);
 	}
 
+	TSharedRef<SWidget> MakeSaveButton(FSimpleDelegate OnSaveClicked, TAttribute<FText> OptionalToolTipText, TAttribute<bool> IsEnabled)
+	{
+		return
+			SNew(SPropertyEditorButton)
+			.Text(LOCTEXT("SaveButtonLabel", "Save"))
+			.ToolTipText(OptionalToolTipText.Get().IsEmpty() ? LOCTEXT("SaveButtonTooltipText", "Save the currently selected asset.") : OptionalToolTipText)
+			.Image(FEditorStyle::GetBrush("Icons.Save"))
+			.OnClickAction(OnSaveClicked)
+			.IsEnabled(IsEnabled)
+			.IsFocusable(false);
+	}
+
 	FBoolProperty* GetEditConditionProperty(const FProperty* InProperty, bool& bNegate)
 	{
 		FBoolProperty* EditConditionProperty = NULL;
diff --git a/Engine/Source/Editor/PropertyEditor/Public/PropertyCustomizationHelpers.h b/Engine/Source/Editor/PropertyEditor/Public/PropertyCustomizationHelpers.h
--- a/Engine/Source/Editor/PropertyEditor/Public/PropertyCustomizationHelpers.h
+++ b/Engine/Source/Editor/PropertyEditor/Public/PropertyCustomizationHelpers.h
@@ -112,6 +112,7 @@ namespace PropertyCustomizationHelpers
 	PROPERTYEDITOR_API TSharedRef<SWidget> MakeSceneDepthPicker(FOnSceneDepthLocationSelected OnSceneDepthLocationSelected);
 	PROPERTYEDITOR_API TSharedRef<SWidget> MakeEditConfigHierarchyButton(FSimpleDelegate OnEditConfigClicked, TAttribute<FText> OptionalToolTipText = FText(), TAttribute<bool> IsEnabled = true);
 	PROPERTYEDITOR_API TSharedRef<SWidget> MakeDocumentationButton(const TSharedRef<FPropertyEditor>& InPropertyEditor);
+	PROPERTYEDITOR_API TSharedRef<SWidget> MakeSaveButton(FSimpleDelegate OnSaveClicked, TAttribute<FText> OptionalToolTipText = FText(), TAttribute<bool> IsEnabled = true);
 
 	/** @return the FBoolProperty edit condition property if one exists. */
 	PROPERTYEDITOR_API FBoolProperty* GetEditConditionProperty(const FProperty* InProperty, bool& bNegate);
