Add support for prioritization of categories.
Change commontext display order

diff --git a/Engine/Plugins/Experimental/CommonUI/Source/CommonUI/Public/CommonTextBlock.h b/Engine/Plugins/Experimental/CommonUI/Source/CommonUI/Public/CommonTextBlock.h
--- a/Engine/Plugins/Experimental/CommonUI/Source/CommonUI/Public/CommonTextBlock.h
+++ b/Engine/Plugins/Experimental/CommonUI/Source/CommonUI/Public/CommonTextBlock.h
@@ -157,7 +157,7 @@ private:
 /**
  * Text block with automatic scrolling for FX / large texts, also supports a larger set of default styling, & custom mobile scaling.
  */
-UCLASS(Config = CommonUI, DefaultConfig, ClassGroup = UI, meta = (Category = "Common UI", DisplayName = "Common Text"))
+UCLASS(Config = CommonUI, DefaultConfig, ClassGroup = UI, meta = (Category = "Common UI", DisplayName = "Common Text", PrioritizeCategories = "Content"))
 class COMMONUI_API UCommonTextBlock : public UTextBlock
 {
 	GENERATED_UCLASS_BODY()
diff --git a/Engine/Source/Editor/PropertyEditor/Private/ObjectPropertyNode.cpp b/Engine/Source/Editor/PropertyEditor/Private/ObjectPropertyNode.cpp
--- a/Engine/Source/Editor/PropertyEditor/Private/ObjectPropertyNode.cpp
+++ b/Engine/Source/Editor/PropertyEditor/Private/ObjectPropertyNode.cpp
@@ -506,7 +506,7 @@ void FObjectPropertyNode::GetCategoryProperties(const TSet<UClass*>& ClassesToCo
 		{
 			if (!CategoriesFromBlueprints.Contains(CategoryName) && !CategoriesFromProperties.Contains(CategoryName))
 			{
-				SortedCategories.Add(CategoryName);
+				SortedCategories.AddUnique(CategoryName);
 			}
 			CategoriesFromProperties.Add(CategoryName);
 		}
@@ -545,6 +545,7 @@ void FObjectPropertyNode::InternalInitChildNodes( FName SinglePropertyName )
 	// Create a merged list of user-enforced sorted info, hidden category info, etc...
 	TSet<FName> CategoriesFromBlueprints, CategoriesFromProperties;
 	TArray<FName> SortedCategories;
+	TArray<FName> PrioritizeCategories;
 
 #if WITH_EDITORONLY_DATA
 	for (UClass* Class : ClassesToConsider)
@@ -556,10 +557,19 @@ void FObjectPropertyNode::InternalInitChildNodes( FName SinglePropertyName )
 				if (!CategoriesFromBlueprints.Contains(TestCategory))
 				{
 					CategoriesFromBlueprints.Add(TestCategory);
-					SortedCategories.Add(TestCategory);
+					SortedCategories.AddUnique(TestCategory);
 				}
 			}
 		}
+
+		TArray<FString> ClassPrioritizeCategories;
+		Class->GetPrioritizeCategories(ClassPrioritizeCategories);
+		for (const FString& ClassPrioritizeCategory : ClassPrioritizeCategories)
+		{
+			FName PrioritizeCategoryName = FName(ClassPrioritizeCategory);
+			SortedCategories.AddUnique(PrioritizeCategoryName);
+			PrioritizeCategories.AddUnique(PrioritizeCategoryName);
+		}
 	}
 #endif
 
@@ -572,17 +582,17 @@ void FObjectPropertyNode::InternalInitChildNodes( FName SinglePropertyName )
 	}
 
 	for (UClass* Class : ClassesToConsider)
-		{
+	{
 		if (Class)
-			{
+		{
 			UScriptStruct* SparseClassDataStruct = Class->GetSparseClassDataStruct();
 			if (SparseClassDataStruct)
-		{
+			{
 				SparseClassDataInstances.Add(Class, TTuple<UScriptStruct*, void*>(SparseClassDataStruct, Class->GetOrCreateSparseClassData()));
 
 				for (TFieldIterator<FProperty> It(SparseClassDataStruct); It; ++It)
-		{
-		 			GetCategoryProperties(ClassesToConsider, *It, bShouldShowDisableEditOnInstance, bShouldShowHiddenProperties, CategoriesFromBlueprints, CategoriesFromProperties, SortedCategories);
+				{
+					GetCategoryProperties(ClassesToConsider, *It, bShouldShowDisableEditOnInstance, bShouldShowHiddenProperties, CategoriesFromBlueprints, CategoriesFromProperties, SortedCategories);
 				}
 			}
 		}
@@ -597,6 +607,15 @@ void FObjectPropertyNode::InternalInitChildNodes( FName SinglePropertyName )
 		}
 	}
 
+	// Remove any prioritized categories that don't actually exist
+	for (const FName& PrioritizeCategory : PrioritizeCategories)
+	{
+		if (!CategoriesFromProperties.Contains(PrioritizeCategory))
+		{
+			SortedCategories.Remove(PrioritizeCategory);
+		}
+	}
+
 	//////////////////////////////////////////
 	// Add the category headers and the child items that belong inside of them.
 
diff --git a/Engine/Source/Editor/UnrealEd/Private/Kismet2/BlueprintEditorUtils.cpp b/Engine/Source/Editor/UnrealEd/Private/Kismet2/BlueprintEditorUtils.cpp
--- a/Engine/Source/Editor/UnrealEd/Private/Kismet2/BlueprintEditorUtils.cpp
+++ b/Engine/Source/Editor/UnrealEd/Private/Kismet2/BlueprintEditorUtils.cpp
@@ -1574,6 +1574,7 @@ void FBlueprintEditorUtils::RecreateClassMetaData(UBlueprint* Blueprint, UClass*
 		Class->RemoveMetaData("HideFunctions");
 		Class->RemoveMetaData("AutoExpandCategories");
 		Class->RemoveMetaData("AutoCollapseCategories");
+		Class->RemoveMetaData("PrioritizeCategories");
 		Class->RemoveMetaData("SparseClassDataTypes");
 		Class->RemoveMetaData("ClassGroupNames");
 		Class->RemoveMetaData("Category");
@@ -1600,6 +1601,10 @@ void FBlueprintEditorUtils::RecreateClassMetaData(UBlueprint* Blueprint, UClass*
 			{
 				Class->SetMetaData(TEXT("AutoCollapseCategories"), *ParentClass->GetMetaData("AutoCollapseCategories"));
 			}
+			if (ParentClass->HasMetaData(TEXT("PrioritizeCategories")))
+			{
+				Class->SetMetaData(TEXT("PrioritizeCategories"), *ParentClass->GetMetaData("PrioritizeCategories"));
+			}
 		}
 
 		if (ParentClass->HasMetaData(TEXT("HideFunctions")))
diff --git a/Engine/Source/Programs/UnrealHeaderTool/Private/HeaderParser.cpp b/Engine/Source/Programs/UnrealHeaderTool/Private/HeaderParser.cpp
--- a/Engine/Source/Programs/UnrealHeaderTool/Private/HeaderParser.cpp
+++ b/Engine/Source/Programs/UnrealHeaderTool/Private/HeaderParser.cpp
@@ -77,6 +77,7 @@ const FName FHeaderParserNames::NAME_BlueprintType(TEXT("BlueprintType"));
 const FName FHeaderParserNames::NAME_AutoCollapseCategories(TEXT("AutoCollapseCategories"));
 const FName FHeaderParserNames::NAME_HideFunctions(TEXT("HideFunctions"));
 const FName FHeaderParserNames::NAME_AutoExpandCategories(TEXT("AutoExpandCategories"));
+const FName FHeaderParserNames::NAME_PrioritizeCategories(TEXT("PrioritizeCategories"));
 
 TArray<FString> FHeaderParser::PropertyCPPTypesRequiringUIRanges = { TEXT("float"), TEXT("double") };
 TArray<FString> FHeaderParser::ReservedTypeNames = { TEXT("none") };
@@ -8494,6 +8495,10 @@ void FHeaderParser::ResetClassData()
 		{
 			CurrentClassDef.SetMetaData(FHeaderParserNames::NAME_AutoCollapseCategories, *SuperClassDef->GetMetaData(FHeaderParserNames::NAME_AutoCollapseCategories));
 		}
+		if (SuperClassDef->HasMetaData(FHeaderParserNames::NAME_PrioritizeCategories))
+		{
+			CurrentClassDef.SetMetaData(FHeaderParserNames::NAME_PrioritizeCategories, *SuperClassDef->GetMetaData(FHeaderParserNames::NAME_PrioritizeCategories));
+		}
 	}
 
 	check(CurrentClassDef.GetClassWithin());
diff --git a/Engine/Source/Programs/UnrealHeaderTool/Private/Specifiers/ClassMetadataSpecifiers.def b/Engine/Source/Programs/UnrealHeaderTool/Private/Specifiers/ClassMetadataSpecifiers.def
--- a/Engine/Source/Programs/UnrealHeaderTool/Private/Specifiers/ClassMetadataSpecifiers.def
+++ b/Engine/Source/Programs/UnrealHeaderTool/Private/Specifiers/ClassMetadataSpecifiers.def
@@ -37,6 +37,7 @@ CLASS_METADATA_SPECIFIER(NotEditInlineNew)
 CLASS_METADATA_SPECIFIER(NotPlaceable)
 CLASS_METADATA_SPECIFIER(PerObjectConfig)
 CLASS_METADATA_SPECIFIER(Placeable)
+CLASS_METADATA_SPECIFIER(PrioritizeCategories)
 CLASS_METADATA_SPECIFIER(ProjectUserConfig)
 CLASS_METADATA_SPECIFIER(ShowCategories)
 CLASS_METADATA_SPECIFIER(ShowFunctions)
diff --git a/Engine/Source/Programs/UnrealHeaderTool/Private/UnrealHeaderTool.h b/Engine/Source/Programs/UnrealHeaderTool/Private/UnrealHeaderTool.h
--- a/Engine/Source/Programs/UnrealHeaderTool/Private/UnrealHeaderTool.h
+++ b/Engine/Source/Programs/UnrealHeaderTool/Private/UnrealHeaderTool.h
@@ -32,4 +32,5 @@ struct FHeaderParserNames
 	static const FName NAME_AutoCollapseCategories;
 	static const FName NAME_HideFunctions;
 	static const FName NAME_AutoExpandCategories;
+	static const FName NAME_PrioritizeCategories;
 };
\ No newline at end of file
diff --git a/Engine/Source/Programs/UnrealHeaderTool/Private/UnrealTypeDefinitionInfo.cpp b/Engine/Source/Programs/UnrealHeaderTool/Private/UnrealTypeDefinitionInfo.cpp
--- a/Engine/Source/Programs/UnrealHeaderTool/Private/UnrealTypeDefinitionInfo.cpp
+++ b/Engine/Source/Programs/UnrealHeaderTool/Private/UnrealTypeDefinitionInfo.cpp
@@ -2585,6 +2585,16 @@ void FUnrealClassDefinitionInfo::ParseClassProperties(TArray<FPropertySpecifier>
 			}
 			break;
 
+		case EClassMetadataSpecifier::PrioritizeCategories:
+
+			FHeaderParser::RequireSpecifierValue(*this, PropSpecifier);
+
+			for (FString& Value : PropSpecifier.Values)
+			{
+				PrioritizeCategories.AddUnique(MoveTemp(Value));
+			}
+			break;
+
 		case EClassMetadataSpecifier::DontAutoCollapseCategories:
 
 			FHeaderParser::RequireSpecifierValue(*this, PropSpecifier);
@@ -2851,6 +2861,7 @@ void FUnrealClassDefinitionInfo::MergeCategoryMetaData(TMap<FName, FString>& InM
 	if (SparseClassDataTypes.Num()) { InMetaData.Add(FHeaderParserNames::NAME_SparseClassDataTypes, FString::Join(SparseClassDataTypes, TEXT(" "))); }
 	if (HideFunctions.Num()) { InMetaData.Add(FHeaderParserNames::NAME_HideFunctions, FString::Join(HideFunctions, TEXT(" "))); }
 	if (AutoExpandCategories.Num()) { InMetaData.Add(FHeaderParserNames::NAME_AutoExpandCategories, FString::Join(AutoExpandCategories, TEXT(" "))); }
+	if (PrioritizeCategories.Num()) { InMetaData.Add(FHeaderParserNames::NAME_PrioritizeCategories, FString::Join(PrioritizeCategories, TEXT(" "))); }
 }
 
 void FUnrealClassDefinitionInfo::GetSparseClassDataTypes(TArray<FString>& OutSparseClassDataTypes) const
diff --git a/Engine/Source/Programs/UnrealHeaderTool/Private/UnrealTypeDefinitionInfo.h b/Engine/Source/Programs/UnrealHeaderTool/Private/UnrealTypeDefinitionInfo.h
--- a/Engine/Source/Programs/UnrealHeaderTool/Private/UnrealTypeDefinitionInfo.h
+++ b/Engine/Source/Programs/UnrealHeaderTool/Private/UnrealTypeDefinitionInfo.h
@@ -3426,6 +3426,7 @@ private:
 	TArray<FString> HideFunctions;
 	TArray<FString> AutoExpandCategories;
 	TArray<FString> AutoCollapseCategories;
+	TArray<FString> PrioritizeCategories;
 	TArray<FString> DependsOn;
 	TArray<FString> ClassGroupNames;
 	TArray<FString> SparseClassDataTypes;
diff --git a/Engine/Source/Runtime/CoreUObject/Private/UObject/Class.cpp b/Engine/Source/Runtime/CoreUObject/Private/UObject/Class.cpp
--- a/Engine/Source/Runtime/CoreUObject/Private/UObject/Class.cpp
+++ b/Engine/Source/Runtime/CoreUObject/Private/UObject/Class.cpp
@@ -4890,6 +4890,7 @@ void UClass::PurgeClass(bool bRecompilingOnLoad)
 		RemoveMetaData("HideFunctions");
 		RemoveMetaData("AutoExpandCategories");
 		RemoveMetaData("AutoCollapseCategories");
+		RemoveMetaData("PrioritizeCategories");
 		RemoveMetaData("ClassGroupNames");
 	}
 #endif
@@ -5558,6 +5559,25 @@ bool UClass::IsAutoExpandCategory(const TCHAR* InCategory) const
 	return false;
 }
 
+void UClass::GetPrioritizeCategories(TArray<FString>& OutPrioritizedCategories) const
+{
+	static const FName NAME_PrioritizeCategories(TEXT("PrioritizeCategories"));
+	if (const FString* PrioritizeCategories = FindMetaData(NAME_PrioritizeCategories))
+	{
+		PrioritizeCategories->ParseIntoArray(OutPrioritizedCategories, TEXT(" "), true);
+	}
+}
+
+bool UClass::IsPrioritizeCategory(const TCHAR* InCategory) const
+{
+	static const FName NAME_PrioritizeCategories(TEXT("PrioritizeCategories"));
+	if (const FString* PrioritizeCategories = FindMetaData(NAME_PrioritizeCategories))
+	{
+		return !!FCString::StrfindDelim(**PrioritizeCategories, InCategory, TEXT(" "));
+	}
+	return false;
+}
+
 void UClass::GetAutoCollapseCategories(TArray<FString>& OutAutoCollapseCategories) const
 {
 	static const FName NAME_AutoCollapseCategories(TEXT("AutoCollapseCategories"));
diff --git a/Engine/Source/Runtime/CoreUObject/Public/UObject/Class.h b/Engine/Source/Runtime/CoreUObject/Public/UObject/Class.h
--- a/Engine/Source/Runtime/CoreUObject/Public/UObject/Class.h
+++ b/Engine/Source/Runtime/CoreUObject/Public/UObject/Class.h
@@ -2742,6 +2742,8 @@ public:
 	bool IsFunctionHidden(const TCHAR* InFunction) const;
 	void GetAutoExpandCategories(TArray<FString>& OutAutoExpandCategories) const;
 	bool IsAutoExpandCategory(const TCHAR* InCategory) const;
+	void GetPrioritizeCategories(TArray<FString>& OutPrioritizedCategories) const;
+	bool IsPrioritizeCategory(const TCHAR* InCategory) const;
 	void GetAutoCollapseCategories(TArray<FString>& OutAutoCollapseCategories) const;
 	bool IsAutoCollapseCategory(const TCHAR* InCategory) const;
 	void GetClassGroupNames(TArray<FString>& OutClassGroupNames) const;
diff --git a/Engine/Source/Runtime/CoreUObject/Public/UObject/ObjectMacros.h b/Engine/Source/Runtime/CoreUObject/Public/UObject/ObjectMacros.h
--- a/Engine/Source/Runtime/CoreUObject/Public/UObject/ObjectMacros.h
+++ b/Engine/Source/Runtime/CoreUObject/Public/UObject/ObjectMacros.h
@@ -760,6 +760,8 @@ namespace UC
 		collapseCategories,
 		/// Display properties in the editor using categories (default behaviour).
 		dontCollapseCategories,
+		/// Specifies category display order, unspecified will follow default display order.
+		prioritizeCategories,
 
 		/// All the properties of the class are hidden in the main display by default, and are only shown in the advanced details section.
 		AdvancedClassDisplay,
