Change array drag-and-drop reordering in details to use standard drop zone indicator and tooltip

- Previously, array item reordering in the details panel used custom feedback by changing the background color of the target row.
  This change makes array item reordering use the standard drop zone indicator from STableRow, which draws a blue line between rows to indicate the drop position.
  The standard indicator also makes it clear whether the dragged item will be placed before or after the target when reordering.

- The tooltip for the decorated drop used to be constructed from scratch.
  It has been changed to only provide custom text/icon but otherwise use the default tooltip from FDecoratedDragDropOp.

diff --git a/Engine/Source/Editor/EditorStyle/Private/StarshipStyle.cpp b/Engine/Source/Editor/EditorStyle/Private/StarshipStyle.cpp
--- a/Engine/Source/Editor/EditorStyle/Private/StarshipStyle.cpp
+++ b/Engine/Source/Editor/EditorStyle/Private/StarshipStyle.cpp
@@ -3266,6 +3266,9 @@ void FStarshipEditorStyle::FStyle::SetupPropertyEditorStyles()
 			.SetInactiveHoveredBrush( FSlateNoResource() )
 			.SetTextColor( DefaultForeground )
 			.SetSelectedTextColor( InvertedForeground )
+			.SetDropIndicator_Above(BOX_BRUSH("Common/DropZoneIndicator_Above", FMargin(10.0f / 16.0f, 10.0f / 16.0f, 0.f, 0.f), SelectionColor))
+			.SetDropIndicator_Onto(BOX_BRUSH("Common/DropZoneIndicator_Onto", FMargin(4.0f / 16.0f), SelectionColor))
+			.SetDropIndicator_Below(BOX_BRUSH("Common/DropZoneIndicator_Below", FMargin(10.0f / 16.0f, 0.f, 0.f, 10.0f / 16.0f), SelectionColor))
 			);
 
 		Set("DetailsView.DropZone.Below", new BOX_BRUSH("Common/VerticalBoxDropZoneIndicator_Below", FMargin(10.0f / 16.0f, 0, 0, 10.0f / 16.0f), SelectionColor_Subdued));
diff --git a/Engine/Source/Editor/PropertyEditor/Private/SDetailSingleItemRow.cpp b/Engine/Source/Editor/PropertyEditor/Private/SDetailSingleItemRow.cpp
--- a/Engine/Source/Editor/PropertyEditor/Private/SDetailSingleItemRow.cpp
+++ b/Engine/Source/Editor/PropertyEditor/Private/SDetailSingleItemRow.cpp
@@ -56,26 +56,51 @@ namespace SDetailSingleItemRow_Helper
 	}
 }
 
-void SDetailSingleItemRow::OnArrayDragEnter(const FDragDropEvent& DragDropEvent)
+void SDetailSingleItemRow::OnArrayDragLeave(const FDragDropEvent& DragDropEvent)
 {
-	bIsHoveredDragTarget = true;
+	TSharedPtr<FDecoratedDragDropOp> DecoratedOp = DragDropEvent.GetOperationAs<FDecoratedDragDropOp>();
+	if (DecoratedOp.IsValid())
+	{
+		DecoratedOp->ResetToDefaultToolTip();
+	}
 }
 
-void SDetailSingleItemRow::OnArrayDragLeave(const FDragDropEvent& DragDropEvent)
+/** Compute the new index to which to move the item when it's dropped onto the given index/drop zone. */
+static int32 ComputeNewIndex(int32 OriginalIndex, int32 DropOntoIndex, EItemDropZone DropZone)
 {
-	bIsHoveredDragTarget = false;
+	check(DropZone != EItemDropZone::OntoItem);
+
+	int32 NewIndex = DropOntoIndex;
+	if (DropZone == EItemDropZone::BelowItem)
+	{
+		// If the drop zone is below, then we actually move it to the next item's index
+		NewIndex++;
+	}
+	if (OriginalIndex < NewIndex)
+	{
+		// If the item is moved down the list, then all the other elements below it are shifted up one
+		NewIndex--;
+	}
+
+	return ensure(NewIndex >= 0) ? NewIndex : 0;
 }
 
-bool SDetailSingleItemRow::CheckValidDrop(const TSharedPtr<SDetailSingleItemRow> RowPtr) const
+bool SDetailSingleItemRow::CheckValidDrop(const TSharedPtr<SDetailSingleItemRow> RowPtr, EItemDropZone DropZone) const
 {
+	// Can't drop onto another array item; need to drop above or below
+	if (DropZone == EItemDropZone::OntoItem)
+	{
+		return false;
+	}
+
 	TSharedPtr<FPropertyNode> SwappingPropertyNode = RowPtr->SwappablePropertyNode;
-	if (SwappingPropertyNode.IsValid() && SwappablePropertyNode.IsValid())
+	if (SwappingPropertyNode.IsValid() && SwappablePropertyNode.IsValid() && SwappingPropertyNode != SwappablePropertyNode)
 	{
-		if (SwappingPropertyNode != SwappablePropertyNode)
-		{
-			int32 OriginalIndex = SwappingPropertyNode->GetArrayIndex();
-			int32 NewIndex = SwappablePropertyNode->GetArrayIndex();
+		const int32 OriginalIndex = SwappingPropertyNode->GetArrayIndex();
+		const int32 NewIndex = ComputeNewIndex(OriginalIndex, SwappablePropertyNode->GetArrayIndex(), DropZone);
 
+		if (OriginalIndex != NewIndex)
+		{
 			IDetailsViewPrivate* DetailsView = OwnerTreeNode.Pin()->GetDetailsView();
 			TSharedPtr<IPropertyHandle> SwappingHandle = PropertyEditorHelpers::GetPropertyHandle(SwappingPropertyNode.ToSharedRef(), DetailsView->GetNotifyHook(), DetailsView->GetPropertyUtilities());
 			TSharedPtr<IPropertyHandleArray> ParentHandle = SwappingHandle->GetParentHandle()->AsArray();
@@ -89,10 +114,8 @@ bool SDetailSingleItemRow::CheckValidDrop(const TSharedPtr<SDetailSingleItemRow>
 	return false;
 }
 
-FReply SDetailSingleItemRow::OnArrayDrop(const FDragDropEvent& DragDropEvent)
+FReply SDetailSingleItemRow::OnArrayAcceptDrop(const FDragDropEvent& DragDropEvent, EItemDropZone DropZone, TSharedPtr<FDetailTreeNode> TargetItem)
 {
-	bIsHoveredDragTarget = false;
-	
 	TSharedPtr<FArrayRowDragDropOp> ArrayDropOp = DragDropEvent.GetOperationAs< FArrayRowDragDropOp >();
 	if (!ArrayDropOp.IsValid())
 	{
@@ -105,7 +128,7 @@ FReply SDetailSingleItemRow::OnArrayDrop(const FDragDropEvent& DragDropEvent)
 		return FReply::Unhandled();
 	}
 
-	if (!CheckValidDrop(RowPtr))
+	if (!CheckValidDrop(RowPtr, DropZone))
 	{
 		return FReply::Unhandled();
 	}
@@ -115,8 +138,8 @@ FReply SDetailSingleItemRow::OnArrayDrop(const FDragDropEvent& DragDropEvent)
 	TSharedPtr<FPropertyNode> SwappingPropertyNode = RowPtr->SwappablePropertyNode;
 	TSharedPtr<IPropertyHandle> SwappingHandle = PropertyEditorHelpers::GetPropertyHandle(SwappingPropertyNode.ToSharedRef(), DetailsView->GetNotifyHook(), DetailsView->GetPropertyUtilities());
 	TSharedPtr<IPropertyHandleArray> ParentHandle = SwappingHandle->GetParentHandle()->AsArray();
-	int32 OriginalIndex = SwappingPropertyNode->GetArrayIndex();
-	int32 NewIndex = SwappablePropertyNode->GetArrayIndex();
+	const int32 OriginalIndex = SwappingPropertyNode->GetArrayIndex();
+	const int32 NewIndex = ComputeNewIndex(OriginalIndex, SwappablePropertyNode->GetArrayIndex(), DropZone);
 
 	// Need to swap the moving and target expansion states before saving
 	bool bOriginalSwappableExpansion = SwappablePropertyNode->HasNodeFlags(EPropertyNodeFlags::Expanded) != 0;
@@ -141,7 +164,7 @@ FReply SDetailSingleItemRow::OnArrayDrop(const FDragDropEvent& DragDropEvent)
 	return FReply::Handled();
 }
 
-TOptional<EItemDropZone> SDetailSingleItemRow::OnArrayCanAcceptDrop(const FDragDropEvent& DragDropEvent, EItemDropZone DropZone, TSharedPtr< FDetailTreeNode > Type)
+TOptional<EItemDropZone> SDetailSingleItemRow::OnArrayCanAcceptDrop(const FDragDropEvent& DragDropEvent, EItemDropZone DropZone, TSharedPtr<FDetailTreeNode> TargetItem)
 {
 	TSharedPtr<FArrayRowDragDropOp> ArrayDropOp = DragDropEvent.GetOperationAs< FArrayRowDragDropOp >();
 	if (!ArrayDropOp.IsValid())
@@ -155,18 +178,24 @@ TOptional<EItemDropZone> SDetailSingleItemRow::OnArrayCanAcceptDrop(const FDragD
 		return TOptional<EItemDropZone>();
 	}
 
-	bool IsValidDrop = CheckValidDrop(RowPtr);
+	// Can't drop onto another array item, so recompute our own drop zone to ensure it's above or below
+	const FGeometry& Geometry = GetTickSpaceGeometry();
+	const float LocalPointerY = Geometry.AbsoluteToLocal(DragDropEvent.GetScreenSpacePosition()).Y;
+	const EItemDropZone OverrideDropZone = LocalPointerY < Geometry.GetLocalSize().Y * 0.5f ? EItemDropZone::AboveItem : EItemDropZone::BelowItem;
+
+	const bool IsValidDrop = CheckValidDrop(RowPtr, OverrideDropZone);
+
+	ArrayDropOp->SetValidTarget(IsValidDrop);
+
 	if (!IsValidDrop)
 	{
-		bIsHoveredDragTarget = false;
+		return TOptional<EItemDropZone>();
 	}
 
-	ArrayDropOp->IsValidTarget = IsValidDrop;
-
-	return TOptional<EItemDropZone>();
+	return OverrideDropZone;
 }
 
-FReply SDetailSingleItemRow::OnArrayHeaderDrop(const FDragDropEvent& DragDropEvent)
+FReply SDetailSingleItemRow::OnArrayHeaderAcceptDrop(const FDragDropEvent& DragDropEvent, EItemDropZone DropZone, TSharedPtr< FDetailTreeNode > Type)
 {
 	OnArrayDragLeave(DragDropEvent);
 	return FReply::Handled();
@@ -250,10 +279,9 @@ void SDetailSingleItemRow::Construct( const FArguments& InArgs, FDetailLayoutCus
 
 	TSharedRef<SWidget> Widget = SNullWidget::NullWidget;
 
-	FOnTableRowDragEnter ArrayDragDelegate;
 	FOnTableRowDragLeave ArrayDragLeaveDelegate;
-	FOnTableRowDrop ArrayDropDelegate;
-	FOnCanAcceptDrop ArrayAcceptDropDelegate;
+	FOnAcceptDrop ArrayAcceptDropDelegate;
+	FOnCanAcceptDrop ArrayCanAcceptDropDelegate;
 
 	FDetailColumnSizeData& ColumnSizeData = InOwnerTreeNode->GetDetailsView()->GetColumnSizeData();
 
@@ -371,10 +399,9 @@ void SDetailSingleItemRow::Construct( const FArguments& InArgs, FDetailLayoutCus
 					(CastField<FArrayProperty>(PropertyNode->GetProperty()) != nullptr && 
 					CastField<FObjectProperty>(CastField<FArrayProperty>(PropertyNode->GetProperty())->Inner) != nullptr)) // Is an object array
 				{
-					ArrayDragDelegate = FOnTableRowDragEnter::CreateSP(this, &SDetailSingleItemRow::OnArrayDragEnter);
 					ArrayDragLeaveDelegate = FOnTableRowDragLeave::CreateSP(this, &SDetailSingleItemRow::OnArrayDragLeave);
-					ArrayDropDelegate = FOnTableRowDrop::CreateSP(this, PropertyNode->IsReorderable() ? &SDetailSingleItemRow::OnArrayDrop : &SDetailSingleItemRow::OnArrayHeaderDrop);
-					ArrayAcceptDropDelegate = FOnCanAcceptDrop::CreateSP(this, &SDetailSingleItemRow::OnArrayCanAcceptDrop);
+					ArrayAcceptDropDelegate = FOnAcceptDrop::CreateSP(this, PropertyNode->IsReorderable() ? &SDetailSingleItemRow::OnArrayAcceptDrop : &SDetailSingleItemRow::OnArrayHeaderAcceptDrop);
+					ArrayCanAcceptDropDelegate = FOnCanAcceptDrop::CreateSP(this, &SDetailSingleItemRow::OnArrayCanAcceptDrop);
 				}
 			}
 
@@ -588,10 +615,9 @@ void SDetailSingleItemRow::Construct( const FArguments& InArgs, FDetailLayoutCus
 		STableRow::FArguments()
 			.Style(FEditorStyle::Get(), "DetailsView.TreeView.TableRow")
 			.ShowSelection(false)
-			.OnDragEnter(ArrayDragDelegate)
 			.OnDragLeave(ArrayDragLeaveDelegate)
-			.OnDrop(ArrayDropDelegate)
-			.OnCanAcceptDrop(ArrayAcceptDropDelegate),
+			.OnAcceptDrop(ArrayAcceptDropDelegate)
+			.OnCanAcceptDrop(ArrayCanAcceptDropDelegate),
 		InOwnerTableView
 	);
 }
@@ -622,11 +648,6 @@ FSlateColor SDetailSingleItemRow::GetOuterBackgroundColor() const
 		return FAppStyle::Get().GetSlateColor("Colors.Hover");
 	}
 
-	if (bIsHoveredDragTarget)
-	{
-		return FAppStyle::Get().GetSlateColor("Colors.Hover2");
-	}
-
 	if (IsHovered())
 	{
 		return FAppStyle::Get().GetSlateColor("Colors.Header");
@@ -643,11 +664,6 @@ FSlateColor SDetailSingleItemRow::GetInnerBackgroundColor() const
 		return FAppStyle::Get().GetSlateColor("Colors.Hover");
 	}
 
-	if (bIsHoveredDragTarget)
-	{
-		return FAppStyle::Get().GetSlateColor("Colors.Hover2");
-	}
-
 	if (IsHovered())
 	{
 		return FAppStyle::Get().GetSlateColor("Colors.Header");
@@ -1163,67 +1179,24 @@ FReply SArrayRowHandle::OnDragDetected(const FGeometry& MyGeometry, const FPoint
 TSharedPtr<FArrayRowDragDropOp> SArrayRowHandle::CreateDragDropOperation(TSharedPtr<SDetailSingleItemRow> InRow)
 {
 	TSharedPtr<FArrayRowDragDropOp> Operation = MakeShareable(new FArrayRowDragDropOp(InRow));
+	Operation->Init();
 
 	return Operation;
 }
 
-FText FArrayRowDragDropOp::GetDecoratorText() const
-{
-	if (IsValidTarget)
-	{
-		return NSLOCTEXT("ArrayDragDrop", "PlaceRowHere", "Place Row Here");
-	}
-	else
-	{
-		return NSLOCTEXT("ArrayDragDrop", "CannotPlaceRowHere", "Cannot Place Row Here");
-	}
-}
-
-const FSlateBrush* FArrayRowDragDropOp::GetDecoratorIcon() const
-{
-	if (IsValidTarget)
-	{
-		return FEditorStyle::GetBrush("Graph.ConnectorFeedback.OK");
-	}
-	else 
-	{
-		return FEditorStyle::GetBrush("Graph.ConnectorFeedback.Error");
-	}
-}
-
 FArrayRowDragDropOp::FArrayRowDragDropOp(TSharedPtr<SDetailSingleItemRow> InRow)
 {
-	IsValidTarget = false;
-
 	check(InRow.IsValid());
 	Row = InRow;
 
 	// mark row as being used for drag and drop
 	InRow->SetIsDragDrop(true);
+}
 
-	DecoratorWidget = SNew(SBorder)
-		.Padding(8.f)
-		.BorderImage(FEditorStyle::GetBrush("Graph.ConnectorFeedback.Border"))
-		.Content()
-		[
-			SNew(SHorizontalBox)
-			+ SHorizontalBox::Slot()
-			.AutoWidth()
-			.VAlign(VAlign_Center)
-			[
-				SNew(SImage)
-				.Image_Raw(this, &FArrayRowDragDropOp::GetDecoratorIcon)
-			]
-			+ SHorizontalBox::Slot()
-			.Padding(5,0,0,0)
-			.AutoWidth()
-			.VAlign(VAlign_Center)
-			[
-				SNew(STextBlock)
-				.Text_Raw(this, &FArrayRowDragDropOp::GetDecoratorText)
-			]
-		];
-
+void FArrayRowDragDropOp::Init()
+{
+	SetValidTarget(false);
+	SetupDefaults();
 	Construct();
 }
 
@@ -1237,4 +1210,18 @@ void FArrayRowDragDropOp::OnDrop(bool bDropWasHandled, const FPointerEvent& Mous
 		// reset value
 		RowPtr->SetIsDragDrop(false);
 	}
-}
\ No newline at end of file
+}
+
+void FArrayRowDragDropOp::SetValidTarget(bool IsValidTarget)
+{
+	if (IsValidTarget)
+	{
+		CurrentHoverText = NSLOCTEXT("ArrayDragDrop", "PlaceRowHere", "Place Row Here");
+		CurrentIconBrush = FEditorStyle::GetBrush("Graph.ConnectorFeedback.OK");
+	}
+	else
+	{
+		CurrentHoverText = NSLOCTEXT("ArrayDragDrop", "CannotPlaceRowHere", "Cannot Place Row Here");
+		CurrentIconBrush = FEditorStyle::GetBrush("Graph.ConnectorFeedback.Error");
+	}
+}
diff --git a/Engine/Source/Editor/PropertyEditor/Private/SDetailSingleItemRow.h b/Engine/Source/Editor/PropertyEditor/Private/SDetailSingleItemRow.h
--- a/Engine/Source/Editor/PropertyEditor/Private/SDetailSingleItemRow.h
+++ b/Engine/Source/Editor/PropertyEditor/Private/SDetailSingleItemRow.h
@@ -93,10 +93,9 @@ private:
 	bool IsBlacklistChecked() const;
 
 
-	void OnArrayDragEnter(const FDragDropEvent& DragDropEvent);
 	void OnArrayDragLeave(const FDragDropEvent& DragDropEvent);
-	FReply OnArrayDrop(const FDragDropEvent& DragDropEvent);
-	FReply OnArrayHeaderDrop(const FDragDropEvent& DragDropEvent);
+	FReply OnArrayAcceptDrop(const FDragDropEvent& DragDropEvent, EItemDropZone DropZone, TSharedPtr<FDetailTreeNode> TargetItem);
+	FReply OnArrayHeaderAcceptDrop(const FDragDropEvent& DragDropEvent, EItemDropZone DropZone, TSharedPtr<FDetailTreeNode> TargetItem);
 
 	TOptional<EItemDropZone> OnArrayCanAcceptDrop(const FDragDropEvent& DragDropEvent, EItemDropZone DropZone, TSharedPtr< FDetailTreeNode > Type);
 
@@ -104,7 +103,7 @@ private:
 
 	/** Checks if the current drop event is being dropped into a valid location
 	 */
-	bool CheckValidDrop(const TSharedPtr<SDetailSingleItemRow> RowPtr) const;
+	bool CheckValidDrop(const TSharedPtr<SDetailSingleItemRow> RowPtr, EItemDropZone DropZone) const;
 	
 	TSharedPtr<FPropertyNode> GetPropertyNode() const;
 	TSharedPtr<IPropertyHandle> GetPropertyHandle() const;
@@ -115,7 +114,6 @@ private:
 	FDetailLayoutCustomization* Customization;
 	FDetailWidgetRow WidgetRow;
 	bool bAllowFavoriteSystem;
-	bool bIsHoveredDragTarget;
 	bool bIsDragDropObject;
 	bool bCachedResetToDefaultEnabled;
 	TSharedPtr<FPropertyNode> SwappablePropertyNode;
@@ -129,19 +127,13 @@ public:
 
 	FArrayRowDragDropOp(TSharedPtr<SDetailSingleItemRow> InRow);
 
-	TSharedPtr<SWidget> DecoratorWidget;
+	/** Inits the tooltip, needs to be called after constructing */
+	void Init();
 
 	virtual void OnDrop(bool bDropWasHandled, const FPointerEvent& MouseEvent) override;
 
-	virtual TSharedPtr<SWidget> GetDefaultDecorator() const override
-	{
-		return DecoratorWidget;
-	}
+	/** Update the drag tool tip indicating whether the current drop target is valid */
+	void SetValidTarget(bool IsValidTarget);
 
 	TWeakPtr<SDetailSingleItemRow> Row;
-	bool IsValidTarget;
-
-private:
-	FText GetDecoratorText() const;
-	const FSlateBrush* GetDecoratorIcon() const;
 };
