Moved FDetailsViewArgs to its own file.

diff --git a/Engine/Source/Editor/PropertyEditor/Public/DetailsViewArgs.h b/Engine/Source/Editor/PropertyEditor/Public/DetailsViewArgs.h
new file mode 100644
--- /dev/null
+++ b/Engine/Source/Editor/PropertyEditor/Public/DetailsViewArgs.h
@@ -0,0 +1,159 @@
+// Copyright Epic Games, Inc. All Rights Reserved.
+
+#pragma once
+
+#include "CoreMinimal.h"
+
+enum class EEditDefaultsOnlyNodeVisibility : uint8
+{
+	/** Always show nodes that have the EditDefaultsOnly aka CPF_DisableEditOnInstance flag. */
+	Show,
+	/** Always hide nodes that have the EditDefaultsOnly aka CPF_DisableEditOnInstance flag. */
+	Hide,
+	/** Let the details panel control it. If the CDO is selected EditDefaultsOnly nodes will be visible, otherwise false. */
+	Automatic,
+};
+
+/**
+ * Init params for a details view widget
+ */
+struct FDetailsViewArgs
+{
+	enum ENameAreaSettings
+	{
+		/** The name area should never be displayed */
+		HideNameArea = 0,
+		/** All object types use the name area */
+		ObjectsUseNameArea = 1<<0,
+		/** Actors prioritize using the name area */
+		ActorsUseNameArea = 1<<1,
+		/** Components and actors prioritize using the name area. Components will display their actor owner as the name */
+		ComponentsAndActorsUseNameArea = 1<<2 | ActorsUseNameArea,
+	};
+
+	/** Controls how CPF_DisableEditOnInstance nodes will be treated */
+	EEditDefaultsOnlyNodeVisibility DefaultsOnlyVisibility;
+	/** The command list from the host of the details view, allowing child widgets to bind actions with a bound chord */
+	TSharedPtr<class FUICommandList> HostCommandList;
+	/** The tab manager from the host of the details view, allowing child widgets to spawn tabs */
+	TSharedPtr<class FTabManager> HostTabManager;
+	/** Optional object filter to use for more complex handling of what a details panel is viewing. */
+	TSharedPtr<class FDetailsViewObjectFilter> ObjectFilter;
+	/** Optional custom filter(s) to apply to the class viewer widget for class object property values. */
+	TArray<TSharedRef<class IClassViewerFilter>> ClassViewerFilters;
+
+	/** Identifier for this details view; NAME_None if this view is anonymous */
+	FName ViewIdentifier;
+	/** Notify hook to call when properties are changed */
+	class FNotifyHook* NotifyHook;
+	/** Settings for displaying the name area (@see ENameAreaSettings) */
+	int32 NameAreaSettings;
+	/** The default column width */
+	float ColumnWidth;
+	/** The minimum width of the right column in Slate units. */
+	float RightColumnMinWidth;
+	/** True if the viewed objects updates from editor selection */
+	bool bUpdatesFromSelection : 1;
+	/** True if this property view can be locked */
+	bool bLockable : 1;
+	/** True if we allow searching */
+	bool bAllowSearch : 1;
+	/** True if you want to not show the tip when no objects are selected (should only be used if viewing actors properties or bObjectsUseNameArea is true ) */
+	bool bHideSelectionTip : 1;
+	/** True if you want the search box to have initial keyboard focus */
+	bool bSearchInitialKeyFocus : 1;
+	/** True if the 'Open Selection in Property Matrix' button should be shown */
+	bool bShowPropertyMatrixButton : 1;
+	/** Allow options to be changed */
+	bool bShowOptions : 1;
+	/** True if you want to show the object label */
+	bool bShowObjectLabel : 1;
+	UE_DEPRECATED(5.0, "bShowActorLabel has been renamed bShowObjectLabel")
+	bool bShowActorLabel : 1;
+	/** True if you want to show the 'Show Only Modified Properties'. Only valid in conjunction with bShowOptions */
+	bool bShowModifiedPropertiesOption : 1;
+	/** Bind this delegate to hide differing properties */
+	bool bShowDifferingPropertiesOption : 1;
+	/** If true, the name area will be created but will not be displayed so it can be placed in a custom location.  */
+	bool bCustomNameAreaLocation : 1;
+	/** If true, the filter area will be created but will not be displayed so it can be placed in a custom location.  */
+	bool bCustomFilterAreaLocation : 1;
+	/** If false, the current properties editor will never display the favorite system */
+	bool bAllowFavoriteSystem : 1;
+	/** If true the details panel will assume each object passed in through SetObjects will be a unique object shown in the tree and not combined with other objects */
+	bool bAllowMultipleTopLevelObjects : 1;
+	/** If false, the details panel's scrollbar will always be hidden. Useful when embedding details panels in widgets that either grow to accommodate them, or with scrollbars of their own. */
+	bool bShowScrollBar : 1;
+	/** If true, all properties will be visible, not just those with CPF_Edit */
+	bool bForceHiddenPropertyVisibility : 1;
+	/** True if you want to show the 'Show Only Keyable Properties'. Only valid in conjunction with bShowOptions */
+	bool bShowKeyablePropertiesOption : 1;
+	/** True if you want to show the 'Show Only Animated Properties'. Only valid in conjunction with bShowOptions */
+	bool bShowAnimatedPropertiesOption: 1;
+	/** True if you want to show a custom filter. */
+	bool bShowCustomFilterOption : 1;
+	/** True if the section selector should be shown. */
+	bool bShowSectionSelector : 1;
+
+public:
+
+	FDetailsViewArgs()
+		: DefaultsOnlyVisibility(EEditDefaultsOnlyNodeVisibility::Show)
+		, ViewIdentifier(NAME_None)
+		, NotifyHook(nullptr)
+		, NameAreaSettings(ActorsUseNameArea)
+		, ColumnWidth(.65f)
+		, RightColumnMinWidth(24)
+		, bUpdatesFromSelection(false)
+		, bLockable(false)
+		, bAllowSearch(true)
+		, bHideSelectionTip(false)
+		, bSearchInitialKeyFocus(false)
+		, bShowPropertyMatrixButton(true)
+		, bShowOptions(true)
+		, bShowObjectLabel(true)
+		, bShowModifiedPropertiesOption(true)
+		, bShowDifferingPropertiesOption(false)
+		, bCustomNameAreaLocation(false)
+		, bCustomFilterAreaLocation(false)
+		, bAllowFavoriteSystem(false)	
+		, bAllowMultipleTopLevelObjects(false)
+		, bShowScrollBar(true)
+		, bForceHiddenPropertyVisibility(false)
+		, bShowKeyablePropertiesOption(true)
+		, bShowAnimatedPropertiesOption(true)
+		, bShowCustomFilterOption(false)
+		, bShowSectionSelector(false)
+	{
+	}
+
+	/** Default constructor */
+	UE_DEPRECATED(5.0, "This constructor is deprecated, please create an empty FDetailsViewArgs and explicitly set any values that you wish to change.")
+	FDetailsViewArgs( const bool InUpdateFromSelection
+					, const bool InLockable = false
+					, const bool InAllowSearch = true
+					, const ENameAreaSettings InNameAreaSettings = ActorsUseNameArea
+					, const bool InHideSelectionTip = false
+					, FNotifyHook* InNotifyHook = NULL
+					, const bool InSearchInitialKeyFocus = false
+					, FName InViewIdentifier = NAME_None )
+		: FDetailsViewArgs()
+	{
+		bUpdatesFromSelection = InUpdateFromSelection;
+		bLockable = InLockable;
+		bAllowSearch = InAllowSearch;
+		NameAreaSettings = InNameAreaSettings;
+		bHideSelectionTip = InHideSelectionTip;
+		NotifyHook = InNotifyHook;
+		bSearchInitialKeyFocus = InSearchInitialKeyFocus;
+		ViewIdentifier = InViewIdentifier;
+	}
+
+	FDetailsViewArgs(const FDetailsViewArgs&);
+	FDetailsViewArgs& operator=(const FDetailsViewArgs&);
+};
+
+PRAGMA_DISABLE_DEPRECATION_WARNINGS
+inline FDetailsViewArgs::FDetailsViewArgs(const FDetailsViewArgs&) = default;
+inline FDetailsViewArgs& FDetailsViewArgs::operator=(const FDetailsViewArgs&) = default;
+PRAGMA_ENABLE_DEPRECATION_WARNINGS
diff --git a/Engine/Source/Editor/PropertyEditor/Public/IDetailsView.h b/Engine/Source/Editor/PropertyEditor/Public/IDetailsView.h
--- a/Engine/Source/Editor/PropertyEditor/Public/IDetailsView.h
+++ b/Engine/Source/Editor/PropertyEditor/Public/IDetailsView.h
@@ -6,6 +6,7 @@
 #include "Widgets/SCompoundWidget.h"
 #include "PropertyEditorDelegates.h"
 #include "Framework/Commands/UICommandList.h"
+#include "DetailsViewArgs.h"
 
 class AActor;
 class FNotifyHook;
@@ -20,160 +21,8 @@ class FComplexPropertyNode;
 
 typedef TArray<TSharedPtr<FComplexPropertyNode>> FRootPropertyNodeList;
 
-enum class EEditDefaultsOnlyNodeVisibility : uint8
-{
-	/** Always show nodes that have the EditDefaultsOnly aka CPF_DisableEditOnInstance flag. */
-	Show,
-	/** Always hide nodes that have the EditDefaultsOnly aka CPF_DisableEditOnInstance flag. */
-	Hide,
-	/** Let the details panel control it. If the CDO is selected EditDefaultsOnly nodes will be visible, otherwise false. */
-	Automatic,
-};
-
 DECLARE_DELEGATE_RetVal_OneParam(bool, FOnValidateDetailsViewPropertyNodes, const FRootPropertyNodeList&)
 
-/**
- * Init params for a details view widget
- */
-struct FDetailsViewArgs
-{
-	enum ENameAreaSettings
-	{
-		/** The name area should never be displayed */
-		HideNameArea = 0,
-		/** All object types use the name area */
-		ObjectsUseNameArea = 1<<0,
-		/** Actors prioritize using the name area */
-		ActorsUseNameArea = 1<<1,
-		/** Components and actors prioritize using the name area. Components will display their actor owner as the name */
-		ComponentsAndActorsUseNameArea = 1<<2 | ActorsUseNameArea,
-	};
-
-	/** Controls how CPF_DisableEditOnInstance nodes will be treated */
-	EEditDefaultsOnlyNodeVisibility DefaultsOnlyVisibility;
-	/** The command list from the host of the details view, allowing child widgets to bind actions with a bound chord */
-	TSharedPtr<class FUICommandList> HostCommandList;
-	/** The tab manager from the host of the details view, allowing child widgets to spawn tabs */
-	TSharedPtr<FTabManager> HostTabManager;
-	/** Optional object filter to use for more complex handling of what a details panel is viewing. */
-	TSharedPtr<FDetailsViewObjectFilter> ObjectFilter;
-	/** Optional custom filter(s) to apply to the class viewer widget for class object property values. */
-	TArray<TSharedRef<class IClassViewerFilter>> ClassViewerFilters;
-
-	/** Identifier for this details view; NAME_None if this view is anonymous */
-	FName ViewIdentifier;
-	/** Notify hook to call when properties are changed */
-	FNotifyHook* NotifyHook;
-	/** Settings for displaying the name area (@see ENameAreaSettings) */
-	int32 NameAreaSettings;
-	/** The default column width */
-	float ColumnWidth;
-	/** The minimum width of the right column in Slate units. */
-	float RightColumnMinWidth;
-	/** True if the viewed objects updates from editor selection */
-	bool bUpdatesFromSelection : 1;
-	/** True if this property view can be locked */
-	bool bLockable : 1;
-	/** True if we allow searching */
-	bool bAllowSearch : 1;
-	/** True if you want to not show the tip when no objects are selected (should only be used if viewing actors properties or bObjectsUseNameArea is true ) */
-	bool bHideSelectionTip : 1;
-	/** True if you want the search box to have initial keyboard focus */
-	bool bSearchInitialKeyFocus : 1;
-	/** True if the 'Open Selection in Property Matrix' button should be shown */
-	bool bShowPropertyMatrixButton : 1;
-	/** Allow options to be changed */
-	bool bShowOptions : 1;
-	/** True if you want to show the object label */
-	bool bShowObjectLabel : 1;
-	UE_DEPRECATED(5.0, "bShowActorLabel has been renamed bShowObjectLabel")
-	bool bShowActorLabel : 1;
-	/** True if you want to show the 'Show Only Modified Properties'. Only valid in conjunction with bShowOptions */
-	bool bShowModifiedPropertiesOption : 1;
-	/** Bind this delegate to hide differing properties */
-	bool bShowDifferingPropertiesOption : 1;
-	/** If true, the name area will be created but will not be displayed so it can be placed in a custom location.  */
-	bool bCustomNameAreaLocation : 1;
-	/** If true, the filter area will be created but will not be displayed so it can be placed in a custom location.  */
-	bool bCustomFilterAreaLocation : 1;
-	/** If false, the current properties editor will never display the favorite system */
-	bool bAllowFavoriteSystem : 1;
-	/** If true the details panel will assume each object passed in through SetObjects will be a unique object shown in the tree and not combined with other objects */
-	bool bAllowMultipleTopLevelObjects : 1;
-	/** If false, the details panel's scrollbar will always be hidden. Useful when embedding details panels in widgets that either grow to accommodate them, or with scrollbars of their own. */
-	bool bShowScrollBar : 1;
-	/** If true, all properties will be visible, not just those with CPF_Edit */
-	bool bForceHiddenPropertyVisibility : 1;
-	/** True if you want to show the 'Show Only Keyable Properties'. Only valid in conjunction with bShowOptions */
-	bool bShowKeyablePropertiesOption : 1;
-	/** True if you want to show the 'Show Only Animated Properties'. Only valid in conjunction with bShowOptions */
-	bool bShowAnimatedPropertiesOption: 1;
-	/** True if you want to show a custom filter. */
-	bool bShowCustomFilterOption : 1;
-
-public:
-
-	FDetailsViewArgs()
-		: DefaultsOnlyVisibility(EEditDefaultsOnlyNodeVisibility::Show)
-		, ViewIdentifier(NAME_None)
-		, NotifyHook(nullptr)
-		, NameAreaSettings(ActorsUseNameArea)
-		, ColumnWidth(.65f)
-		, RightColumnMinWidth(50)
-		, bUpdatesFromSelection(false)
-		, bLockable(false)
-		, bAllowSearch(true)
-		, bHideSelectionTip(false)
-		, bSearchInitialKeyFocus(false)
-		, bShowPropertyMatrixButton(true)
-		, bShowOptions(true)
-		, bShowObjectLabel(true)
-		, bShowModifiedPropertiesOption(true)
-		, bShowDifferingPropertiesOption(false)
-		, bCustomNameAreaLocation(false)
-		, bCustomFilterAreaLocation(false)
-		, bAllowFavoriteSystem(false)	
-		, bAllowMultipleTopLevelObjects(false)
-		, bShowScrollBar(true)
-		, bForceHiddenPropertyVisibility(false)
-		, bShowKeyablePropertiesOption(true)
-		, bShowAnimatedPropertiesOption(true)
-		, bShowCustomFilterOption(false)
-	{
-	}
-
-	/** Default constructor */
-	UE_DEPRECATED(5.0, "This constructor is deprecated, please create an empty FDetailsViewArgs and explicitly set any values that you wish to change.")
-	FDetailsViewArgs( const bool InUpdateFromSelection
-					, const bool InLockable = false
-					, const bool InAllowSearch = true
-					, const ENameAreaSettings InNameAreaSettings = ActorsUseNameArea
-					, const bool InHideSelectionTip = false
-					, FNotifyHook* InNotifyHook = NULL
-					, const bool InSearchInitialKeyFocus = false
-					, FName InViewIdentifier = NAME_None )
-		: FDetailsViewArgs()
-	{
-		bUpdatesFromSelection = InUpdateFromSelection;
-		bLockable = InLockable;
-		bAllowSearch = InAllowSearch;
-		NameAreaSettings = InNameAreaSettings;
-		bHideSelectionTip = InHideSelectionTip;
-		NotifyHook = InNotifyHook;
-		bSearchInitialKeyFocus = InSearchInitialKeyFocus;
-		ViewIdentifier = InViewIdentifier;
-	}
-
-	FDetailsViewArgs(const FDetailsViewArgs&);
-	FDetailsViewArgs& operator=(const FDetailsViewArgs&);
-};
-
-PRAGMA_DISABLE_DEPRECATION_WARNINGS
-inline FDetailsViewArgs::FDetailsViewArgs(const FDetailsViewArgs&) = default;
-inline FDetailsViewArgs& FDetailsViewArgs::operator=(const FDetailsViewArgs&) = default;
-PRAGMA_ENABLE_DEPRECATION_WARNINGS
-
-
 /**
  * Interface class for all detail views
  */
