The IsEnabled() attribute of FDetailPropertyRows and FDetailWidgetRows is now correctly used.

diff --git a/Engine/Source/Editor/PropertyEditor/Private/DetailItemNode.cpp b/Engine/Source/Editor/PropertyEditor/Private/DetailItemNode.cpp
--- a/Engine/Source/Editor/PropertyEditor/Private/DetailItemNode.cpp
+++ b/Engine/Source/Editor/PropertyEditor/Private/DetailItemNode.cpp
@@ -266,11 +266,8 @@ bool FDetailItemNode::GenerateStandaloneWidget(FDetailWidgetRow& OutRow) const
 		// We make some slight modifications to the row here before giving it to OutRow
 		if (HasMultiColumnWidget())
 		{
-			TSharedPtr<SWidget> NameWidget;
-			TSharedPtr<SWidget> ValueWidget;
-
-			NameWidget = Row.NameWidget.Widget;
-			ValueWidget =
+			TSharedPtr<SWidget> NameWidget = Row.NameWidget.Widget;
+			TSharedPtr<SWidget> ValueWidget =
 				SNew(SConstrainedBox)
 				.MinWidth(Row.ValueWidget.MinWidth)
 				.MaxWidth(Row.ValueWidget.MaxWidth)
@@ -278,10 +275,38 @@ bool FDetailItemNode::GenerateStandaloneWidget(FDetailWidgetRow& OutRow) const
 					Row.ValueWidget.Widget
 				];
 
-			if (Row.IsEnabledAttr.IsSet() || Row.IsEnabledAttr.IsBound())
+			if (Row.IsEnabledAttr.IsSet() || Row.IsValueEnabledAttr.IsSet() || Row.EditConditionValue.IsSet())
 			{
-				NameWidget->SetEnabled(Row.IsEnabledAttr);
-				ValueWidget->SetEnabled(Row.IsEnabledAttr);
+				// copies of attributes for lambda captures
+				TAttribute<bool> PropertyEnabledAttribute = IsPropertyEditingEnabled();
+				TAttribute<bool> RowIsEnabledAttribute = Row.IsEnabledAttr;
+				TAttribute<bool> RowEditConditionAttribute = Row.EditConditionValue;
+
+				TAttribute<bool> IsEnabledAttribute = TAttribute<bool>::CreateLambda(
+					[PropertyEnabledAttribute, RowIsEnabledAttribute, RowEditConditionAttribute]()
+					{
+						return PropertyEnabledAttribute.Get() && RowIsEnabledAttribute.Get(true) && RowEditConditionAttribute.Get(true);
+					});
+
+				// there's an unavoidable conflict here if the user customizes the widget to have a custom IsEnabled,
+				// and a custom EditCondition/IsEnabled on the widget row - we choose to favor the row in this case
+				NameWidget->SetEnabled(IsEnabledAttribute);
+
+				if (Row.IsValueEnabledAttr.IsSet())
+				{
+					TAttribute<bool> RowIsValueEnabledAttribute = Row.IsValueEnabledAttr;
+					TAttribute<bool> IsValueWidgetEnabledAttribute = TAttribute<bool>::CreateLambda(
+						[IsEnabledAttribute, RowIsValueEnabledAttribute]()
+						{
+							return IsEnabledAttribute.Get() && RowIsValueEnabledAttribute.Get(true);
+						});
+
+					ValueWidget->SetEnabled(IsValueWidgetEnabledAttribute);
+				}
+				else
+				{
+					ValueWidget->SetEnabled(IsEnabledAttribute);
+				}
 			}
 
 			OutRow.NameContent()
@@ -303,7 +328,7 @@ bool FDetailItemNode::GenerateStandaloneWidget(FDetailWidgetRow& OutRow) const
 		}
 
 		OutRow.CustomResetToDefault = Row.CustomResetToDefault;
-
+		OutRow.IsEnabledAttr = Row.IsEnabledAttr;
 		OutRow.EditConditionValue = Row.EditConditionValue;
 		OutRow.OnEditConditionValueChanged = Row.OnEditConditionValueChanged;
 
diff --git a/Engine/Source/Editor/PropertyEditor/Private/DetailPropertyRow.cpp b/Engine/Source/Editor/PropertyEditor/Private/DetailPropertyRow.cpp
--- a/Engine/Source/Editor/PropertyEditor/Private/DetailPropertyRow.cpp
+++ b/Engine/Source/Editor/PropertyEditor/Private/DetailPropertyRow.cpp
@@ -279,7 +279,7 @@ void FDetailPropertyRow::OnItemNodeInitialized( TSharedRef<FDetailCategoryImpl>
 		FixEmptyHeaderRowInContainers(PropertyHandle, CustomPropertyWidget);
 
 		// set initial value of enabled attribute to settings from struct customization
-		if (CustomPropertyWidget->IsEnabledAttr.IsBound())
+		if (CustomPropertyWidget->IsEnabledAttr.IsSet())
 		{
 			CustomIsEnabledAttrib = CustomPropertyWidget->IsEnabledAttr;
 		}		
@@ -332,12 +332,7 @@ void FDetailPropertyRow::OnGenerateChildren( FDetailNodeList& OutChildren )
 void FDetailPropertyRow::GenerateChildrenForPropertyNode( TSharedPtr<FPropertyNode>& RootPropertyNode, FDetailNodeList& OutChildren )
 {
 	// Children should be disabled if we are disabled
-	TAttribute<bool> ParentEnabledState = CustomIsEnabledAttrib;
-	if( IsParentEnabled.IsBound() || HasEditCondition() )
-	{
-		// Bind a delegate to the edit condition so our children will be disabled if the edit condition fails
-		ParentEnabledState.Bind( this, &FDetailPropertyRow::GetEnabledState );
-	}
+	TAttribute<bool> ParentEnabledState = TAttribute<bool>::CreateSP(this, &FDetailPropertyRow::GetEnabledState);
 
 	if( PropertyTypeLayoutBuilder.IsValid() && bShowCustomPropertyChildren )
 	{
@@ -636,14 +631,18 @@ bool FDetailPropertyRow::GetEnabledState() const
 {
 	bool Result = IsParentEnabled.Get();
 
-	if (CustomEditConditionValue.IsSet())
-	{
-		Result = Result && CustomEditConditionValue.Get();
-	}
-	
-	if (CustomIsEnabledAttrib.IsSet())
+	Result = Result && CustomIsEnabledAttrib.Get(true);
+
+	if (HasEditCondition())
 	{
-		Result = Result && CustomIsEnabledAttrib.Get();
+		if (CustomEditConditionValue.IsSet())
+		{
+			Result = Result && CustomEditConditionValue.Get();
+		}
+		else
+		{
+			Result = Result && PropertyEditor->IsEditConditionMet();
+		}
 	}
 
 	return Result;
@@ -692,7 +691,7 @@ void FDetailPropertyRow::SetWidgetRowProperties(FDetailWidgetRow& Row) const
 {
 	// set edit condition handlers - use customized if provided
 	TAttribute<bool> EditConditionValue = CustomEditConditionValue;
-	if (!EditConditionValue.IsSet() && !EditConditionValue.IsBound())
+	if (!EditConditionValue.IsSet())
 	{
 		EditConditionValue = TAttribute<bool>(PropertyEditor.ToSharedRef(), &FPropertyEditor::IsEditConditionMet);
 	}
@@ -710,6 +709,7 @@ void FDetailPropertyRow::SetWidgetRowProperties(FDetailWidgetRow& Row) const
 	}
 
 	Row.EditCondition(EditConditionValue, OnEditConditionValueChanged);
+	Row.IsEnabled(CustomIsEnabledAttrib);
 	Row.CustomResetToDefault = CustomResetToDefault;
 	Row.PropertyHandles.Add(GetPropertyHandle());
 
@@ -742,15 +742,7 @@ void FDetailPropertyRow::MakeNameOrKeyWidget( FDetailWidgetRow& Row, const TShar
 		HorizontalAlignment = InCustomRow->NameWidget.HorizontalAlignment;
 	}
 
-	TAttribute<bool> IsEnabledAttrib; 
-	if (CustomIsEnabledAttrib.IsBound() || CustomIsEnabledAttrib.IsSet())
-	{
-		IsEnabledAttrib = CustomIsEnabledAttrib;
-	}
-	else
-	{
-		IsEnabledAttrib.Bind( this, &FDetailPropertyRow::GetEnabledState );
-	}
+	TAttribute<bool> IsEnabledAttrib = TAttribute<bool>::CreateSP( this, &FDetailPropertyRow::GetEnabledState );
 
 	TSharedRef<SHorizontalBox> NameHorizontalBox = SNew(SHorizontalBox)
 		.Clipping(EWidgetClipping::OnDemand);
@@ -840,15 +832,7 @@ void FDetailPropertyRow::MakeValueWidget( FDetailWidgetRow& Row, const TSharedPt
 		HorizontalAlignment = InCustomRow->ValueWidget.HorizontalAlignment;
 	}
 
-	TAttribute<bool> IsEnabledAttrib; 
-	if (CustomIsEnabledAttrib.IsBound() || CustomIsEnabledAttrib.IsSet())
-	{
-		IsEnabledAttrib = CustomIsEnabledAttrib;
-	}
-	else
-	{
-		IsEnabledAttrib.Bind( this, &FDetailPropertyRow::GetEnabledState );
-	}
+	TAttribute<bool> IsEnabledAttrib = TAttribute<bool>::CreateSP( this, &FDetailPropertyRow::GetEnabledState );
 
 	TSharedRef<SHorizontalBox> ValueWidget = 
 		SNew( SHorizontalBox )
diff --git a/Engine/Source/Editor/PropertyEditor/Private/SDetailSingleItemRow.cpp b/Engine/Source/Editor/PropertyEditor/Private/SDetailSingleItemRow.cpp
--- a/Engine/Source/Editor/PropertyEditor/Private/SDetailSingleItemRow.cpp
+++ b/Engine/Source/Editor/PropertyEditor/Private/SDetailSingleItemRow.cpp
@@ -307,37 +307,26 @@ void SDetailSingleItemRow::Construct( const FArguments& InArgs, FDetailLayoutCus
 
 			TSharedPtr<SWidget> ExtensionWidget = WidgetRow.ExtensionWidget.Widget;
 
-			TAttribute<bool> IsEnabledAttribute;
-			if (WidgetRow.IsEnabledAttr.IsSet() || WidgetRow.IsEnabledAttr.IsBound())
-			{
-				TAttribute<bool> RowEnabledAttr = WidgetRow.IsEnabledAttr;
-				TAttribute<bool> PropertyEnabledAttr = InOwnerTreeNode->IsPropertyEditingEnabled();
-				IsEnabledAttribute = TAttribute<bool>::Create(
-					[RowEnabledAttr, PropertyEnabledAttr]()
-					{
-						return RowEnabledAttr.Get() && PropertyEnabledAttr.Get();
-					});
-			}
-			else
-			{
-				IsEnabledAttribute = InOwnerTreeNode->IsPropertyEditingEnabled();
-			}
+			// copies of attributes for lambda captures
+			TAttribute<bool> PropertyEnabledAttribute = InOwnerTreeNode->IsPropertyEditingEnabled();
+			TAttribute<bool> RowEditConditionAttribute = WidgetRow.EditConditionValue;
+			TAttribute<bool> RowIsEnabledAttribute = WidgetRow.IsEnabledAttr;
 
-			TAttribute<bool> IsValueWidgetEnabledAttribute;
-			if (WidgetRow.IsValueEnabledAttr.IsSet() || WidgetRow.IsValueEnabledAttr.IsBound()) {
-				TAttribute<bool> ValueEnabledAttr = WidgetRow.IsValueEnabledAttr;
-				IsValueWidgetEnabledAttribute = TAttribute<bool>::Create(
-					[IsEnabledAttribute, ValueEnabledAttr]()
-					{
-						return IsEnabledAttribute.Get() && ValueEnabledAttr.Get();
-					});
-			}
-			else {
-				IsValueWidgetEnabledAttribute = IsEnabledAttribute;
-			}
+			TAttribute<bool> IsEnabledAttribute = TAttribute<bool>::CreateLambda(
+				[PropertyEnabledAttribute, RowIsEnabledAttribute, RowEditConditionAttribute]()
+				{
+					return PropertyEnabledAttribute.Get() && RowIsEnabledAttribute.Get(true) && RowEditConditionAttribute.Get(true);
+				});
+
+			TAttribute<bool> RowIsValueEnabledAttribute = WidgetRow.IsValueEnabledAttr;
+			TAttribute<bool> IsValueEnabledAttribute = TAttribute<bool>::CreateLambda(
+				[IsEnabledAttribute, RowIsValueEnabledAttribute]()
+				{
+					return IsEnabledAttribute.Get() && RowIsValueEnabledAttribute.Get(true);
+				});
 
 			NameWidget->SetEnabled(IsEnabledAttribute);
-			ValueWidget->SetEnabled(IsValueWidgetEnabledAttribute);
+			ValueWidget->SetEnabled(IsValueEnabledAttribute);
 			ExtensionWidget->SetEnabled(IsEnabledAttribute);
 
 			TSharedRef<SSplitter> Splitter = SNew(SSplitter)
