Drag-and-drop to reorder variables in FUserDefinedStructureEditor

- Allow drag-and-drop behavior to be customized for a details widget row
- Cleanup PropertyEditorHelpers::MakePropertyReorderHandle (it was setting the enabled attribute which always ended up getting overwritten)
- Add drag-and-drop handler+op for UserDefinedStructureEditor to enable drag-and-drop reordering. Removed buttons for moving up/down.

diff --git a/Engine/Source/Editor/Kismet/Private/UserDefinedStructureEditor.cpp b/Engine/Source/Editor/Kismet/Private/UserDefinedStructureEditor.cpp
--- a/Engine/Source/Editor/Kismet/Private/UserDefinedStructureEditor.cpp
+++ b/Engine/Source/Editor/Kismet/Private/UserDefinedStructureEditor.cpp
@@ -20,10 +20,12 @@
 #include "PropertyEditorModule.h"
 #include "IDetailCustomNodeBuilder.h"
 #include "IDetailChildrenBuilder.h"
+#include "IDetailDragDropHandler.h"
 #include "DetailLayoutBuilder.h"
 #include "DetailCategoryBuilder.h"
 #include "Kismet2/BlueprintEditorUtils.h"
 #include "Widgets/Layout/SSplitter.h"
+#include "DragAndDrop/DecoratedDragDropOp.h"
 
 #include "PropertyCustomizationHelpers.h"
 #include "SPinTypeSelector.h"
@@ -568,6 +570,140 @@ enum EMemberFieldPosition
 	MFP_Last	=	0x2,
 };
 
+///////////////////////////////////////////////////////////////////////////////////////
+// FUserDefinedStructureFieldDragDropOp
+
+/** Provides information about the source row (single field) being dragged */
+class FUserDefinedStructureFieldDragDropOp : public FDecoratedDragDropOp
+{
+public:
+	DRAG_DROP_OPERATOR_TYPE(FUserDefinedStructureFieldDragDropOp, FDecoratedDragDropOp);
+
+	FUserDefinedStructureFieldDragDropOp(TWeakPtr<FUserDefinedStructureDetails> InStructureDetails, const FGuid& InFieldGuid)
+		: StructureDetails(InStructureDetails)
+		, FieldGuid(InFieldGuid)
+	{
+		if (TSharedPtr<FUserDefinedStructureDetails> StructureDetailsSP = InStructureDetails.Pin())
+		{
+			VariableFriendlyName = FStructureEditorUtils::GetVariableFriendlyName(StructureDetailsSP->GetUserDefinedStruct(), FieldGuid);
+		}
+	}
+
+	void Init()
+	{
+		SetValidTarget(false);
+		SetupDefaults();
+		Construct();
+	}
+
+	void SetValidTarget(bool IsValidTarget)
+	{
+		FFormatNamedArguments Args;
+		Args.Add(TEXT("StructVariableName"), FText::FromString(VariableFriendlyName));
+
+		if (IsValidTarget)
+		{
+			CurrentHoverText = FText::Format(LOCTEXT("MoveVariableHere", "Move '{StructVariableName}' Here"), Args);
+			CurrentIconBrush = FEditorStyle::GetBrush("Graph.ConnectorFeedback.OK");
+		}
+		else
+		{
+			CurrentHoverText = FText::Format(LOCTEXT("CannotMoveVariableHere", "Cannot Move '{StructVariableName}' Here"), Args);
+			CurrentIconBrush = FEditorStyle::GetBrush("Graph.ConnectorFeedback.Error");
+		}
+	}
+
+	const TWeakPtr<FUserDefinedStructureDetails>& GetStructureDetails() const
+	{
+		return StructureDetails;
+	}
+
+	const FGuid& GetFieldGuid() const
+	{
+		return FieldGuid;
+	}
+
+private:
+	TWeakPtr<FUserDefinedStructureDetails> StructureDetails;
+	FGuid FieldGuid;
+	FString VariableFriendlyName;
+};
+
+///////////////////////////////////////////////////////////////////////////////////////
+// FUserDefinedStructureFieldDragDropHandler
+
+/** Handles drag-and-drop (as both source and target) for a single field's widget row */
+class FUserDefinedStructureFieldDragDropHandler : public IDetailDragDropHandler
+{
+public:
+	FUserDefinedStructureFieldDragDropHandler(TWeakPtr<FUserDefinedStructureDetails> InStructureDetails, const FGuid& InFieldGuid)
+		: StructureDetails(InStructureDetails)
+		, FieldGuid(InFieldGuid)
+	{
+	}
+
+	virtual TSharedPtr<FDragDropOperation> CreateDragDropOperation() const override
+	{
+		TSharedPtr<FUserDefinedStructureFieldDragDropOp> DragOp = MakeShared<FUserDefinedStructureFieldDragDropOp>(StructureDetails, FieldGuid);
+		DragOp->Init();
+		return DragOp;
+	}
+
+	virtual TOptional<EItemDropZone> CanAcceptDrop(const FDragDropEvent& DragDropSource, EItemDropZone DropZone) const override
+	{
+		const TSharedPtr<FUserDefinedStructureFieldDragDropOp> DragOp = DragDropSource.GetOperationAs<FUserDefinedStructureFieldDragDropOp>();
+		if (!DragOp.IsValid())
+		{
+			return TOptional<EItemDropZone>();
+		}
+
+		// Struct must match between drag source and drop target
+		const TSharedPtr<FUserDefinedStructureDetails> OtherStructureDetailsSP = DragOp->GetStructureDetails().Pin();
+		const TSharedPtr<FUserDefinedStructureDetails> MyStructureDetailsSP = StructureDetails.Pin();
+		if (!OtherStructureDetailsSP.IsValid() || !MyStructureDetailsSP.IsValid() || OtherStructureDetailsSP->GetUserDefinedStruct() != MyStructureDetailsSP->GetUserDefinedStruct())
+		{
+			DragOp->SetValidTarget(false);
+			return TOptional<EItemDropZone>();
+		}
+
+		// Struct fields must be moved above or below, so don't allow dropping directly onto a row
+		const EItemDropZone OverrideZone = (DropZone == EItemDropZone::BelowItem) ? EItemDropZone::BelowItem : EItemDropZone::AboveItem;
+		const FStructureEditorUtils::EMovePosition MovePosition = (OverrideZone == EItemDropZone::BelowItem) ? FStructureEditorUtils::PositionBelow : FStructureEditorUtils::PositionAbove;
+		if (!FStructureEditorUtils::CanMoveVariable(MyStructureDetailsSP->GetUserDefinedStruct(), DragOp->GetFieldGuid(), FieldGuid, MovePosition))
+		{
+			DragOp->SetValidTarget(false);
+			return TOptional<EItemDropZone>();
+		}
+
+		DragOp->SetValidTarget(true);
+		return OverrideZone;
+	}
+
+	virtual bool AcceptDrop(const FDragDropEvent& DragDropSource, EItemDropZone DropZone) const override
+	{
+		const TSharedPtr<FUserDefinedStructureFieldDragDropOp> DragOp = DragDropSource.GetOperationAs<FUserDefinedStructureFieldDragDropOp>();
+		if (!DragOp.IsValid())
+		{
+			return false;
+		}
+
+		// Struct must match between drag source and drop target
+		const TSharedPtr<FUserDefinedStructureDetails> OtherStructureDetailsSP = DragOp->GetStructureDetails().Pin();
+		const TSharedPtr<FUserDefinedStructureDetails> MyStructureDetailsSP = StructureDetails.Pin();
+		if (!OtherStructureDetailsSP.IsValid() || !MyStructureDetailsSP.IsValid() || OtherStructureDetailsSP->GetUserDefinedStruct() != MyStructureDetailsSP->GetUserDefinedStruct())
+		{
+			return false;
+		}
+
+		const FStructureEditorUtils::EMovePosition MovePosition = (DropZone == EItemDropZone::BelowItem) ? FStructureEditorUtils::PositionBelow : FStructureEditorUtils::PositionAbove;
+		return FStructureEditorUtils::MoveVariable(MyStructureDetailsSP->GetUserDefinedStruct(), DragOp->GetFieldGuid(), FieldGuid, MovePosition);
+	}
+
+private:
+	TWeakPtr<FUserDefinedStructureDetails> StructureDetails;
+	FGuid FieldGuid;
+};
+
 ///////////////////////////////////////////////////////////////////////////////////////
 // FUserDefinedStructureFieldLayout
 
@@ -846,33 +982,13 @@ public:
 		}
 	}
 
-	FReply OnMoveUp()
-	{
-		auto StructureDetailsSP = StructureDetails.Pin();
-		if (StructureDetailsSP.IsValid() && !(PositionFlags & EMemberFieldPosition::MFP_First))
-		{
-			FStructureEditorUtils::MoveVariable(StructureDetailsSP->GetUserDefinedStruct(), FieldGuid, FStructureEditorUtils::MD_Up);
-		}
-		return FReply::Handled();
-	}
-
-	FReply OnMoveDown()
-	{
-		auto StructureDetailsSP = StructureDetails.Pin();
-		if (StructureDetailsSP.IsValid() && !(PositionFlags & EMemberFieldPosition::MFP_Last))
-		{
-			FStructureEditorUtils::MoveVariable(StructureDetailsSP->GetUserDefinedStruct(), FieldGuid, FStructureEditorUtils::MD_Down);
-		}
-		return FReply::Handled();
-	}
-
 	virtual void GenerateHeaderRowContent( FDetailWidgetRow& NodeRow ) override 
 	{
 		auto K2Schema = GetDefault<UEdGraphSchema_K2>();
 
 		TSharedPtr<SImage> ErrorIcon;
 
-		const float ValueContentWidth = 250.0f;
+		const float ValueContentWidth = 200.0f;
 
 		NodeRow
 		.NameContent()
@@ -915,34 +1031,6 @@ public:
 				.TypeTreeFilter(ETypeTreeFilter::None)
 				.Font( IDetailLayoutBuilder::GetDetailFont() )
 			]
-			+ SHorizontalBox::Slot()
-			.AutoWidth()
-			.HAlign(HAlign_Right)
-			.VAlign(VAlign_Center)
-			[
-				SNew(SButton)
-				.ContentPadding(0)
-				.OnClicked(this, &FUserDefinedStructureFieldLayout::OnMoveUp)
-				.IsEnabled(!(EMemberFieldPosition::MFP_First & PositionFlags))
-				[
-					SNew(SImage)
-					.Image(FEditorStyle::GetBrush("Icons.ChevronUp"))
-				]
-			]
-			+ SHorizontalBox::Slot()
-			.AutoWidth()
-			.HAlign(HAlign_Right)
-			.VAlign(VAlign_Center)
-			[
-				SNew(SButton)
-				.ContentPadding(0)
-				.OnClicked(this, &FUserDefinedStructureFieldLayout::OnMoveDown)
-				.IsEnabled(!(EMemberFieldPosition::MFP_Last & PositionFlags))
-				[
-					SNew(SImage)
-					.Image(FEditorStyle::GetBrush("Icons.ChevronDown"))
-				]
-			]
 			+SHorizontalBox::Slot()
 			.AutoWidth()
 			.HAlign(HAlign_Right)
@@ -953,7 +1041,9 @@ public:
 					LOCTEXT("RemoveVariable", "Remove member variable"),
 					TAttribute<bool>::Create(TAttribute<bool>::FGetter::CreateSP(this, &FUserDefinedStructureFieldLayout::IsRemoveButtonEnabled)))
 			]
-		];
+		]
+		.DragDropHandler(MakeShared<FUserDefinedStructureFieldDragDropHandler>(StructureDetails, FieldGuid))
+		;
 
 		if (ErrorIcon.IsValid())
 		{
diff --git a/Engine/Source/Editor/PropertyEditor/Private/DetailPropertyRow.cpp b/Engine/Source/Editor/PropertyEditor/Private/DetailPropertyRow.cpp
--- a/Engine/Source/Editor/PropertyEditor/Private/DetailPropertyRow.cpp
+++ b/Engine/Source/Editor/PropertyEditor/Private/DetailPropertyRow.cpp
@@ -128,6 +128,12 @@ IDetailPropertyRow& FDetailPropertyRow::OverrideResetToDefault(const FResetToDef
 	return *this;
 }
 
+IDetailPropertyRow& FDetailPropertyRow::DragDropHandler(TSharedPtr<IDetailDragDropHandler> InDragDropHandler)
+{
+	CustomDragDropHandler = InDragDropHandler;
+	return *this;
+}
+
 void FDetailPropertyRow::GetDefaultWidgets( TSharedPtr<SWidget>& OutNameWidget, TSharedPtr<SWidget>& OutValueWidget, bool bAddWidgetDecoration )
 {
 	FDetailWidgetRow Row;
@@ -711,6 +717,7 @@ void FDetailPropertyRow::SetWidgetRowProperties(FDetailWidgetRow& Row) const
 	Row.EditCondition(EditConditionValue, OnEditConditionValueChanged);
 	Row.IsEnabled(CustomIsEnabledAttrib);
 	Row.CustomResetToDefault = CustomResetToDefault;
+	Row.CustomDragDropHandler = CustomDragDropHandler;
 	Row.PropertyHandles.Add(GetPropertyHandle());
 
 	// set custom actions and reset to default
diff --git a/Engine/Source/Editor/PropertyEditor/Private/DetailPropertyRow.h b/Engine/Source/Editor/PropertyEditor/Private/DetailPropertyRow.h
--- a/Engine/Source/Editor/PropertyEditor/Private/DetailPropertyRow.h
+++ b/Engine/Source/Editor/PropertyEditor/Private/DetailPropertyRow.h
@@ -35,6 +35,7 @@ public:
 	virtual IDetailPropertyRow& ShouldAutoExpand(bool bForceExpansion) override;
 	virtual IDetailPropertyRow& Visibility( TAttribute<EVisibility> Visibility ) override;
 	virtual IDetailPropertyRow& OverrideResetToDefault(const FResetToDefaultOverride& ResetToDefault) override;
+	virtual IDetailPropertyRow& DragDropHandler(TSharedPtr<IDetailDragDropHandler> InDragDropHandler) override;
 	virtual FDetailWidgetRow& CustomWidget( bool bShowChildren = false ) override;
 	virtual FDetailWidgetDecl* CustomNameWidget() override;
 	virtual FDetailWidgetDecl* CustomValueWidget() override;
@@ -206,6 +207,8 @@ private:
 	FOnBooleanValueChanged CustomEditConditionValueChanged;
 	/** User customized reset to default */
 	TOptional<FResetToDefaultOverride> CustomResetToDefault;
+	/** User customized drag/drop handler */
+	TSharedPtr<IDetailDragDropHandler> CustomDragDropHandler;
 	/** The category this row resides in */
 	TWeakPtr<FDetailCategoryImpl> ParentCategory;
 	/** Root of the property node if this node comes from an external tree */
diff --git a/Engine/Source/Editor/PropertyEditor/Private/PropertyEditorHelpers.cpp b/Engine/Source/Editor/PropertyEditor/Private/PropertyEditorHelpers.cpp
--- a/Engine/Source/Editor/PropertyEditor/Private/PropertyEditorHelpers.cpp
+++ b/Engine/Source/Editor/PropertyEditor/Private/PropertyEditorHelpers.cpp
@@ -816,7 +816,7 @@ namespace PropertyEditorHelpers
 		return PropertyNode.IsValid() ? !PropertyNode.Pin()->IsEditConst() : false;
 	}
 
-	TSharedRef<SWidget> MakePropertyReorderHandle(const TSharedRef<FPropertyNode>& PropertyNode, TSharedPtr<SDetailSingleItemRow> InParentRow)
+	TSharedRef<SWidget> MakePropertyReorderHandle(TSharedPtr<SDetailSingleItemRow> InParentRow, TAttribute<bool> InEnabledAttr)
 	{
 		TSharedRef<SArrayRowHandle> Handle = SNew(SArrayRowHandle)
 			.Content()
@@ -831,11 +831,9 @@ namespace PropertyEditorHelpers
 					.Image(FCoreStyle::Get().GetBrush("VerticalBoxDragIndicatorShort"))
 				]
 			]
-		.ParentRow(InParentRow);
-		TWeakPtr<FPropertyNode> NodePtr(PropertyNode);
-		TAttribute<bool>::FGetter IsPropertyButtonEnabledDelegate = TAttribute<bool>::FGetter::CreateStatic(&IsPropertyButtonEnabled, NodePtr);
-		TAttribute<bool> IsEnabledAttribute = TAttribute<bool>::Create(IsPropertyButtonEnabledDelegate);
-		Handle->SetEnabled(IsEnabledAttribute);
+		.ParentRow(InParentRow)
+		.IsEnabled(InEnabledAttr)
+		.Visibility_Lambda([InParentRow]() { return InParentRow->IsHovered() ? EVisibility::Visible : EVisibility::Hidden; });
 		return Handle;
 	}
 
diff --git a/Engine/Source/Editor/PropertyEditor/Private/PropertyEditorHelpers.h b/Engine/Source/Editor/PropertyEditor/Private/PropertyEditorHelpers.h
--- a/Engine/Source/Editor/PropertyEditor/Private/PropertyEditorHelpers.h
+++ b/Engine/Source/Editor/PropertyEditor/Private/PropertyEditorHelpers.h
@@ -230,7 +230,7 @@ namespace PropertyEditorHelpers
 	void MakeRequiredPropertyButtons( const TSharedRef< FPropertyEditor >& PropertyEditor, TArray< TSharedRef<SWidget> >& OutButtons, const TArray<EPropertyButton::Type>& ButtonsToIgnore = TArray<EPropertyButton::Type>(), bool bUsingAssetPicker = true );
 
 	TSharedRef<SWidget> MakePropertyButton( const EPropertyButton::Type ButtonType, const TSharedRef< FPropertyEditor >& PropertyEditor );
-	TSharedRef<SWidget> MakePropertyReorderHandle(const TSharedRef<FPropertyNode>& PropertyNode, TSharedPtr<SDetailSingleItemRow> InParentRow);
+	TSharedRef<SWidget> MakePropertyReorderHandle(TSharedPtr<SDetailSingleItemRow> InParentRow, TAttribute<bool> InEnabledAttr);
 	/**
 	 * Recursively finds all object property nodes in a property tree
 	 *
diff --git a/Engine/Source/Editor/PropertyEditor/Private/SDetailSingleItemRow.cpp b/Engine/Source/Editor/PropertyEditor/Private/SDetailSingleItemRow.cpp
--- a/Engine/Source/Editor/PropertyEditor/Private/SDetailSingleItemRow.cpp
+++ b/Engine/Source/Editor/PropertyEditor/Private/SDetailSingleItemRow.cpp
@@ -7,6 +7,7 @@
 #include "DetailWidgetRow.h"
 #include "Editor.h"
 #include "EditorMetadataOverrides.h"
+#include "IDetailDragDropHandler.h"
 #include "IDetailKeyframeHandler.h"
 #include "IDetailPropertyExtensionHandler.h"
 #include "ObjectPropertyNode.h"
@@ -56,7 +57,7 @@ namespace SDetailSingleItemRow_Helper
 	}
 }
 
-void SDetailSingleItemRow::OnArrayDragLeave(const FDragDropEvent& DragDropEvent)
+void SDetailSingleItemRow::OnArrayOrCustomDragLeave(const FDragDropEvent& DragDropEvent)
 {
 	TSharedPtr<FDecoratedDragDropOp> DecoratedOp = DragDropEvent.GetOperationAs<FDecoratedDragDropOp>();
 	if (DecoratedOp.IsValid())
@@ -197,10 +198,33 @@ TOptional<EItemDropZone> SDetailSingleItemRow::OnArrayCanAcceptDrop(const FDragD
 
 FReply SDetailSingleItemRow::OnArrayHeaderAcceptDrop(const FDragDropEvent& DragDropEvent, EItemDropZone DropZone, TSharedPtr< FDetailTreeNode > Type)
 {
-	OnArrayDragLeave(DragDropEvent);
+	OnArrayOrCustomDragLeave(DragDropEvent);
 	return FReply::Handled();
 }
 
+FReply SDetailSingleItemRow::OnCustomAcceptDrop(const FDragDropEvent& DragDropEvent, EItemDropZone DropZone, TSharedPtr<FDetailTreeNode> TargetItem)
+{
+	// This should only be registered as a delegate if there's a custom handler
+	if (ensure(WidgetRow.CustomDragDropHandler))
+	{
+		if (WidgetRow.CustomDragDropHandler->AcceptDrop(DragDropEvent, DropZone))
+		{
+			return FReply::Handled();
+		}
+	}
+	return FReply::Unhandled();
+}
+
+TOptional<EItemDropZone> SDetailSingleItemRow::OnCustomCanAcceptDrop(const FDragDropEvent& DragDropEvent, EItemDropZone DropZone, TSharedPtr<FDetailTreeNode> Type)
+{
+	// This should only be registered as a delegate if there's a custom handler
+	if (ensure(WidgetRow.CustomDragDropHandler))
+	{
+		return WidgetRow.CustomDragDropHandler->CanAcceptDrop(DragDropEvent, DropZone);
+	}
+	return TOptional<EItemDropZone>();
+}
+
 TSharedPtr<FPropertyNode> SDetailSingleItemRow::GetPropertyNode() const
 {
 	TSharedPtr<FPropertyNode> PropertyNode = Customization->GetPropertyNode();
@@ -279,9 +303,9 @@ void SDetailSingleItemRow::Construct( const FArguments& InArgs, FDetailLayoutCus
 
 	TSharedRef<SWidget> Widget = SNullWidget::NullWidget;
 
-	FOnTableRowDragLeave ArrayDragLeaveDelegate;
-	FOnAcceptDrop ArrayAcceptDropDelegate;
-	FOnCanAcceptDrop ArrayCanAcceptDropDelegate;
+	FOnTableRowDragLeave DragLeaveDelegate;
+	FOnAcceptDrop AcceptDropDelegate;
+	FOnCanAcceptDrop CanAcceptDropDelegate;
 
 	FDetailColumnSizeData& ColumnSizeData = InOwnerTreeNode->GetDetailsView()->GetColumnSizeData();
 
@@ -359,18 +383,30 @@ void SDetailSingleItemRow::Construct( const FArguments& InArgs, FDetailLayoutCus
 					SNew(SDetailRowIndent, SharedThis(this))
 				];
 
-			TSharedPtr<FPropertyNode> PropertyNode = Customization->GetPropertyNode();
-			if (PropertyNode.IsValid())
+			if (WidgetRow.CustomDragDropHandler)
+			{
+				TSharedPtr<SDetailSingleItemRow> InRow = SharedThis(this);
+				TSharedRef<SWidget> ReorderHandle = PropertyEditorHelpers::MakePropertyReorderHandle(InRow, IsEnabledAttribute);
+				
+				NameColumnBox->AddSlot()
+					.HAlign(HAlign_Left)
+					.VAlign(VAlign_Center)
+					.Padding(-2, 0, -10, 0)
+					.AutoWidth()
+					[
+						ReorderHandle
+					];
+
+				DragLeaveDelegate = FOnTableRowDragLeave::CreateSP(this, &SDetailSingleItemRow::OnArrayOrCustomDragLeave);
+				AcceptDropDelegate = FOnAcceptDrop::CreateSP(this, &SDetailSingleItemRow::OnCustomAcceptDrop);
+				CanAcceptDropDelegate = FOnCanAcceptDrop::CreateSP(this, &SDetailSingleItemRow::OnCustomCanAcceptDrop);
+			}
+			else if (TSharedPtr<FPropertyNode> PropertyNode = Customization->GetPropertyNode())
 			{
 				if (PropertyNode->IsReorderable())
 				{
 					TSharedPtr<SDetailSingleItemRow> InRow = SharedThis(this);
-					TSharedRef<SWidget> ArrayHandle = PropertyEditorHelpers::MakePropertyReorderHandle(PropertyNode.ToSharedRef(), InRow);
-					ArrayHandle->SetEnabled(IsEnabledAttribute);
-					ArrayHandle->SetVisibility(TAttribute<EVisibility>::Create([this]() 
-					{
-						return this->IsHovered() ? EVisibility::Visible : EVisibility::Hidden;
-					}));
+					TSharedRef<SWidget> ArrayHandle = PropertyEditorHelpers::MakePropertyReorderHandle(InRow, IsEnabledAttribute);
 
 					NameColumnBox->AddSlot()
 						.HAlign(HAlign_Left)
@@ -388,9 +424,9 @@ void SDetailSingleItemRow::Construct( const FArguments& InArgs, FDetailLayoutCus
 					(CastField<FArrayProperty>(PropertyNode->GetProperty()) != nullptr && 
 					CastField<FObjectProperty>(CastField<FArrayProperty>(PropertyNode->GetProperty())->Inner) != nullptr)) // Is an object array
 				{
-					ArrayDragLeaveDelegate = FOnTableRowDragLeave::CreateSP(this, &SDetailSingleItemRow::OnArrayDragLeave);
-					ArrayAcceptDropDelegate = FOnAcceptDrop::CreateSP(this, PropertyNode->IsReorderable() ? &SDetailSingleItemRow::OnArrayAcceptDrop : &SDetailSingleItemRow::OnArrayHeaderAcceptDrop);
-					ArrayCanAcceptDropDelegate = FOnCanAcceptDrop::CreateSP(this, &SDetailSingleItemRow::OnArrayCanAcceptDrop);
+					DragLeaveDelegate = FOnTableRowDragLeave::CreateSP(this, &SDetailSingleItemRow::OnArrayOrCustomDragLeave);
+					AcceptDropDelegate = FOnAcceptDrop::CreateSP(this, PropertyNode->IsReorderable() ? &SDetailSingleItemRow::OnArrayAcceptDrop : &SDetailSingleItemRow::OnArrayHeaderAcceptDrop);
+					CanAcceptDropDelegate = FOnCanAcceptDrop::CreateSP(this, &SDetailSingleItemRow::OnArrayCanAcceptDrop);
 				}
 			}
 
@@ -604,9 +640,9 @@ void SDetailSingleItemRow::Construct( const FArguments& InArgs, FDetailLayoutCus
 		STableRow::FArguments()
 			.Style(FEditorStyle::Get(), "DetailsView.TreeView.TableRow")
 			.ShowSelection(false)
-			.OnDragLeave(ArrayDragLeaveDelegate)
-			.OnAcceptDrop(ArrayAcceptDropDelegate)
-			.OnCanAcceptDrop(ArrayCanAcceptDropDelegate),
+			.OnDragLeave(DragLeaveDelegate)
+			.OnAcceptDrop(AcceptDropDelegate)
+			.OnCanAcceptDrop(CanAcceptDropDelegate),
 		InOwnerTableView
 	);
 }
@@ -658,7 +694,7 @@ FSlateColor SDetailSingleItemRow::GetInnerBackgroundColor() const
 		return FAppStyle::Get().GetSlateColor("Colors.Header");
 	}
 	
-	if (bIsDragDropObject)
+	if (DragOperation.IsValid())
 	{
 		return FAppStyle::Get().GetSlateColor("Colors.Hover");
 	}
@@ -1131,9 +1167,21 @@ bool SDetailSingleItemRow::IsHighlighted() const
 	return OwnerTreeNodePtr.IsValid() ? OwnerTreeNodePtr->IsHighlighted() : false;
 }
 
-void SDetailSingleItemRow::SetIsDragDrop(bool bInIsDragDrop)
+TSharedPtr<FDragDropOperation> SDetailSingleItemRow::CreateDragDropOperation()
 {
-	bIsDragDropObject = bInIsDragDrop;
+	if (WidgetRow.CustomDragDropHandler)
+	{
+		TSharedPtr<FDragDropOperation> DragOp = WidgetRow.CustomDragDropHandler->CreateDragDropOperation();
+		DragOperation = DragOp;
+		return DragOp;
+	}
+	else
+	{
+		TSharedPtr<FArrayRowDragDropOp> ArrayDragOp = MakeShareable(new FArrayRowDragDropOp(SharedThis(this)));
+		ArrayDragOp->Init();
+		DragOperation = ArrayDragOp;
+		return ArrayDragOp;
+	}
 }
 
 void SDetailSingleItemRow::Tick(const FGeometry& AllottedGeometry, const double InCurrentTime, const float InDeltaTime)
@@ -1155,7 +1203,7 @@ FReply SArrayRowHandle::OnDragDetected(const FGeometry& MyGeometry, const FPoint
 {
 	if (MouseEvent.IsMouseButtonDown(EKeys::LeftMouseButton))
 	{
-		TSharedPtr<FDragDropOperation> DragDropOp = CreateDragDropOperation(ParentRow.Pin());
+		TSharedPtr<FDragDropOperation> DragDropOp = ParentRow.Pin()->CreateDragDropOperation();
 		if (DragDropOp.IsValid())
 		{
 			return FReply::Handled().BeginDragDrop(DragDropOp.ToSharedRef());
@@ -1165,21 +1213,10 @@ FReply SArrayRowHandle::OnDragDetected(const FGeometry& MyGeometry, const FPoint
 	return FReply::Unhandled();
 }
 
-TSharedPtr<FArrayRowDragDropOp> SArrayRowHandle::CreateDragDropOperation(TSharedPtr<SDetailSingleItemRow> InRow)
-{
-	TSharedPtr<FArrayRowDragDropOp> Operation = MakeShareable(new FArrayRowDragDropOp(InRow));
-	Operation->Init();
-
-	return Operation;
-}
-
 FArrayRowDragDropOp::FArrayRowDragDropOp(TSharedPtr<SDetailSingleItemRow> InRow)
 {
 	check(InRow.IsValid());
 	Row = InRow;
-
-	// mark row as being used for drag and drop
-	InRow->SetIsDragDrop(true);
 }
 
 void FArrayRowDragDropOp::Init()
@@ -1189,18 +1226,6 @@ void FArrayRowDragDropOp::Init()
 	Construct();
 }
 
-void FArrayRowDragDropOp::OnDrop(bool bDropWasHandled, const FPointerEvent& MouseEvent)
-{
-	FDecoratedDragDropOp::OnDrop(bDropWasHandled, MouseEvent);
-
-	TSharedPtr<SDetailSingleItemRow> RowPtr = Row.Pin();
-	if (RowPtr.IsValid())
-	{
-		// reset value
-		RowPtr->SetIsDragDrop(false);
-	}
-}
-
 void FArrayRowDragDropOp::SetValidTarget(bool IsValidTarget)
 {
 	if (IsValidTarget)
diff --git a/Engine/Source/Editor/PropertyEditor/Private/SDetailSingleItemRow.h b/Engine/Source/Editor/PropertyEditor/Private/SDetailSingleItemRow.h
--- a/Engine/Source/Editor/PropertyEditor/Private/SDetailSingleItemRow.h
+++ b/Engine/Source/Editor/PropertyEditor/Private/SDetailSingleItemRow.h
@@ -37,7 +37,6 @@ public:
 	};
 
 	FReply OnDragDetected(const FGeometry& MyGeometry, const FPointerEvent& MouseEvent) override;
-	TSharedPtr<class FArrayRowDragDropOp> CreateDragDropOperation(TSharedPtr<SDetailSingleItemRow> InRow);
 
 private:
 	TWeakPtr<SDetailSingleItemRow> ParentRow;
@@ -61,7 +60,7 @@ public:
 	 */
 	void Construct( const FArguments& InArgs, FDetailLayoutCustomization* InCustomization, bool bHasMultipleColumns, TSharedRef<FDetailTreeNode> InOwnerTreeNode, const TSharedRef<STableViewBase>& InOwnerTableView );
 
-	void SetIsDragDrop(bool bInIsDragDrop);
+	TSharedPtr<FDragDropOperation> CreateDragDropOperation();
 
 	virtual void Tick(const FGeometry& AllottedGeometry, const double InCurrentTime, const float InDeltaTime);
 protected:
@@ -93,13 +92,14 @@ private:
 	bool IsBlacklistChecked() const;
 
 
-	void OnArrayDragLeave(const FDragDropEvent& DragDropEvent);
+	void OnArrayOrCustomDragLeave(const FDragDropEvent& DragDropEvent);
 	FReply OnArrayAcceptDrop(const FDragDropEvent& DragDropEvent, EItemDropZone DropZone, TSharedPtr<FDetailTreeNode> TargetItem);
 	FReply OnArrayHeaderAcceptDrop(const FDragDropEvent& DragDropEvent, EItemDropZone DropZone, TSharedPtr<FDetailTreeNode> TargetItem);
 
 	TOptional<EItemDropZone> OnArrayCanAcceptDrop(const FDragDropEvent& DragDropEvent, EItemDropZone DropZone, TSharedPtr< FDetailTreeNode > Type);
 
-	TSharedPtr<FPropertyNode> GetCopyPastePropertyNode() const;
+	FReply OnCustomAcceptDrop(const FDragDropEvent& DragDropEvent, EItemDropZone DropZone, TSharedPtr<FDetailTreeNode> TargetItem);
+	TOptional<EItemDropZone> OnCustomCanAcceptDrop(const FDragDropEvent& DragDropEvent, EItemDropZone DropZone, TSharedPtr<FDetailTreeNode> Type);
 
 	/** Checks if the current drop event is being dropped into a valid location
 	 */
@@ -114,10 +114,10 @@ private:
 	FDetailLayoutCustomization* Customization;
 	FDetailWidgetRow WidgetRow;
 	bool bAllowFavoriteSystem;
-	bool bIsDragDropObject;
 	bool bCachedResetToDefaultEnabled;
 	TSharedPtr<FPropertyNode> SwappablePropertyNode;
 	TSharedPtr<SButton> ExpanderArrow;
+	TWeakPtr<FDragDropOperation> DragOperation; // last drag initiated by this widget
 };
 
 class FArrayRowDragDropOp : public FDecoratedDragDropOp
@@ -130,8 +130,6 @@ public:
 	/** Inits the tooltip, needs to be called after constructing */
 	void Init();
 
-	virtual void OnDrop(bool bDropWasHandled, const FPointerEvent& MouseEvent) override;
-
 	/** Update the drag tool tip indicating whether the current drop target is valid */
 	void SetValidTarget(bool IsValidTarget);
 
diff --git a/Engine/Source/Editor/PropertyEditor/Public/DetailWidgetRow.h b/Engine/Source/Editor/PropertyEditor/Public/DetailWidgetRow.h
--- a/Engine/Source/Editor/PropertyEditor/Public/DetailWidgetRow.h
+++ b/Engine/Source/Editor/PropertyEditor/Public/DetailWidgetRow.h
@@ -15,6 +15,7 @@
 
 class FDetailWidgetRow;
 class FResetToDefaultOverride;
+class IDetailDragDropHandler;
 
 /** Widget declaration for custom widgets in a widget row */
 class FDetailWidgetDecl
@@ -127,6 +128,7 @@ public:
 		WholeRowWidget = FDetailWidgetDecl(*this, Other.WholeRowWidget);
 		VisibilityAttr = Other.VisibilityAttr;
 		IsEnabledAttr = Other.IsEnabledAttr;
+		IsValueEnabledAttr = Other.IsValueEnabledAttr;
 		FilterTextString = Other.FilterTextString;
 		CopyMenuAction = Other.CopyMenuAction;
 		PasteMenuAction = Other.PasteMenuAction;
@@ -135,6 +137,7 @@ public:
 		CustomResetToDefault = Other.CustomResetToDefault;
 		EditConditionValue = Other.EditConditionValue;
 		OnEditConditionValueChanged = Other.OnEditConditionValueChanged;
+		CustomDragDropHandler = Other.CustomDragDropHandler;
 		PropertyHandles = Other.PropertyHandles;
 		return *this;
 	}
@@ -313,6 +316,15 @@ public:
 		return *this;
 	}
 	
+	/**
+	 * Sets a handler for the row to be a source or target of drag-and-drop operations.
+	 */
+	FDetailWidgetRow& DragDropHandler(TSharedPtr<IDetailDragDropHandler> InDragDropHandler)
+	{
+		CustomDragDropHandler = InDragDropHandler;
+		return *this;
+	}
+
 	/**
 	* Used to provide all the property handles this WidgetRow represent
 	*/
@@ -373,6 +385,8 @@ public:
 	TAttribute<bool> EditConditionValue;
 	/** Custom edit condition value changed handler. */
 	FOnBooleanValueChanged OnEditConditionValueChanged;
+	/** Custom handler for drag-and-drop of the row */
+	TSharedPtr<IDetailDragDropHandler> CustomDragDropHandler;
 	/* All property handle that this custom widget represent */
 	TArray<TSharedPtr<IPropertyHandle>> PropertyHandles;
 };
diff --git a/Engine/Source/Editor/PropertyEditor/Public/IDetailDragDropHandler.h b/Engine/Source/Editor/PropertyEditor/Public/IDetailDragDropHandler.h
new file mode 100644
--- /dev/null
+++ b/Engine/Source/Editor/PropertyEditor/Public/IDetailDragDropHandler.h
@@ -0,0 +1,28 @@
+// Copyright Epic Games, Inc. All Rights Reserved.
+
+#pragma once
+
+#include "CoreMinimal.h"
+#include "Widgets/Views/STableRow.h"
+
+class FDragDropEvent;
+class FDragDropOperation;
+
+/** Specifies how a widget row can initiate a drag-and-drop (be dragged) and/or act as the target of a drag-and-drop (be dropped onto). */
+class IDetailDragDropHandler
+{
+public:
+	virtual ~IDetailDragDropHandler() {}
+
+	/** Create a drag operation with the current widget row as the source. */
+	virtual TSharedPtr<FDragDropOperation> CreateDragDropOperation() const = 0;
+
+	/** Accepts an arbitrary drag-and-drop (doesn't need to be from another widget row) with the current widget row as the target. */
+	virtual bool AcceptDrop(const FDragDropEvent& DragDropEvent, EItemDropZone DropZone) const = 0;
+
+	/**
+	 * Determines whether the drag-and-drop event contains data that can be accepted by the current widget row as the target.
+	 * The handler can return a different drop zone if desired, or no drop zone if it can't accept the drop.
+	 */
+	virtual TOptional<EItemDropZone> CanAcceptDrop(const FDragDropEvent& DragDropEvent, EItemDropZone DropZone) const = 0;
+};
diff --git a/Engine/Source/Editor/PropertyEditor/Public/IDetailPropertyRow.h b/Engine/Source/Editor/PropertyEditor/Public/IDetailPropertyRow.h
--- a/Engine/Source/Editor/PropertyEditor/Public/IDetailPropertyRow.h
+++ b/Engine/Source/Editor/PropertyEditor/Public/IDetailPropertyRow.h
@@ -13,6 +13,7 @@
 
 class FDetailWidgetRow;
 class FDetailWidgetDecl;
+class IDetailDragDropHandler;
 
 class FDetailColumnSizeData
 {
@@ -305,6 +306,13 @@ public:
 	 */
 	virtual IDetailPropertyRow& OverrideResetToDefault(const FResetToDefaultOverride& ResetToDefault) = 0;
 
+	/**
+	 * Sets a handler for this property row to be a source or target of drag-and-drop operations
+	 * 
+	 * @param InDragDropHandler	Handler used when starting a drag or accepting a drop operation
+	 */
+	virtual IDetailPropertyRow& DragDropHandler(TSharedPtr<IDetailDragDropHandler> InDragDropHandler) = 0;
+
 	/**
 	 * Returns the name and value widget of this property row.  You can use this widget to apply further customization to existing widgets (by using this  with CustomWidget)
 	 *
diff --git a/Engine/Source/Editor/UnrealEd/Private/Kismet2/StructureEditorUtils.cpp b/Engine/Source/Editor/UnrealEd/Private/Kismet2/StructureEditorUtils.cpp
--- a/Engine/Source/Editor/UnrealEd/Private/Kismet2/StructureEditorUtils.cpp
+++ b/Engine/Source/Editor/UnrealEd/Private/Kismet2/StructureEditorUtils.cpp
@@ -753,26 +753,91 @@ bool FStructureEditorUtils::ChangeSaveGameEnabled(UUserDefinedStruct* Struct, FG
 	return false;
 }
 
-bool FStructureEditorUtils::MoveVariable(UUserDefinedStruct* Struct, FGuid VarGuid, EMoveDirection MoveDirection)
+/** Compute the initial and new indices to move the specified variable above/below another variable. */
+static bool ComputeIndicesForMove(
+	const TArray<FStructVariableDescription>& DescArray,
+	const FGuid& MoveVarGuid,
+	const FGuid& RelativeToGuid,
+	FStructureEditorUtils::EMovePosition Position,
+	int32& OutInitialIndex,
+	int32& OutNewIndex)
+{
+	int32 InitialIndex = DescArray.IndexOfByPredicate(
+		[MoveVarGuid](const FStructVariableDescription& Desc)
+		{
+			return Desc.VarGuid == MoveVarGuid;
+		});
+	int32 NewIndex = DescArray.IndexOfByPredicate(
+		[RelativeToGuid](const FStructVariableDescription& Desc)
+		{
+			return Desc.VarGuid == RelativeToGuid;
+		});
+	if (InitialIndex == INDEX_NONE || NewIndex == INDEX_NONE)
+	{
+		return false;
+	}
+
+	if (Position == FStructureEditorUtils::PositionBelow)
+	{
+		// If moving below a variable, then we actually move it to the next variable's index
+		NewIndex++;
+	}
+
+	if (InitialIndex < NewIndex)
+	{
+		// When the element is removed from the array, all the other elements below it are shifted by one,
+		// so moving an element down the array causes its new index to shift by one
+		NewIndex--;
+	}
+
+	if (InitialIndex == NewIndex)
+	{
+		// No move is happening because the index didn't change.
+		return false;
+	}
+
+	if (!ensure(NewIndex >= 0 && NewIndex < DescArray.Num()))
+	{
+		// New index is out of bounds - this shouldn't happen!
+		return false;
+	}
+
+	OutInitialIndex = InitialIndex;
+	OutNewIndex = NewIndex;
+	return true;
+}
+
+bool FStructureEditorUtils::MoveVariable(UUserDefinedStruct* Struct, FGuid MoveVarGuid, FGuid RelativeToGuid, EMovePosition Position)
 {
 	if (Struct)
 	{
-		const bool bMoveUp = (EMoveDirection::MD_Up == MoveDirection);
 		TArray<FStructVariableDescription>& DescArray = GetVarDesc(Struct);
-		const int32 InitialIndex = bMoveUp ? 1 : 0;
-		const int32 IndexLimit = DescArray.Num() - (bMoveUp ? 0 : 1);
-		for (int32 Index = InitialIndex; Index < IndexLimit; ++Index)
+		int32 InitialIndex, NewIndex;
+		if (!ComputeIndicesForMove(DescArray, MoveVarGuid, RelativeToGuid, Position, InitialIndex, NewIndex))
 		{
-			if (DescArray[Index].VarGuid == VarGuid)
-			{
-				const FScopedTransaction Transaction(LOCTEXT("ReorderVariables", "Variables reordered"));
-				ModifyStructData(Struct);
-
-				DescArray.Swap(Index, Index + (bMoveUp ? -1 : 1));
-				OnStructureChanged(Struct, EStructureEditorChangeInfo::MovedVariable);
-				return true;
-			}
+			return false;
 		}
+
+		const FScopedTransaction Transaction(LOCTEXT("ReorderVariables", "Variables reordered"));
+		ModifyStructData(Struct);
+
+		FStructVariableDescription MoveDesc = DescArray[InitialIndex];
+		DescArray.RemoveAt(InitialIndex);
+		DescArray.Insert(MoveDesc, NewIndex);
+
+		OnStructureChanged(Struct, EStructureEditorChangeInfo::MovedVariable);
+		return true;
+	}
+	return false;
+}
+
+bool FStructureEditorUtils::CanMoveVariable(UUserDefinedStruct* Struct, FGuid MoveVarGuid, FGuid RelativeToGuid, EMovePosition Position)
+{
+	if (Struct)
+	{
+		TArray<FStructVariableDescription>& DescArray = GetVarDesc(Struct);
+		int32 OldIndex, NewIndex; // populated but unused
+		return ComputeIndicesForMove(DescArray, MoveVarGuid, RelativeToGuid, Position, OldIndex, NewIndex);
 	}
 	return false;
 }
diff --git a/Engine/Source/Editor/UnrealEd/Public/Kismet2/StructureEditorUtils.h b/Engine/Source/Editor/UnrealEd/Public/Kismet2/StructureEditorUtils.h
--- a/Engine/Source/Editor/UnrealEd/Public/Kismet2/StructureEditorUtils.h
+++ b/Engine/Source/Editor/UnrealEd/Public/Kismet2/StructureEditorUtils.h
@@ -103,12 +103,26 @@ public:
 
 	static bool ChangeSaveGameEnabled(UUserDefinedStruct* Struct, FGuid VarGuid, bool bInSaveGame);
 
-	enum EMoveDirection
+	enum EMovePosition
 	{
-		MD_Up,
-		MD_Down,
+		PositionAbove,
+		PositionBelow,
 	};
-	static bool MoveVariable(UUserDefinedStruct* Struct, FGuid VarGuid, EMoveDirection MoveDirection);
+
+	/**
+	 * Move the variable with the given guid in the struct to be immediately above or below another variable.
+	 * 
+	 * @param Struct         The struct containing the variable to move.
+	 * @param MoveVarGuid    The guid of the variable being moved.
+	 * @param RelativeToGuid The guid of the variable above/below which MoveVarGuid is being moved.
+	 * @param Position       Whether to put MoveVarGuid above or below RelativeToGuid.
+	 * 
+	 * @return Whether the variable was actually moved.
+	 */
+	static bool MoveVariable(UUserDefinedStruct* Struct, FGuid MoveVarGuid, FGuid RelativeToGuid, EMovePosition Position);
+
+	/** Checks whether MoveVariable can actually move the variable with the given guid. */
+	static bool CanMoveVariable(UUserDefinedStruct* Struct, FGuid MoveVarGuid, FGuid RelativeToGuid, EMovePosition Position);
 
 	//Multi-line text
 	static bool CanEnableMultiLineText(const UUserDefinedStruct* Struct, FGuid VarGuid);
