Fix to Tooltip on Editor ComboBox. The tool tip set by that code was never seen and hide  the Tooltip set by the ComboBox that does the work correctly.

diff --git a/Engine/Source/Editor/PropertyEditor/Private/UserInterface/PropertyEditor/SPropertyEditorCombo.cpp b/Engine/Source/Editor/PropertyEditor/Private/UserInterface/PropertyEditor/SPropertyEditorCombo.cpp
--- a/Engine/Source/Editor/PropertyEditor/Private/UserInterface/PropertyEditor/SPropertyEditorCombo.cpp
+++ b/Engine/Source/Editor/PropertyEditor/Private/UserInterface/PropertyEditor/SPropertyEditorCombo.cpp
@@ -242,14 +242,6 @@ void SPropertyEditorCombo::SendToObjects( const FString& NewValue )
 			check(Index != INDEX_NONE);
 
 			Value = Enum->GetNameStringByIndex(Index);
-
-			FText ToolTipValue = Enum->GetToolTipTextByIndex(Index);
-			FText ToolTipText = Property->GetToolTipText();
-			if (!ToolTipValue.IsEmpty())
-			{
-				ToolTipText = FText::Format(FText::FromString(TEXT("{0}\n\n{1}")), ToolTipText, ToolTipValue);
-			}
-			SetToolTipText(ToolTipText);
 		}
 
 		ComboArgs.PropertyHandle->SetValueFromFormattedString(Value);
