DetailsViewArgs::RightColumnMinWidth now set to 30 by default, which doesn't cause the Reset to Default arrow to be ellipsed out.

diff --git a/Engine/Source/Editor/PropertyEditor/Public/DetailsViewArgs.h b/Engine/Source/Editor/PropertyEditor/Public/DetailsViewArgs.h
--- a/Engine/Source/Editor/PropertyEditor/Public/DetailsViewArgs.h
+++ b/Engine/Source/Editor/PropertyEditor/Public/DetailsViewArgs.h
@@ -103,7 +103,7 @@ public:
 		, NotifyHook(nullptr)
 		, NameAreaSettings(ActorsUseNameArea)
 		, ColumnWidth(.65f)
-		, RightColumnMinWidth(24)
+		, RightColumnMinWidth(30)
 		, bUpdatesFromSelection(false)
 		, bLockable(false)
 		, bAllowSearch(true)
