Disallow edge case where a details row is drag-and-dropped between an expanded parent and its first child (eg when array elements are structs)

For example:
- A
+ B    <- don't allow drop below B, unexpected user interaction
|- a
|- b
- C    <- user can drop above C instead

diff --git a/Engine/Source/Editor/PropertyEditor/Private/SDetailSingleItemRow.cpp b/Engine/Source/Editor/PropertyEditor/Private/SDetailSingleItemRow.cpp
--- a/Engine/Source/Editor/PropertyEditor/Private/SDetailSingleItemRow.cpp
+++ b/Engine/Source/Editor/PropertyEditor/Private/SDetailSingleItemRow.cpp
@@ -94,6 +94,19 @@ bool SDetailSingleItemRow::CheckValidDrop(const TSharedPtr<SDetailSingleItemRow>
 		return false;
 	}
 
+	// Disallow dropping immediately below an expanded item (e.g. array of structs) since the drop zone is between the item and its children
+	// User can drop above the next array item instead
+	// For example:
+	// - A    <- allow drop below A because it's collapsed
+	// + B    <- don't allow drop below B because it's expanded;
+	// | - a     dropping here seems to "parent" the dropped row to B
+	// | - b
+	// - C    <- user can drop above C instead
+	if (DropZone == EItemDropZone::BelowItem && IsItemExpanded())
+	{
+		return false;
+	}
+
 	TSharedPtr<FPropertyNode> SwappingPropertyNode = RowPtr->SwappablePropertyNode;
 	if (SwappingPropertyNode.IsValid() && SwappablePropertyNode.IsValid() && SwappingPropertyNode != SwappablePropertyNode)
 	{
@@ -217,6 +230,12 @@ FReply SDetailSingleItemRow::OnCustomAcceptDrop(const FDragDropEvent& DragDropEv
 
 TOptional<EItemDropZone> SDetailSingleItemRow::OnCustomCanAcceptDrop(const FDragDropEvent& DragDropEvent, EItemDropZone DropZone, TSharedPtr<FDetailTreeNode> Type)
 {
+	// Disallow drop between expanded parent item and its first child
+	if (DropZone == EItemDropZone::BelowItem && IsItemExpanded())
+	{
+		return TOptional<EItemDropZone>();
+	}
+
 	// This should only be registered as a delegate if there's a custom handler
 	if (ensure(WidgetRow.CustomDragDropHandler))
 	{
