Actually use OptionalToolTipText param in PropertyCustomizationHelpers::Make*Button helpers

- The OptionalToolTipText parameter in these helpers was previously not surfaced in the UI because the Text property was being used as the tooltip (the buttons are images only)
- Since the two sets of strings were duplicative (only one spot for text on this widget), I kept only one, in most cases just moving the ToolTipText into the Text (e.g. "Use" -> "Use Selected Asset from Content Browser")
- Slightly tweaked layout of SPropertyEditorButton so that the tooltip appears on the entire button, not just the image widget

diff --git a/Engine/Source/Editor/PropertyEditor/Private/PropertyCustomizationHelpers.cpp b/Engine/Source/Editor/PropertyEditor/Private/PropertyCustomizationHelpers.cpp
--- a/Engine/Source/Editor/PropertyEditor/Private/PropertyCustomizationHelpers.cpp
+++ b/Engine/Source/Editor/PropertyEditor/Private/PropertyCustomizationHelpers.cpp
@@ -50,7 +50,7 @@ namespace PropertyCustomizationHelpers
 			, _Image( FEditorStyle::GetBrush("Default") )
 			, _IsFocusable( true )
 		{}
-			SLATE_ARGUMENT( FText, Text )
+			SLATE_ATTRIBUTE( FText, Text )
 			SLATE_ARGUMENT( const FSlateBrush*, Image )
 			SLATE_EVENT( FSimpleDelegate, OnClickAction )
 
@@ -69,11 +69,11 @@ namespace PropertyCustomizationHelpers
 				.VAlign(VAlign_Center)
 				.WidthOverride(22)
 				.HeightOverride(22)
+				.ToolTipText(InArgs._Text)
 				[
 					SNew(SButton)
 					.ButtonStyle( FAppStyle::Get(), "SimpleButton" )
 					.OnClicked( this, &SPropertyEditorButton::OnClick )
-					.ToolTipText( InArgs._Text )
 					.ContentPadding(0)
 					.IsFocusable(InArgs._IsFocusable)
 					[ 
@@ -100,8 +100,7 @@ namespace PropertyCustomizationHelpers
 	{
 		return
 			SNew(SPropertyEditorButton)
-			.Text(LOCTEXT("ResetButtonLabel", "ResetToDefault"))
-			.ToolTipText(OptionalToolTipText.Get().IsEmpty() ? LOCTEXT("ResetButtonToolTipText", "Resets Element to Default Value") : OptionalToolTipText)
+			.Text(OptionalToolTipText.Get().IsEmpty() ? LOCTEXT("ResetButtonToolTipText", "Reset Element to Default Value") : OptionalToolTipText)
 			.Image(FEditorStyle::GetBrush("PropertyWindow.DiffersFromDefault"))
 			.OnClickAction(OnResetClicked)
 			.IsEnabled(IsEnabled)
@@ -114,8 +113,7 @@ namespace PropertyCustomizationHelpers
 	{
 		return	
 			SNew( SPropertyEditorButton )
-			.Text( LOCTEXT( "AddButtonLabel", "Add" ) )
-			.ToolTipText( OptionalToolTipText.Get().IsEmpty() ? LOCTEXT( "AddButtonToolTipText", "Adds Element") : OptionalToolTipText )
+			.Text( OptionalToolTipText.Get().IsEmpty() ? LOCTEXT( "AddButtonToolTipText", "Add Element") : OptionalToolTipText )
 			.Image( FEditorStyle::GetBrush("Icons.PlusCircle") )
 			.OnClickAction( OnAddClicked )
 			.IsEnabled(IsEnabled)
@@ -126,8 +124,7 @@ namespace PropertyCustomizationHelpers
 	{
 		return	
 			SNew( SPropertyEditorButton )
-			.Text( LOCTEXT( "RemoveButtonLabel", "Remove" ) )
-			.ToolTipText( OptionalToolTipText.Get().IsEmpty() ? LOCTEXT( "RemoveButtonToolTipText", "Removes Element") : OptionalToolTipText )
+			.Text( OptionalToolTipText.Get().IsEmpty() ? LOCTEXT( "RemoveButtonToolTipText", "Remove Element") : OptionalToolTipText )
 			.Image( FEditorStyle::GetBrush("Icons.Minus") )
 			.OnClickAction( OnRemoveClicked )
 			.IsEnabled(IsEnabled)
@@ -138,8 +135,7 @@ namespace PropertyCustomizationHelpers
 	{
 		return
 			SNew( SPropertyEditorButton )
-			.Text( LOCTEXT( "EmptyButtonLabel", "Empty" ) )
-			.ToolTipText( OptionalToolTipText.Get().IsEmpty() ? LOCTEXT( "EmptyButtonToolTipText", "Removes All Elements") : OptionalToolTipText )
+			.Text( OptionalToolTipText.Get().IsEmpty() ? LOCTEXT( "EmptyButtonToolTipText", "Remove All Elements") : OptionalToolTipText )
 			.Image( FEditorStyle::GetBrush("Icons.Delete") )
 			.OnClickAction( OnEmptyClicked )
 			.IsEnabled(IsEnabled)
@@ -150,8 +146,7 @@ namespace PropertyCustomizationHelpers
 	{
 		return
 			SNew( SPropertyEditorButton )
-			.Text( LOCTEXT( "UseButtonLabel", "Use") )
-			.ToolTipText( OptionalToolTipText.Get().IsEmpty() ? LOCTEXT( "UseButtonToolTipText", "Use Selected Asset from Content Browser") : OptionalToolTipText )
+			.Text( OptionalToolTipText.Get().IsEmpty() ? LOCTEXT( "UseButtonToolTipText", "Use Selected Asset from Content Browser") : OptionalToolTipText )
 			.Image( FEditorStyle::GetBrush("Icons.CircleArrowLeft") )
 			.OnClickAction( OnUseSelectedClicked )
 			.IsEnabled(IsEnabled)
@@ -162,8 +157,7 @@ namespace PropertyCustomizationHelpers
 	{
 		return
 			SNew( SPropertyEditorButton )
-			.Text( LOCTEXT( "DeleteButtonLabel", "Delete") )
-			.ToolTipText( OptionalToolTipText.Get().IsEmpty() ? LOCTEXT( "DeleteButtonToolTipText", "Delete") : OptionalToolTipText )
+			.Text( OptionalToolTipText.Get().IsEmpty() ? LOCTEXT( "DeleteButtonToolTipText", "Delete") : OptionalToolTipText )
 			.Image( FEditorStyle::GetBrush("Icons.Delete") )
 			.OnClickAction( OnDeleteClicked )
 			.IsEnabled(IsEnabled)
@@ -174,8 +168,7 @@ namespace PropertyCustomizationHelpers
 	{
 		return
 			SNew( SPropertyEditorButton )
-			.Text( LOCTEXT( "ClearButtonLabel", "Clear") )
-			.ToolTipText( OptionalToolTipText.Get().IsEmpty() ? LOCTEXT( "ClearButtonToolTipText", "Clear Path") : OptionalToolTipText )
+			.Text( OptionalToolTipText.Get().IsEmpty() ? LOCTEXT( "ClearButtonToolTipText", "Clear") : OptionalToolTipText )
 			.Image(FAppStyle::Get().GetBrush("Icons.X"))
 			.OnClickAction( OnClearClicked )
 			.IsEnabled(IsEnabled)
@@ -212,8 +205,7 @@ namespace PropertyCustomizationHelpers
 	{
 		return
 			SNew( SPropertyEditorButton )
-			.Text( LOCTEXT( "BrowseButtonLabel", "Browse") )
-			.ToolTipText( OptionalToolTipText.Get().IsEmpty() ? LOCTEXT( "BrowseButtonToolTipText", "Browse to Asset in Content Browser") : OptionalToolTipText )
+			.Text( OptionalToolTipText.Get().IsEmpty() ? LOCTEXT( "BrowseButtonToolTipText", "Browse to Asset in Content Browser") : OptionalToolTipText )
 			.Image( FEditorStyle::GetBrush("Icons.Search") )
 			.OnClickAction( OnFindClicked )
 			.IsEnabled(IsEnabled)
@@ -224,8 +216,7 @@ namespace PropertyCustomizationHelpers
 	{
 		return
 			SNew( SPropertyEditorButton )
-			.Text( LOCTEXT( "NewBlueprintButtonLabel", "New Blueprint") )
-			.ToolTipText( OptionalToolTipText.Get().IsEmpty() ? LOCTEXT( "NewBlueprintButtonToolTipText", "Create New Blueprint") : OptionalToolTipText )
+			.Text( OptionalToolTipText.Get().IsEmpty() ? LOCTEXT( "NewBlueprintButtonToolTipText", "Create New Blueprint") : OptionalToolTipText )
 			.Image( FEditorStyle::GetBrush("Icons.PlusCircle") )
 			.OnClickAction( OnNewBlueprintClicked )
 			.IsEnabled(IsEnabled)
@@ -364,8 +355,7 @@ namespace PropertyCustomizationHelpers
 	{
 		return
 			SNew(SPropertyEditorButton)
-			.Text(LOCTEXT("EditConfigHierarchyButtonLabel", "Edit Config Hierarchy"))
-			.ToolTipText(OptionalToolTipText.Get().IsEmpty() ? LOCTEXT("EditConfigHierarchyButtonToolTipText", "Edit the config values of this property") : OptionalToolTipText)
+			.Text(OptionalToolTipText.Get().IsEmpty() ? LOCTEXT("EditConfigHierarchyButtonToolTipText", "Edit the config values of this property") : OptionalToolTipText)
 			.Image(FEditorStyle::GetBrush("DetailsView.EditConfigProperties"))
 			.OnClickAction(OnEditConfigClicked)
 			.IsEnabled(IsEnabled)
@@ -397,8 +387,7 @@ namespace PropertyCustomizationHelpers
 	{
 		return
 			SNew(SPropertyEditorButton)
-			.Text(LOCTEXT("SaveButtonLabel", "Save"))
-			.ToolTipText(OptionalToolTipText.Get().IsEmpty() ? LOCTEXT("SaveButtonTooltipText", "Save the currently selected asset.") : OptionalToolTipText)
+			.Text(OptionalToolTipText.Get().IsEmpty() ? LOCTEXT("SaveButtonTooltipText", "Save the currently selected asset.") : OptionalToolTipText)
 			.Image(FEditorStyle::GetBrush("Icons.Save"))
 			.OnClickAction(OnSaveClicked)
 			.IsEnabled(IsEnabled)
