ICWYU. Removed "CoreMinimal.h" from Private headers in PropertyEditor module.

diff --git a/Engine/Source/Editor/PropertyEditor/Private/DetailTreeNode.h b/Engine/Source/Editor/PropertyEditor/Private/DetailTreeNode.h
--- a/Engine/Source/Editor/PropertyEditor/Private/DetailTreeNode.h
+++ b/Engine/Source/Editor/PropertyEditor/Private/DetailTreeNode.h
@@ -2,12 +2,11 @@
 
 #pragma once
 
-#include "CoreMinimal.h"
-#include "Misc/Attribute.h"
-#include "PropertyPath.h"
 #include "IPropertyUtilities.h"
 #include "IDetailTreeNode.h"
 #include "PropertyCustomizationHelpers.h"
+#include "PropertyPath.h"
+#include "Misc/Attribute.h"
 
 class FDetailCategoryImpl;
 class IDetailsViewPrivate;
@@ -18,6 +17,8 @@ class FPropertyNode;
 class FComplexPropertyNode;
 struct FDetailFilter;
 
+typedef TArray<TSharedRef<class FDetailTreeNode>> FDetailNodeList;
+
 enum class ENodeVisibility : uint8
 {
 	// Hidden but can be visible if parent is visible due to filtering
@@ -28,9 +29,6 @@ enum class ENodeVisibility : uint8
 	Visible,
 };
 
-typedef TArray<TSharedRef<class FDetailTreeNode>> FDetailNodeList;
-
-
 class FDetailTreeNode : public IDetailTreeNode
 {
 public:
diff --git a/Engine/Source/Editor/PropertyEditor/Private/IDetailsViewPrivate.h b/Engine/Source/Editor/PropertyEditor/Private/IDetailsViewPrivate.h
--- a/Engine/Source/Editor/PropertyEditor/Private/IDetailsViewPrivate.h
+++ b/Engine/Source/Editor/PropertyEditor/Private/IDetailsViewPrivate.h
@@ -2,11 +2,10 @@
 
 #pragma once
 
-#include "CoreMinimal.h"
 #include "AssetThumbnail.h"
 #include "DetailTreeNode.h"
-#include "PropertyNode.h"
 #include "IDetailsView.h"
+#include "PropertyNode.h"
 
 class FEditConditionParser;
 class FNotifyHook;
diff --git a/Engine/Source/Editor/PropertyEditor/Private/SDetailTableRowBase.h b/Engine/Source/Editor/PropertyEditor/Private/SDetailTableRowBase.h
--- a/Engine/Source/Editor/PropertyEditor/Private/SDetailTableRowBase.h
+++ b/Engine/Source/Editor/PropertyEditor/Private/SDetailTableRowBase.h
@@ -2,20 +2,19 @@
 
 #pragma once
 
-#include "CoreMinimal.h"
-#include "InputCoreTypes.h"
-#include "Input/Reply.h"
 #include "DetailTreeNode.h"
-#include "Widgets/Views/STableViewBase.h"
-#include "Widgets/Views/STableRow.h"
-#include "Textures/SlateIcon.h"
-#include "Framework/Commands/UIAction.h"
-#include "Widgets/Views/SListView.h"
 #include "IDetailsViewPrivate.h"
-#include "Layout/WidgetPath.h"
+#include "InputCoreTypes.h"
 #include "Framework/Application/MenuStack.h"
 #include "Framework/Application/SlateApplication.h"
+#include "Framework/Commands/UIAction.h"
 #include "Framework/MultiBox/MultiBoxBuilder.h"
+#include "Input/Reply.h"
+#include "Layout/WidgetPath.h"
+#include "Textures/SlateIcon.h"
+#include "Widgets/Views/SListView.h"
+#include "Widgets/Views/STableViewBase.h"
+#include "Widgets/Views/STableRow.h"
 
 class SDetailTableRowBase : public STableRow< TSharedPtr< FDetailTreeNode > >
 {
diff --git a/Engine/Source/Editor/PropertyEditor/Private/UserInterface/PropertyEditor/PropertyEditorConstants.h b/Engine/Source/Editor/PropertyEditor/Private/UserInterface/PropertyEditor/PropertyEditorConstants.h
--- a/Engine/Source/Editor/PropertyEditor/Private/UserInterface/PropertyEditor/PropertyEditorConstants.h
+++ b/Engine/Source/Editor/PropertyEditor/Private/UserInterface/PropertyEditor/PropertyEditorConstants.h
@@ -1,8 +1,9 @@
 // Copyright Epic Games, Inc. All Rights Reserved.
 #pragma once
 
-#include "CoreMinimal.h"
 #include "Styling/SlateColor.h"
+#include "UObject/NameTypes.h"
+#include "Templates/SharedPointer.h"
 
 struct FSlateBrush;
 
