Moved SDetailCategoryTableRow into its own file instead of hiding in DetailCategoryGroupNode.

diff --git a/Engine/Source/Editor/PropertyEditor/Private/DetailCategoryBuilderImpl.cpp b/Engine/Source/Editor/PropertyEditor/Private/DetailCategoryBuilderImpl.cpp
--- a/Engine/Source/Editor/PropertyEditor/Private/DetailCategoryBuilderImpl.cpp
+++ b/Engine/Source/Editor/PropertyEditor/Private/DetailCategoryBuilderImpl.cpp
@@ -1,29 +1,26 @@
 // Copyright Epic Games, Inc. All Rights Reserved.
 
 #include "DetailCategoryBuilderImpl.h"
-#include "ObjectPropertyNode.h"
-#include "Misc/ConfigCacheIni.h"
 
+#include "DetailAdvancedDropdownNode.h"
+#include "DetailBuilderTypes.h"
 #include "DetailCategoryGroupNode.h"
+#include "DetailGroup.h"
 #include "DetailItemNode.h"
-#include "DetailAdvancedDropdownNode.h"
 #include "DetailPropertyRow.h"
-#include "SDetailSingleItemRow.h"
-#include "DetailGroup.h"
-#include "StructurePropertyNode.h"
-#include "ItemPropertyNode.h"
 #include "IPropertyGenerationUtilities.h"
-#include "DetailBuilderTypes.h"
+#include "ItemPropertyNode.h"
+#include "Misc/ConfigCacheIni.h"
+#include "Modules/ModuleManager.h"
+#include "ObjectPropertyNode.h"
 #include "PropertyCustomizationHelpers.h"
+#include "PropertyEditorModule.h"
 #include "PropertyEditorWhitelist.h"
+#include "SDetailCategoryTableRow.h"
+#include "SDetailSingleItemRow.h"
+#include "StructurePropertyNode.h"
 #include "Styling/StyleColors.h"
-
-namespace DetailLayoutConstants
-{
-	// Padding for each layout row
-	const FMargin RowPadding(10.0f, 2.0f);
-}
-
+ 
 void FDetailLayout::AddLayout(const FDetailLayoutCustomization& Layout)
 {
 	if (Layout.bAdvanced)
diff --git a/Engine/Source/Editor/PropertyEditor/Private/DetailCategoryGroupNode.cpp b/Engine/Source/Editor/PropertyEditor/Private/DetailCategoryGroupNode.cpp
--- a/Engine/Source/Editor/PropertyEditor/Private/DetailCategoryGroupNode.cpp
+++ b/Engine/Source/Editor/PropertyEditor/Private/DetailCategoryGroupNode.cpp
@@ -2,186 +2,7 @@
 
 #include "DetailCategoryGroupNode.h"
 #include "PropertyCustomizationHelpers.h"
-#include "PropertyEditorConstants.h"
-#include "SDetailExpanderArrow.h"
-#include "SDetailRowIndent.h"
-#include "Styling/StyleColors.h"
-
-void SDetailCategoryTableRow::Construct( const FArguments& InArgs, TSharedRef<FDetailTreeNode> InOwnerTreeNode, const TSharedRef<STableViewBase>& InOwnerTableView )
-{
-	OwnerTreeNode = InOwnerTreeNode;
-
-	bIsInnerCategory = InArgs._InnerCategory;
-	bShowBorder = InArgs._ShowBorder;
-
-	const float VerticalPadding = bIsInnerCategory ? 6 : 8;
-
-	FDetailColumnSizeData& ColumnSizeData = InOwnerTreeNode->GetDetailsView()->GetColumnSizeData();
-
-	TSharedRef<SHorizontalBox> HeaderBox = SNew(SHorizontalBox)
-		+ SHorizontalBox::Slot()
-		.HAlign(HAlign_Left)
-		.VAlign(VAlign_Fill)
-		.AutoWidth()
-		[
-			SNew(SDetailRowIndent, SharedThis(this))
-		]
-		+ SHorizontalBox::Slot()
-		.HAlign(HAlign_Left)
-		.VAlign(VAlign_Center)
-		.Padding(5,0,0,0)
-		.AutoWidth()
-		[
-			SNew(SDetailExpanderArrow, SharedThis(this))
-		]
-		+ SHorizontalBox::Slot()
-		.VAlign(VAlign_Center)
-		.Padding(12, VerticalPadding, 0, VerticalPadding)
-		.FillWidth(1)
-		[
-			SNew(STextBlock)
-			.TransformPolicy(ETextTransformPolicy::ToUpper)
-			.Text(InArgs._DisplayName)
-			.Font(FAppStyle::Get().GetFontStyle(bIsInnerCategory ? PropertyEditorConstants::PropertyFontStyle : PropertyEditorConstants::CategoryFontStyle))
-			.TextStyle(FAppStyle::Get(), "DetailsView.CategoryTextStyle")
-		];
-
-	if (InArgs._HeaderContent.IsValid())
-	{
-		HeaderBox->AddSlot()
-		.VAlign(VAlign_Center)
-		.FillWidth(1)
-		[
-			InArgs._HeaderContent.ToSharedRef()
-		];
-	}
-	
-	TWeakPtr<STableViewBase> OwnerTableViewWeak = InOwnerTableView;
-	auto GetScrollbarWellBrush = [this, OwnerTableViewWeak]()
-	{
-		return SDetailTableRowBase::IsScrollBarVisible(OwnerTableViewWeak) ?
-			FAppStyle::Get().GetBrush("DetailsView.GridLine") : 
-			this->GetBackgroundImage();
-	};
-
-	auto GetScrollbarWellTint = [this, OwnerTableViewWeak]()
-	{
-		return SDetailTableRowBase::IsScrollBarVisible(OwnerTableViewWeak) ?
-			FSlateColor(EStyleColor::White) : 
-			this->GetInnerBackgroundColor();
-	};
-
-	this->ChildSlot
-	[
-		SNew( SBorder )
-		.BorderImage(FAppStyle::Get().GetBrush("DetailsView.GridLine"))
-		.Padding(FMargin(0, 0, 0, 1))
-		[
-			SNew(SHorizontalBox)
-			+ SHorizontalBox::Slot()
-			.HAlign(HAlign_Fill)
-			[
-				SNew(SBorder)
-				.BorderImage(this, &SDetailCategoryTableRow::GetBackgroundImage)
-				.BorderBackgroundColor(this, &SDetailCategoryTableRow::GetInnerBackgroundColor)
-				.Padding(0)
-				[
-					HeaderBox
-				]
-			]
-			+ SHorizontalBox::Slot()
-			.HAlign(HAlign_Right)
-			.VAlign(VAlign_Fill)
-			.AutoWidth()
-			[
-				SNew(SBorder)
-				.BorderImage_Lambda(GetScrollbarWellBrush)
-				.BorderBackgroundColor_Lambda(GetScrollbarWellTint)
-				.Padding(FMargin(0, 0, SDetailTableRowBase::ScrollBarPadding, 0))
-			]
-		]
-	];
-
-	STableRow< TSharedPtr< FDetailTreeNode > >::ConstructInternal(
-		STableRow::FArguments()
-			.Style(FEditorStyle::Get(), "DetailsView.TreeView.TableRow")
-			.ShowSelection(false),
-		InOwnerTableView
-	);
-}
-
-EVisibility SDetailCategoryTableRow::IsSeparatorVisible() const
-{
-	return bIsInnerCategory || IsItemExpanded() ? EVisibility::Collapsed : EVisibility::Visible;
-}
-
-const FSlateBrush* SDetailCategoryTableRow::GetBackgroundImage() const
-{
-	if (bShowBorder)
-	{
-		if (bIsInnerCategory)
-		{
-			return FAppStyle::Get().GetBrush("DetailsView.CategoryMiddle");
-		}
-
-		if (IsHovered())
-		{
-			return FAppStyle::Get().GetBrush("DetailsView.CategoryTop_Hovered");
-		}
-		else
-		{
-			return FAppStyle::Get().GetBrush("DetailsView.CategoryTop");
-		}
-	}
-
-	return nullptr;
-}
-
-FSlateColor SDetailCategoryTableRow::GetInnerBackgroundColor() const
-{
-	if (bShowBorder && bIsInnerCategory)
-	{
-		int32 IndentLevel = -1;
-		if (OwnerTablePtr.IsValid())
-		{
-			IndentLevel = GetIndentLevel();
-		}
-
-		IndentLevel = FMath::Max(IndentLevel - 1, 0);
-
-		return PropertyEditorConstants::GetRowBackgroundColor(IndentLevel);
-	}
-
-	return FSlateColor(FLinearColor::White);
-}
-
-FSlateColor SDetailCategoryTableRow::GetOuterBackgroundColor() const
-{
-	if (IsHovered())
-	{
-		return FAppStyle::Get().GetSlateColor("Colors.Header");
-	}
-
-	return FAppStyle::Get().GetSlateColor("Colors.Panel");
-}
-
-FReply SDetailCategoryTableRow::OnMouseButtonDown( const FGeometry& MyGeometry, const FPointerEvent& MouseEvent )
-{
-	if (MouseEvent.GetEffectingButton() == EKeys::LeftMouseButton)
-	{
-		ToggleExpansion();
-		return FReply::Handled();
-	}
-	else
-	{
-		return FReply::Unhandled();
-	}
-}
-
-FReply SDetailCategoryTableRow::OnMouseButtonDoubleClick(const FGeometry& InMyGeometry, const FPointerEvent& InMouseEvent)
-{
-	return OnMouseButtonDown(InMyGeometry, InMouseEvent);
-}
+#include "SDetailCategoryTableRow.h"
 
 FDetailCategoryGroupNode::FDetailCategoryGroupNode(FName InGroupName, TSharedRef<FDetailCategoryImpl> InParentCategory)
 	: ParentCategory(InParentCategory.Get())
@@ -210,7 +31,6 @@ TSharedRef< ITableRow > FDetailCategoryGroupNode::GenerateWidgetForTableView( co
 		.ShowBorder( bShowBorder );
 }
 
-
 bool FDetailCategoryGroupNode::GenerateStandaloneWidget(FDetailWidgetRow& OutRow) const
 {
 	OutRow.NameContent()
diff --git a/Engine/Source/Editor/PropertyEditor/Private/DetailCategoryGroupNode.h b/Engine/Source/Editor/PropertyEditor/Private/DetailCategoryGroupNode.h
--- a/Engine/Source/Editor/PropertyEditor/Private/DetailCategoryGroupNode.h
+++ b/Engine/Source/Editor/PropertyEditor/Private/DetailCategoryGroupNode.h
@@ -2,46 +2,13 @@
 
 #pragma once
 
-#include "CoreMinimal.h"
-#include "Layout/Visibility.h"
-#include "Input/Reply.h"
-#include "IPropertyUtilities.h"
-#include "DetailTreeNode.h"
-#include "Widgets/Views/STableViewBase.h"
-#include "Widgets/DeclarativeSyntaxSupport.h"
-#include "Widgets/Views/STableRow.h"
 #include "DetailCategoryBuilderImpl.h"
-#include "SDetailTableRowBase.h"
+#include "DetailTreeNode.h"
+#include "IDetailsView.h"
+#include "IPropertyUtilities.h"
 #include "PropertyCustomizationHelpers.h"
-
-class SDetailCategoryTableRow : public SDetailTableRowBase
-{
-public:
-	SLATE_BEGIN_ARGS( SDetailCategoryTableRow )
-		: _InnerCategory( false )
-		, _ShowBorder( true )
-	{}
-		SLATE_ARGUMENT( FText, DisplayName )
-		SLATE_ARGUMENT( bool, InnerCategory )
-		SLATE_ARGUMENT( TSharedPtr<SWidget>, HeaderContent )
-		SLATE_ARGUMENT( bool, ShowBorder )
-	SLATE_END_ARGS()
-
-	void Construct( const FArguments& InArgs, TSharedRef<FDetailTreeNode> InOwnerTreeNode, const TSharedRef<STableViewBase>& InOwnerTableView );
-
-private:
-	EVisibility IsSeparatorVisible() const;
-	const FSlateBrush* GetBackgroundImage() const;
-	FSlateColor GetInnerBackgroundColor() const;
-	FSlateColor GetOuterBackgroundColor() const;
-
-	virtual FReply OnMouseButtonDown( const FGeometry& MyGeometry, const FPointerEvent& MouseEvent ) override;
-	virtual FReply OnMouseButtonDoubleClick( const FGeometry& InMyGeometry, const FPointerEvent& InMouseEvent ) override;
-
-private:
-	bool bIsInnerCategory;
-	bool bShowBorder;
-};
+#include "PropertyHandle.h"
+#include "UObject/NameTypes.h"
 
 class FDetailCategoryGroupNode : public FDetailTreeNode, public TSharedFromThis<FDetailCategoryGroupNode>
 {
diff --git a/Engine/Source/Editor/PropertyEditor/Private/DetailItemNode.cpp b/Engine/Source/Editor/PropertyEditor/Private/DetailItemNode.cpp
--- a/Engine/Source/Editor/PropertyEditor/Private/DetailItemNode.cpp
+++ b/Engine/Source/Editor/PropertyEditor/Private/DetailItemNode.cpp
@@ -1,16 +1,17 @@
 // Copyright Epic Games, Inc. All Rights Reserved.
 
 #include "DetailItemNode.h"
+#include "CategoryPropertyNode.h"
 #include "DetailCategoryGroupNode.h"
 #include "DetailGroup.h"
 #include "DetailPropertyRow.h"
-#include "SConstrainedBox.h"
-#include "SDetailSingleItemRow.h"
 #include "IDetailKeyframeHandler.h"
 #include "ObjectPropertyNode.h"
+#include "SConstrainedBox.h"
 #include "PropertyCustomizationHelpers.h"
 #include "PropertyEditorWhitelist.h"
-#include "CategoryPropertyNode.h"
+#include "SDetailCategoryTableRow.h"
+#include "SDetailSingleItemRow.h"
 
 FDetailItemNode::FDetailItemNode(const FDetailLayoutCustomization& InCustomization, TSharedRef<FDetailCategoryImpl> InParentCategory, TAttribute<bool> InIsParentEnabled, TSharedPtr<IDetailGroup> InParentGroup)
 	: Customization( InCustomization )
diff --git a/Engine/Source/Editor/PropertyEditor/Private/SDetailCategoryTableRow.cpp b/Engine/Source/Editor/PropertyEditor/Private/SDetailCategoryTableRow.cpp
new file mode 100644
--- /dev/null
+++ b/Engine/Source/Editor/PropertyEditor/Private/SDetailCategoryTableRow.cpp
@@ -0,0 +1,184 @@
+// Copyright Epic Games, Inc. All Rights Reserved.
+
+#include "SDetailCategoryTableRow.h"
+
+#include "PropertyEditorConstants.h"
+#include "SDetailExpanderArrow.h"
+#include "SDetailRowIndent.h"
+#include "Styling/StyleColors.h"
+
+void SDetailCategoryTableRow::Construct(const FArguments& InArgs, TSharedRef<FDetailTreeNode> InOwnerTreeNode, const TSharedRef<STableViewBase>& InOwnerTableView)
+{
+	OwnerTreeNode = InOwnerTreeNode;
+
+	bIsInnerCategory = InArgs._InnerCategory;
+	bShowBorder = InArgs._ShowBorder;
+
+	const float VerticalPadding = bIsInnerCategory ? 6 : 8;
+
+	FDetailColumnSizeData& ColumnSizeData = InOwnerTreeNode->GetDetailsView()->GetColumnSizeData();
+
+	TSharedRef<SHorizontalBox> HeaderBox = SNew(SHorizontalBox)
+		+ SHorizontalBox::Slot()
+		.HAlign(HAlign_Left)
+		.VAlign(VAlign_Fill)
+		.AutoWidth()
+		[
+			SNew(SDetailRowIndent, SharedThis(this))
+		]
+		+ SHorizontalBox::Slot()
+		.HAlign(HAlign_Left)
+		.VAlign(VAlign_Center)
+		.Padding(5, 0, 0, 0)
+		.AutoWidth()
+		[
+			SNew(SDetailExpanderArrow, SharedThis(this))
+		]
+		+ SHorizontalBox::Slot()
+		.VAlign(VAlign_Center)
+		.Padding(12, VerticalPadding, 0, VerticalPadding)
+		.FillWidth(1)
+		[
+			SNew(STextBlock)
+			.TransformPolicy(ETextTransformPolicy::ToUpper)
+			.Text(InArgs._DisplayName)
+			.Font(FAppStyle::Get().GetFontStyle(bIsInnerCategory ? PropertyEditorConstants::PropertyFontStyle : PropertyEditorConstants::CategoryFontStyle))
+			.TextStyle(FAppStyle::Get(), "DetailsView.CategoryTextStyle")
+		];
+
+	if (InArgs._HeaderContent.IsValid())
+	{
+		HeaderBox->AddSlot()
+			.VAlign(VAlign_Center)
+			.FillWidth(1)
+			[
+				InArgs._HeaderContent.ToSharedRef()
+			];
+	}
+
+	TWeakPtr<STableViewBase> OwnerTableViewWeak = InOwnerTableView;
+	auto GetScrollbarWellBrush = [this, OwnerTableViewWeak]()
+	{
+		return SDetailTableRowBase::IsScrollBarVisible(OwnerTableViewWeak) ?
+			FAppStyle::Get().GetBrush("DetailsView.GridLine") :
+			this->GetBackgroundImage();
+	};
+
+	auto GetScrollbarWellTint = [this, OwnerTableViewWeak]()
+	{
+		return SDetailTableRowBase::IsScrollBarVisible(OwnerTableViewWeak) ?
+			FSlateColor(EStyleColor::White) :
+			this->GetInnerBackgroundColor();
+	};
+
+	this->ChildSlot
+	[
+		SNew(SBorder)
+		.BorderImage(FAppStyle::Get().GetBrush("DetailsView.GridLine"))
+		.Padding(FMargin(0, 0, 0, 1))
+		[
+			SNew(SHorizontalBox)
+			+ SHorizontalBox::Slot()
+			.HAlign(HAlign_Fill)
+			[
+				SNew(SBorder)
+				.BorderImage(this, &SDetailCategoryTableRow::GetBackgroundImage)
+				.BorderBackgroundColor(this, &SDetailCategoryTableRow::GetInnerBackgroundColor)
+				.Padding(0)
+				[
+					HeaderBox
+				]
+			]
+			+ SHorizontalBox::Slot()
+			.HAlign(HAlign_Right)
+			.VAlign(VAlign_Fill)
+			.AutoWidth()
+			[
+				SNew(SBorder)
+				.BorderImage_Lambda(GetScrollbarWellBrush)
+				.BorderBackgroundColor_Lambda(GetScrollbarWellTint)
+				.Padding(FMargin(0, 0, SDetailTableRowBase::ScrollBarPadding, 0))
+			]
+		]
+	];
+
+	STableRow< TSharedPtr< FDetailTreeNode > >::ConstructInternal(
+		STableRow::FArguments()
+		.Style(FEditorStyle::Get(), "DetailsView.TreeView.TableRow")
+		.ShowSelection(false),
+		InOwnerTableView
+	);
+}
+
+EVisibility SDetailCategoryTableRow::IsSeparatorVisible() const
+{
+	return bIsInnerCategory || IsItemExpanded() ? EVisibility::Collapsed : EVisibility::Visible;
+}
+
+const FSlateBrush* SDetailCategoryTableRow::GetBackgroundImage() const
+{
+	if (bShowBorder)
+	{
+		if (bIsInnerCategory)
+		{
+			return FAppStyle::Get().GetBrush("DetailsView.CategoryMiddle");
+		}
+
+		if (IsHovered())
+		{
+			return FAppStyle::Get().GetBrush("DetailsView.CategoryTop_Hovered");
+		}
+		else
+		{
+			return FAppStyle::Get().GetBrush("DetailsView.CategoryTop");
+		}
+	}
+
+	return nullptr;
+}
+
+FSlateColor SDetailCategoryTableRow::GetInnerBackgroundColor() const
+{
+	if (bShowBorder && bIsInnerCategory)
+	{
+		int32 IndentLevel = -1;
+		if (OwnerTablePtr.IsValid())
+		{
+			IndentLevel = GetIndentLevel();
+		}
+
+		IndentLevel = FMath::Max(IndentLevel - 1, 0);
+
+		return PropertyEditorConstants::GetRowBackgroundColor(IndentLevel);
+	}
+
+	return FSlateColor(FLinearColor::White);
+}
+
+FSlateColor SDetailCategoryTableRow::GetOuterBackgroundColor() const
+{
+	if (IsHovered())
+	{
+		return FAppStyle::Get().GetSlateColor("Colors.Header");
+	}
+
+	return FAppStyle::Get().GetSlateColor("Colors.Panel");
+}
+
+FReply SDetailCategoryTableRow::OnMouseButtonDown(const FGeometry& MyGeometry, const FPointerEvent& MouseEvent)
+{
+	if (MouseEvent.GetEffectingButton() == EKeys::LeftMouseButton)
+	{
+		ToggleExpansion();
+		return FReply::Handled();
+	}
+	else
+	{
+		return FReply::Unhandled();
+	}
+}
+
+FReply SDetailCategoryTableRow::OnMouseButtonDoubleClick(const FGeometry& InMyGeometry, const FPointerEvent& InMouseEvent)
+{
+	return OnMouseButtonDown(InMyGeometry, InMouseEvent);
+}
diff --git a/Engine/Source/Editor/PropertyEditor/Private/SDetailCategoryTableRow.h b/Engine/Source/Editor/PropertyEditor/Private/SDetailCategoryTableRow.h
new file mode 100644
--- /dev/null
+++ b/Engine/Source/Editor/PropertyEditor/Private/SDetailCategoryTableRow.h
@@ -0,0 +1,41 @@
+// Copyright Epic Games, Inc. All Rights Reserved.
+
+#pragma once
+
+#include "DetailTreeNode.h"
+#include "Input/Reply.h"
+#include "Layout/Visibility.h"
+#include "SDetailTableRowBase.h"
+#include "Templates/SharedPointer.h"
+#include "Widgets/DeclarativeSyntaxSupport.h"
+#include "Widgets/Views/STableRow.h"
+#include "Widgets/Views/STableViewBase.h"
+
+class SDetailCategoryTableRow : public SDetailTableRowBase
+{
+public:
+	SLATE_BEGIN_ARGS(SDetailCategoryTableRow)
+		: _InnerCategory(false)
+		, _ShowBorder(true)
+	{}
+		SLATE_ARGUMENT(FText, DisplayName)
+		SLATE_ARGUMENT(bool, InnerCategory)
+		SLATE_ARGUMENT(TSharedPtr<SWidget>, HeaderContent)
+		SLATE_ARGUMENT(bool, ShowBorder)
+	SLATE_END_ARGS()
+
+	void Construct(const FArguments& InArgs, TSharedRef<FDetailTreeNode> InOwnerTreeNode, const TSharedRef<STableViewBase>& InOwnerTableView);
+
+private:
+	EVisibility IsSeparatorVisible() const;
+	const FSlateBrush* GetBackgroundImage() const;
+	FSlateColor GetInnerBackgroundColor() const;
+	FSlateColor GetOuterBackgroundColor() const;
+
+	FReply OnMouseButtonDown(const FGeometry& MyGeometry, const FPointerEvent& MouseEvent) override;
+	FReply OnMouseButtonDoubleClick(const FGeometry& InMyGeometry, const FPointerEvent& InMouseEvent) override;
+
+private:
+	bool bIsInnerCategory;
+	bool bShowBorder;
+};
