FEditConditionParser::ConsumeBool is now case-insensitive since UHT turns an unquoted true/false into TRUE/FALSE.

diff --git a/Engine/Source/Editor/PropertyEditor/Private/EditConditionParser.cpp b/Engine/Source/Editor/PropertyEditor/Private/EditConditionParser.cpp
--- a/Engine/Source/Editor/PropertyEditor/Private/EditConditionParser.cpp
+++ b/Engine/Source/Editor/PropertyEditor/Private/EditConditionParser.cpp
@@ -31,13 +31,13 @@ static const TCHAR PropertyBreakingChars[] = { '|', '=', '&', '>', '<', '!', '+'
 
 static TOptional<FExpressionError> ConsumeBool(FExpressionTokenConsumer& Consumer)
 {
-	TOptional<FStringToken> TrueToken = Consumer.GetStream().ParseToken(TEXT("true"));
+	TOptional<FStringToken> TrueToken = Consumer.GetStream().ParseTokenIgnoreCase(TEXT("true"));
 	if (TrueToken.IsSet())
 	{
 		Consumer.Add(TrueToken.GetValue(), true);
 	}
 
-	TOptional<FStringToken> FalseToken = Consumer.GetStream().ParseToken(TEXT("false"));
+	TOptional<FStringToken> FalseToken = Consumer.GetStream().ParseTokenIgnoreCase(TEXT("false"));
 	if (FalseToken.IsSet())
 	{
 		Consumer.Add(FalseToken.GetValue(), false);
