Added missing null checks in GetEnumValue that were causing FEditConditionContext to dereference null when deleting a struct from an array.

Moved WeakPtr null check first in all Get*Value, as it's the cheapest one to perform.

Added const to several FComplexPropertyNode functions that were const in all but practice, and to their callsites in FEditConditionContext.

AlreadyLogged was storing just a null pointer, and thus only logging the first field that was found to be missing per session.

diff --git a/Engine/Source/Editor/PropertyEditor/Private/EditConditionContext.cpp b/Engine/Source/Editor/PropertyEditor/Private/EditConditionContext.cpp
--- a/Engine/Source/Editor/PropertyEditor/Private/EditConditionContext.cpp
+++ b/Engine/Source/Editor/PropertyEditor/Private/EditConditionContext.cpp
@@ -29,14 +29,11 @@ const FBoolProperty* FEditConditionContext::GetSingleBoolProperty(const TSharedP
 		{
 			if (BoolProperty != nullptr)
 			{
-				// second property token
+				// second property token in the same expression, this can't be a simple expression like "bValue == false"
 				return nullptr;
 			}
 
-			const FProperty* Field = FindFProperty<FProperty>(Property->GetOwnerStruct(), *PropertyToken->PropertyName);
-			BoolProperty = CastField<FBoolProperty>(Field);
-			
-			// not a bool
+			BoolProperty = FindFProperty<FBoolProperty>(Property->GetOwnerStruct(), *PropertyToken->PropertyName);
 			if (BoolProperty == nullptr)
 			{
 				return nullptr;
@@ -47,32 +44,36 @@ const FBoolProperty* FEditConditionContext::GetSingleBoolProperty(const TSharedP
 	return BoolProperty;
 }
 
-static TSet<const FProperty*> AlreadyLogged;
+static TSet<TPair<FName, FString>> AlreadyLogged;
 
 template<typename T>
-T* FindTypedField(const TWeakPtr<FPropertyNode>& PropertyNode, const FString& PropertyName)
+const T* FindTypedField(const TSharedPtr<FPropertyNode>& PropertyNode, const FString& FieldName)
 {
-	if (PropertyNode.IsValid())
+	if (!PropertyNode.IsValid())
 	{
-		TSharedPtr<FPropertyNode> PinnedNode = PropertyNode.Pin();
-		const FProperty* Property = PinnedNode->GetProperty();
+		return nullptr;
+	}
 
-		FProperty* Field = FindFProperty<FProperty>(Property->GetOwnerStruct(), *PropertyName);
-		if (Field == nullptr)
-		{
-			if (!AlreadyLogged.Find(Field))
-			{
-				AlreadyLogged.Add(Field);
-				UE_LOG(LogEditCondition, Error, TEXT("EditCondition parsing failed: Field name \"%s\" was not found in class \"%s\"."), *PropertyName, *Property->GetOwnerStruct()->GetName());
-			}
+	const FProperty* Property = PropertyNode->GetProperty();
+	if (Property == nullptr)
+	{
+		return nullptr;
+	}
 
-			return nullptr;
+	const FProperty* Field = FindFProperty<FProperty>(Property->GetOwnerStruct(), *FieldName);
+	if (Field == nullptr)
+	{
+		TPair<FName, FString> FieldKey(Property->GetOwnerStruct()->GetFName(), FieldName);
+		if (!AlreadyLogged.Find(FieldKey))
+		{
+			AlreadyLogged.Add(FieldKey);
+			UE_LOG(LogEditCondition, Error, TEXT("EditCondition parsing failed: Field name \"%s\" was not found in class \"%s\"."), *FieldName, *Property->GetOwnerStruct()->GetName());
 		}
 
-		return CastField<T>(Field);
+		return nullptr;
 	}
 
-	return nullptr;
+	return CastField<T>(Field);
 }
 
 /** 
@@ -82,9 +83,9 @@ T* FindTypedField(const TWeakPtr<FPropertyNode>& PropertyNode, const FString& Pr
  * Note: We do not support nested containers.
  * The result can be nullptr in exceptional cases, eg. if the UI is getting rebuilt.
  */
-static FPropertyNode* GetEditConditionParentNode(const TSharedPtr<FPropertyNode>& PropertyNode)
+static const FPropertyNode* GetEditConditionParentNode(const TSharedPtr<FPropertyNode>& PropertyNode)
 {
-	FPropertyNode* ParentNode = PropertyNode->GetParentNode();
+	const FPropertyNode* ParentNode = PropertyNode->GetParentNode();
 	FFieldVariant PropertyOuter = PropertyNode->GetProperty()->GetOwnerVariant();
 
 	if (PropertyOuter.Get<FArrayProperty>() != nullptr ||
@@ -104,33 +105,33 @@ static FPropertyNode* GetEditConditionParentNode(const TSharedPtr<FPropertyNode>
 	return ParentNode;
 }
 
-static uint8* GetPropertyValuePtr(const FProperty* Property, const TSharedPtr<FPropertyNode>& PropertyNode, FPropertyNode* ParentNode, FComplexPropertyNode* ComplexParentNode, int32 Index)
+static const uint8* GetPropertyValuePtr(const FProperty* Property, const TSharedPtr<FPropertyNode>& PropertyNode, const FPropertyNode* ParentNode, const FComplexPropertyNode* ComplexParentNode, int32 Index)
 {
-	uint8* ValuePtr = ComplexParentNode->GetValuePtrOfInstance(Index, Property, ParentNode);
+	const uint8* ValuePtr = ComplexParentNode->GetValuePtrOfInstance(Index, Property, ParentNode);
 	return ValuePtr;
 }
 
 TOptional<bool> FEditConditionContext::GetBoolValue(const FString& PropertyName) const
 {
-	const FBoolProperty* BoolProperty = FindTypedField<FBoolProperty>(PropertyNode, PropertyName);
-	if (BoolProperty == nullptr)
+	TSharedPtr<FPropertyNode> PinnedNode = PropertyNode.Pin();
+	if (!PinnedNode.IsValid())
 	{
 		return TOptional<bool>();
 	}
 
-	TSharedPtr<FPropertyNode> PinnedNode = PropertyNode.Pin();
-	if (!PinnedNode.IsValid())
+	const FBoolProperty* BoolProperty = FindTypedField<FBoolProperty>(PinnedNode, PropertyName);
+	if (BoolProperty == nullptr)
 	{
 		return TOptional<bool>();
 	}
 
-	FPropertyNode* ParentNode = GetEditConditionParentNode(PinnedNode);
+	const FPropertyNode* ParentNode = GetEditConditionParentNode(PinnedNode);
 	if (ParentNode == nullptr)
 	{
 		return TOptional<bool>();
 	}
 
-	FComplexPropertyNode* ComplexParentNode = PinnedNode->FindComplexParent();
+	const FComplexPropertyNode* ComplexParentNode = PinnedNode->FindComplexParent();
 	if (ComplexParentNode == nullptr)
 	{
 		return TOptional<bool>();
@@ -139,7 +140,7 @@ TOptional<bool> FEditConditionContext::GetBoolValue(const FString& PropertyName)
 	TOptional<bool> Result;
 	for (int32 Index = 0; Index < ComplexParentNode->GetInstancesNum(); ++Index)
 	{
-		uint8* ValuePtr = GetPropertyValuePtr(BoolProperty, PinnedNode, ParentNode, ComplexParentNode, Index);
+		const uint8* ValuePtr = GetPropertyValuePtr(BoolProperty, PinnedNode, ParentNode, ComplexParentNode, Index);
 		if (ValuePtr == nullptr)
 		{
 			return TOptional<bool>();
@@ -162,30 +163,25 @@ TOptional<bool> FEditConditionContext::GetBoolValue(const FString& PropertyName)
 
 TOptional<int64> FEditConditionContext::GetIntegerValue(const FString& PropertyName) const
 {
-	const FNumericProperty* NumericProperty = FindTypedField<FNumericProperty>(PropertyNode, PropertyName);
-	if (NumericProperty == nullptr)
+	TSharedPtr<FPropertyNode> PinnedNode = PropertyNode.Pin();
+	if (!PinnedNode.IsValid())
 	{
 		return TOptional<int64>();
 	}
 
-	if (!NumericProperty->IsInteger())
+	const FNumericProperty* NumericProperty = FindTypedField<FNumericProperty>(PinnedNode, PropertyName);
+	if (NumericProperty == nullptr || !NumericProperty->IsInteger())
 	{
 		return TOptional<int64>();
 	}
 
-	TSharedPtr<FPropertyNode> PinnedNode = PropertyNode.Pin();
-	if (!PinnedNode.IsValid())
-	{
-		return TOptional<int64>();
-	}
-
-	FPropertyNode* ParentNode = GetEditConditionParentNode(PinnedNode);
+	const FPropertyNode* ParentNode = GetEditConditionParentNode(PinnedNode);
 	if (ParentNode == nullptr)
 	{
 		return TOptional<int64>();
 	}
 
-	FComplexPropertyNode* ComplexParentNode = PinnedNode->FindComplexParent();
+	const FComplexPropertyNode* ComplexParentNode = PinnedNode->FindComplexParent();
 	if (ComplexParentNode == nullptr)
 	{
 		return TOptional<int64>();
@@ -194,7 +190,7 @@ TOptional<int64> FEditConditionContext::GetIntegerValue(const FString& PropertyN
 	TOptional<int64> Result;
 	for (int32 Index = 0; Index < ComplexParentNode->GetInstancesNum(); ++Index)
 	{
-		uint8* ValuePtr = GetPropertyValuePtr(NumericProperty, PinnedNode, ParentNode, ComplexParentNode, Index);
+		const uint8* ValuePtr = GetPropertyValuePtr(NumericProperty, PinnedNode, ParentNode, ComplexParentNode, Index);
 		if (ValuePtr == nullptr)
 		{
 			return TOptional<int64>();
@@ -217,25 +213,25 @@ TOptional<int64> FEditConditionContext::GetIntegerValue(const FString& PropertyN
 
 TOptional<double> FEditConditionContext::GetNumericValue(const FString& PropertyName) const
 {
-	const FNumericProperty* NumericProperty = FindTypedField<FNumericProperty>(PropertyNode, PropertyName);
-	if (NumericProperty == nullptr)
+	TSharedPtr<FPropertyNode> PinnedNode = PropertyNode.Pin();
+	if (!PinnedNode.IsValid())
 	{
 		return TOptional<double>();
 	}
 
-	TSharedPtr<FPropertyNode> PinnedNode = PropertyNode.Pin();
-	if (!PinnedNode.IsValid())
+	const FNumericProperty* NumericProperty = FindTypedField<FNumericProperty>(PinnedNode, PropertyName);
+	if (NumericProperty == nullptr)
 	{
 		return TOptional<double>();
 	}
 
-	FPropertyNode* ParentNode = GetEditConditionParentNode(PinnedNode);
+	const FPropertyNode* ParentNode = GetEditConditionParentNode(PinnedNode);
 	if (ParentNode == nullptr)
 	{
 		return TOptional<double>();
 	}
 
-	FComplexPropertyNode* ComplexParentNode = PinnedNode->FindComplexParent();
+	const FComplexPropertyNode* ComplexParentNode = PinnedNode->FindComplexParent();
 	if (ComplexParentNode == nullptr)
 	{
 		return TOptional<double>();
@@ -244,7 +240,7 @@ TOptional<double> FEditConditionContext::GetNumericValue(const FString& Property
 	TOptional<double> Result;
 	for (int32 Index = 0; Index < ComplexParentNode->GetInstancesNum(); ++Index)
 	{
-		uint8* ValuePtr = GetPropertyValuePtr(NumericProperty, PinnedNode, ParentNode, ComplexParentNode, Index);
+		const uint8* ValuePtr = GetPropertyValuePtr(NumericProperty, PinnedNode, ParentNode, ComplexParentNode, Index);
 		if (ValuePtr == nullptr)
 		{
 			return TOptional<double>();
@@ -277,7 +273,13 @@ TOptional<double> FEditConditionContext::GetNumericValue(const FString& Property
 
 TOptional<FString> FEditConditionContext::GetEnumValue(const FString& PropertyName) const
 {
-	const FProperty* Property = FindTypedField<FProperty>(PropertyNode, PropertyName);
+	TSharedPtr<FPropertyNode> PinnedNode = PropertyNode.Pin();
+	if (!PinnedNode.IsValid())
+	{
+		return TOptional<FString>();
+	}
+
+	const FProperty* Property = FindTypedField<FProperty>(PinnedNode, PropertyName);
 	if (Property == nullptr)
 	{
 		return TOptional<FString>();
@@ -295,29 +297,24 @@ TOptional<FString> FEditConditionContext::GetEnumValue(const FString& PropertyNa
 		NumericProperty = ByteProperty;
 		EnumType = ByteProperty->GetIntPropertyEnum();
 	}
-	else
-	{
-		return TOptional<FString>();
-	}
 
-	if (!NumericProperty->IsInteger())
+	if (EnumType == nullptr)
 	{
 		return TOptional<FString>();
 	}
-
-	TSharedPtr<FPropertyNode> PinnedNode = PropertyNode.Pin();
-	if (!PinnedNode.IsValid())
+	
+	if (NumericProperty == nullptr || !NumericProperty->IsInteger())
 	{
 		return TOptional<FString>();
 	}
 
-	FPropertyNode* ParentNode = GetEditConditionParentNode(PinnedNode);
+	const FPropertyNode* ParentNode = GetEditConditionParentNode(PinnedNode);
 	if (ParentNode == nullptr)
 	{
 		return TOptional<FString>();
 	}
 
-	FComplexPropertyNode* ComplexParentNode = PinnedNode->FindComplexParent();
+	const FComplexPropertyNode* ComplexParentNode = PinnedNode->FindComplexParent();
 	if (ComplexParentNode == nullptr)
 	{
 		return TOptional<FString>();
@@ -326,7 +323,7 @@ TOptional<FString> FEditConditionContext::GetEnumValue(const FString& PropertyNa
 	TOptional<int64> Result;
 	for (int32 Index = 0; Index < ComplexParentNode->GetInstancesNum(); ++Index)
 	{
-		uint8* ValuePtr = GetPropertyValuePtr(Property, PinnedNode, ParentNode, ComplexParentNode, Index);
+		const uint8* ValuePtr = GetPropertyValuePtr(NumericProperty, PinnedNode, ParentNode, ComplexParentNode, Index);
 		if (ValuePtr == nullptr)
 		{
 			return TOptional<FString>();
@@ -354,31 +351,31 @@ TOptional<FString> FEditConditionContext::GetEnumValue(const FString& PropertyNa
 
 TOptional<UObject*> FEditConditionContext::GetPointerValue(const FString& PropertyName) const
 {
-	const FProperty* Property = FindTypedField<FProperty>(PropertyNode, PropertyName);
-	if (Property == nullptr)
+	TSharedPtr<FPropertyNode> PinnedNode = PropertyNode.Pin();
+	if (!PinnedNode.IsValid())
 	{
 		return TOptional<UObject*>();
 	}
 
-	const FObjectPropertyBase* ObjectProperty = CastField<FObjectPropertyBase>(Property);
-	if (ObjectProperty == nullptr)
+	const FProperty* Property = FindTypedField<FProperty>(PinnedNode, PropertyName);
+	if (Property == nullptr)
 	{
 		return TOptional<UObject*>();
 	}
 
-	TSharedPtr<FPropertyNode> PinnedNode = PropertyNode.Pin();
-	if (!PinnedNode.IsValid())
+	const FObjectPropertyBase* ObjectProperty = CastField<FObjectPropertyBase>(Property);
+	if (ObjectProperty == nullptr)
 	{
 		return TOptional<UObject*>();
 	}
 
-	FPropertyNode* ParentNode = GetEditConditionParentNode(PinnedNode);
+	const FPropertyNode* ParentNode = GetEditConditionParentNode(PinnedNode);
 	if (ParentNode == nullptr)
 	{
 		return TOptional<UObject*>();
 	}
 
-	FComplexPropertyNode* ComplexParentNode = PinnedNode->FindComplexParent();
+	const FComplexPropertyNode* ComplexParentNode = PinnedNode->FindComplexParent();
 	if (ComplexParentNode == nullptr)
 	{
 		return TOptional<UObject*>();
@@ -387,7 +384,7 @@ TOptional<UObject*> FEditConditionContext::GetPointerValue(const FString& Proper
 	TOptional<UObject*> Result;
 	for (int32 Index = 0; Index < ComplexParentNode->GetInstancesNum(); ++Index)
 	{
-		uint8* ValuePtr = GetPropertyValuePtr(Property, PinnedNode, ParentNode, ComplexParentNode, Index);
+		const uint8* ValuePtr = GetPropertyValuePtr(Property, PinnedNode, ParentNode, ComplexParentNode, Index);
 		if (ValuePtr == nullptr)
 		{
 			return TOptional<UObject*>();
@@ -410,7 +407,13 @@ TOptional<UObject*> FEditConditionContext::GetPointerValue(const FString& Proper
 
 TOptional<FString> FEditConditionContext::GetTypeName(const FString& PropertyName) const
 {
-	const FProperty* Property = FindTypedField<FProperty>(PropertyNode, PropertyName);
+	TSharedPtr<FPropertyNode> PinnedNode = PropertyNode.Pin();
+	if (!PinnedNode.IsValid())
+	{
+		return TOptional<FString>();
+	}
+
+	const FProperty* Property = FindTypedField<FProperty>(PinnedNode, PropertyName);
 	if (Property == nullptr)
 	{
 		return TOptional<FString>();
@@ -418,11 +421,20 @@ TOptional<FString> FEditConditionContext::GetTypeName(const FString& PropertyNam
 
 	if (const FEnumProperty* EnumProperty = CastField<FEnumProperty>(Property))
 	{
-		return EnumProperty->GetEnum()->GetName();
+		const UEnum* EnumType = EnumProperty->GetEnum();
+		if (EnumType == nullptr)
+		{
+			return TOptional<FString>();
+		}
+		return EnumType->GetName();
 	}
 	else if (const FByteProperty* ByteProperty = CastField<FByteProperty>(Property))
 	{
-		return ByteProperty->GetIntPropertyEnum()->GetName();
+		const UEnum* EnumType = ByteProperty->GetIntPropertyEnum();
+		if (EnumType != nullptr)
+		{
+			return EnumType->GetName();
+		}
 	}
 
 	return Property->GetCPPType();
@@ -430,7 +442,7 @@ TOptional<FString> FEditConditionContext::GetTypeName(const FString& PropertyNam
 
 TOptional<int64> FEditConditionContext::GetIntegerValueOfEnum(const FString& EnumTypeName, const FString& MemberName) const
 {
-	UEnum* EnumType = FindObject<UEnum>((UObject*) ANY_PACKAGE, *EnumTypeName, true);
+	const UEnum* EnumType = FindObject<UEnum>((UObject*) ANY_PACKAGE, *EnumTypeName, true);
 	if (EnumType == nullptr)
 	{
 		return TOptional<int64>();
diff --git a/Engine/Source/Editor/PropertyEditor/Private/ObjectPropertyNode.h b/Engine/Source/Editor/PropertyEditor/Private/ObjectPropertyNode.h
--- a/Engine/Source/Editor/PropertyEditor/Private/ObjectPropertyNode.h
+++ b/Engine/Source/Editor/PropertyEditor/Private/ObjectPropertyNode.h
@@ -88,11 +88,11 @@ public:
 	virtual TArray<const UStruct*> GetAllStructures() const override;
 
 	virtual int32 GetInstancesNum() const override{ return GetNumObjects(); }
-	virtual uint8* GetMemoryOfInstance(int32 Index) override
+	virtual uint8* GetMemoryOfInstance(int32 Index) const override
 	{
 		return (uint8*)GetUObject(Index);
 	}
-	virtual uint8* GetValuePtrOfInstance(int32 Index, const FProperty* InProperty, FPropertyNode* InParentNode) override
+	virtual uint8* GetValuePtrOfInstance(int32 Index, const FProperty* InProperty, const FPropertyNode* InParentNode) const override
 	{
 		if (InParentNode == nullptr || InProperty == nullptr)
 		{
@@ -113,7 +113,7 @@ public:
 
 		return InProperty->ContainerPtrToValuePtr<uint8>(ParentPtr);
 	}
-	virtual TWeakObjectPtr<UObject> GetInstanceAsUObject(int32 Index) override
+	virtual TWeakObjectPtr<UObject> GetInstanceAsUObject(int32 Index) const override
 	{
 		check(Objects.IsValidIndex(Index));
 		return Objects[Index];
diff --git a/Engine/Source/Editor/PropertyEditor/Private/PropertyNode.h b/Engine/Source/Editor/PropertyEditor/Private/PropertyNode.h
--- a/Engine/Source/Editor/PropertyEditor/Private/PropertyNode.h
+++ b/Engine/Source/Editor/PropertyEditor/Private/PropertyNode.h
@@ -1150,8 +1150,8 @@ public:
 	FComplexPropertyNode() : FPropertyNode() {}
 	virtual ~FComplexPropertyNode() {}
 
-	virtual FComplexPropertyNode* AsComplexNode() override{ return this; }
-	virtual const FComplexPropertyNode* AsComplexNode() const override{ return this; }
+	virtual FComplexPropertyNode* AsComplexNode() override { return this; }
+	virtual const FComplexPropertyNode* AsComplexNode() const override { return this; }
 
 	virtual class FStructurePropertyNode* AsStructureNode() { return nullptr; }
 	virtual const FStructurePropertyNode* AsStructureNode() const { return nullptr; }
@@ -1164,13 +1164,13 @@ public:
 	virtual TArray<const UStruct*> GetAllStructures() const = 0;
 
 	virtual int32 GetInstancesNum() const = 0;
-	virtual uint8* GetMemoryOfInstance(int32 Index) = 0;
+	virtual uint8* GetMemoryOfInstance(int32 Index) const = 0;
 
 	/**
 	 * Returns a pointer to the stored value of InProperty on InParentNode's Index'th instance.
 	 */
-	virtual uint8* GetValuePtrOfInstance(int32 Index, const FProperty* InProperty, FPropertyNode* InParentNode) = 0;
-	virtual TWeakObjectPtr<UObject> GetInstanceAsUObject(int32 Index) = 0;
+	virtual uint8* GetValuePtrOfInstance(int32 Index, const FProperty* InProperty, const FPropertyNode* InParentNode) const = 0;
+	virtual TWeakObjectPtr<UObject> GetInstanceAsUObject(int32 Index) const = 0;
 	virtual EPropertyType GetPropertyType() const = 0;
 
 	virtual void Disconnect() = 0;
diff --git a/Engine/Source/Editor/PropertyEditor/Private/StructurePropertyNode.h b/Engine/Source/Editor/PropertyEditor/Private/StructurePropertyNode.h
--- a/Engine/Source/Editor/PropertyEditor/Private/StructurePropertyNode.h
+++ b/Engine/Source/Editor/PropertyEditor/Private/StructurePropertyNode.h
@@ -117,19 +117,19 @@ public:
 	{ 
 		return HasValidStructData() ? 1 : 0;
 	}
-	virtual uint8* GetMemoryOfInstance(int32 Index) override
+	virtual uint8* GetMemoryOfInstance(int32 Index) const override
 	{ 
 		check(0 == Index);
 		return HasValidStructData() ? StructData->GetStructMemory() : NULL;
 	}
-	virtual uint8* GetValuePtrOfInstance(int32 Index, const FProperty* InProperty, FPropertyNode* InParentNode) override
+	virtual uint8* GetValuePtrOfInstance(int32 Index, const FProperty* InProperty, const FPropertyNode* InParentNode) const override
 	{ 
 		check(0 == Index);
 		uint8* StructBaseAddress = HasValidStructData() ? StructData->GetStructMemory() : nullptr;
 		uint8* ParentBaseAddress = InParentNode ? InParentNode->GetValueAddress(StructBaseAddress, false) : nullptr;
 		return InProperty ? InProperty->ContainerPtrToValuePtr<uint8>(ParentBaseAddress) : nullptr;
 	}
-	virtual TWeakObjectPtr<UObject> GetInstanceAsUObject(int32 Index) override
+	virtual TWeakObjectPtr<UObject> GetInstanceAsUObject(int32 Index) const override
 	{
 		check(0 == Index);
 		return NULL;
