Added section selector to the details view. It is now possible to register a category to a section, which will cause the entire category to be shown if the relevant section is selected.

Moved FDetailFilter to its own file.

diff --git a/Engine/Source/Editor/DetailCustomizations/Public/DetailCustomizations.h b/Engine/Source/Editor/DetailCustomizations/Public/DetailCustomizations.h
--- a/Engine/Source/Editor/DetailCustomizations/Public/DetailCustomizations.h
+++ b/Engine/Source/Editor/DetailCustomizations/Public/DetailCustomizations.h
@@ -32,6 +32,7 @@ public:
 private:
 	void RegisterPropertyTypeCustomizations();
 	void RegisterObjectCustomizations();
+	void RegisterSectionMappings();
 
 	/**
 	 * Registers a custom class
diff --git a/Engine/Source/Editor/LevelEditor/Private/SActorDetails.cpp b/Engine/Source/Editor/LevelEditor/Private/SActorDetails.cpp
--- a/Engine/Source/Editor/LevelEditor/Private/SActorDetails.cpp
+++ b/Engine/Source/Editor/LevelEditor/Private/SActorDetails.cpp
@@ -1,43 +1,44 @@
 // Copyright Epic Games, Inc. All Rights Reserved.
 
 #include "SActorDetails.h"
-#include "Widgets/SBoxPanel.h"
+
 #include "Components/ActorComponent.h"
 #include "Components/SceneComponent.h"
-#include "GameFramework/Actor.h"
-#include "Engine/Blueprint.h"
+#include "DetailsViewObjectFilter.h"
 #include "Editor.h"
-#include "HAL/FileManager.h"
-#include "Modules/ModuleManager.h"
-#include "Widgets/Layout/SBorder.h"
-#include "Widgets/Images/SImage.h"
-#include "Widgets/Text/SRichTextBlock.h"
-#include "Widgets/Layout/SBox.h"
-#include "Widgets/Layout/SSplitter.h"
-#include "EditorStyleSet.h"
 #include "Editor/UnrealEdEngine.h"
+#include "EditorStyleSet.h"
+#include "Elements/Framework/EngineElementsLibrary.h"
+#include "Elements/Framework/TypedElementRegistry.h"
+#include "Elements/Framework/TypedElementSelectionSet.h"
+#include "Engine/Blueprint.h"
 #include "Engine/BlueprintGeneratedClass.h"
-#include "Kismet2/ComponentEditorUtils.h"
 #include "Engine/Selection.h"
-#include "UnrealEdGlobals.h"
-#include "LevelEditor.h"
-#include "Kismet2/KismetEditorUtilities.h"
-#include "PropertyEditorModule.h"
+#include "GameFramework/Actor.h"
+#include "HAL/FileManager.h"
+#include "IDetailRootObjectCustomization.h"
 #include "IDetailsView.h"
+#include "Kismet2/ComponentEditorUtils.h"
+#include "Kismet2/KismetEditorUtilities.h"
+#include "LevelEditor.h"
 #include "LevelEditorGenericDetails.h"
-#include "ScopedTransaction.h"
-#include "SourceCodeNavigation.h"
-#include "Widgets/Docking/SDockTab.h"
-#include "DetailsViewObjectFilter.h"
-#include "IDetailRootObjectCustomization.h"
-#include "Elements/Framework/TypedElementRegistry.h"
-#include "Elements/Framework/TypedElementSelectionSet.h"
-#include "Elements/Framework/EngineElementsLibrary.h"
+#include "Modules/ModuleManager.h"
+#include "PropertyEditorModule.h"
 #include "SSubobjectEditor.h"
+#include "SSubobjectEditorModule.h"
 #include "SSubobjectInstanceEditor.h"
+#include "ScopedTransaction.h"
+#include "SourceCodeNavigation.h"
 #include "SubobjectData.h"
 #include "SubobjectDataSubsystem.h"
-#include "SSubobjectEditorModule.h"
+#include "UnrealEdGlobals.h"
+#include "Widgets/Docking/SDockTab.h"
+#include "Widgets/Images/SImage.h"
+#include "Widgets/Layout/SBorder.h"
+#include "Widgets/Layout/SBox.h"
+#include "Widgets/Layout/SSplitter.h"
+#include "Widgets/SBoxPanel.h"
+#include "Widgets/Text/SRichTextBlock.h"
 
 class SActorDetailsUneditableComponentWarning : public SCompoundWidget
 {
@@ -100,7 +101,6 @@
 	FLevelEditorModule& LevelEditor = FModuleManager::GetModuleChecked<FLevelEditorModule>("LevelEditor");
 	LevelEditor.OnComponentsEdited().AddRaw(this, &SActorDetails::OnComponentsEditedInWorld);
 
-	FPropertyEditorModule& PropPlugin = FModuleManager::LoadModuleChecked<FPropertyEditorModule>("PropertyEditor");
 	FDetailsViewArgs DetailsViewArgs;
 	DetailsViewArgs.bUpdatesFromSelection = true;
 	DetailsViewArgs.bLockable = true;
@@ -113,6 +113,9 @@
 	DetailsViewArgs.DefaultsOnlyVisibility = EEditDefaultsOnlyNodeVisibility::Hide;
 	DetailsViewArgs.HostCommandList = InCommandList;
 	DetailsViewArgs.HostTabManager = InTabManager;
+	DetailsViewArgs.bShowSectionSelector = true;
+
+	FPropertyEditorModule& PropPlugin = FModuleManager::LoadModuleChecked<FPropertyEditorModule>("PropertyEditor");
 	DetailsView = PropPlugin.CreateDetailView(DetailsViewArgs);
 
 	auto IsPropertyVisible = [](const FPropertyAndParent& PropertyAndParent)
diff --git a/Engine/Source/Editor/PropertyEditor/Private/DetailCategoryBuilderImpl.cpp b/Engine/Source/Editor/PropertyEditor/Private/DetailCategoryBuilderImpl.cpp
--- a/Engine/Source/Editor/PropertyEditor/Private/DetailCategoryBuilderImpl.cpp
+++ b/Engine/Source/Editor/PropertyEditor/Private/DetailCategoryBuilderImpl.cpp
@@ -973,6 +973,22 @@ void FDetailCategoryImpl::FilterNode(const FDetailFilter& InFilter)
 		return;
 	}
 
+	if (!InFilter.VisibleSections.IsEmpty())
+	{
+		const UStruct* BaseStruct = GetParentBaseStructure();
+		if (BaseStruct != nullptr)
+		{
+			static FName PropertyEditor("PropertyEditor");
+			const FPropertyEditorModule& PropertyModule = FModuleManager::GetModuleChecked<FPropertyEditorModule>(PropertyEditor);
+
+			FName SectionName = PropertyModule.FindSectionForCategory(BaseStruct, CategoryName);
+			if (SectionName.IsNone() || !InFilter.VisibleSections.Contains(SectionName))
+			{
+				return;
+			}
+		}
+	}
+
 	if (InlinePropertyNode.IsValid())
 	{
 		bHasVisibleDetails = true;
diff --git a/Engine/Source/Editor/PropertyEditor/Private/DetailFilter.h b/Engine/Source/Editor/PropertyEditor/Private/DetailFilter.h
new file mode 100644
--- /dev/null
+++ b/Engine/Source/Editor/PropertyEditor/Private/DetailFilter.h
@@ -0,0 +1,51 @@
+// Copyright Epic Games, Inc. All Rights Reserved.
+
+#pragma once
+
+#include "CoreMinimal.h"
+
+/** Represents a filter which controls the visibility of items in the details view */
+struct FDetailFilter
+{
+	FDetailFilter() :
+		bShowOnlyModified(false)
+		, bShowAllAdvanced(false)
+		, bShowAllChildrenIfCategoryMatches(true)
+		, bShowOnlyKeyable(false)
+		, bShowOnlyAnimated(false)
+		, bShowFavoritesCategory(false)
+		, bShowOnlyWhitelisted(false)
+	{}
+
+	bool IsEmptyFilter() const
+	{
+		return FilterStrings.Num() == 0
+			&& VisibleSections.Num() == 0
+			&& bShowAllAdvanced == false
+			&& bShowOnlyWhitelisted == false
+			&& bShowAllChildrenIfCategoryMatches == false
+			&& bShowOnlyKeyable == false
+			&& bShowOnlyAnimated == false;
+	}
+
+	/** Any user search terms that items must match */
+	TArray<FString> FilterStrings;
+	/** If we should only show modified properties */
+	bool bShowOnlyModified;
+	/** If we should show all advanced properties */
+	bool bShowAllAdvanced;
+	/** If we should show all the children if their category name matches the search */
+	bool bShowAllChildrenIfCategoryMatches;
+	/** If we should only show keyable properties */
+	bool bShowOnlyKeyable;
+	/** If we should only show animated properties */
+	bool bShowOnlyAnimated;
+	/** If we should show the favorites category. */
+	bool bShowFavoritesCategory;
+	/** If we should only show whitelisted properties */
+	bool bShowOnlyWhitelisted;
+	/** The set of whitelisted properties to show. */
+	TSet<FPropertyPath> WhitelistedProperties;
+	/** The set of selected sections to show. If empty, all sections are shown.*/
+	TSet<FName> VisibleSections;
+};
diff --git a/Engine/Source/Editor/PropertyEditor/Private/PropertyEditorModule.cpp b/Engine/Source/Editor/PropertyEditor/Private/PropertyEditorModule.cpp
--- a/Engine/Source/Editor/PropertyEditor/Private/PropertyEditorModule.cpp
+++ b/Engine/Source/Editor/PropertyEditor/Private/PropertyEditorModule.cpp
@@ -2,36 +2,35 @@
 
 
 #include "PropertyEditorModule.h"
-#include "UObject/UnrealType.h"
-#include "Widgets/Layout/SBorder.h"
-#include "Modules/ModuleManager.h"
-#include "Framework/Application/SlateApplication.h"
-#include "Engine/UserDefinedEnum.h"
-#include "Engine/UserDefinedStruct.h"
-#include "Presentation/PropertyEditor/PropertyEditor.h"
-#include "SSingleProperty.h"
-#include "IDetailsView.h"
-#include "SDetailsView.h"
-#include "IPropertyTableWidgetHandle.h"
-#include "IPropertyTable.h"
-#include "UserInterface/PropertyTable/SPropertyTable.h"
-#include "UserInterface/PropertyTable/PropertyTableWidgetHandle.h"
-#include "IAssetTools.h"
 #include "AssetToolsModule.h"
-#include "SPropertyTreeViewImpl.h"
-#include "Interfaces/IMainFrameModule.h"
+#include "IAssetTools.h"
+#include "IDetailsView.h"
 #include "IPropertyChangeListener.h"
+#include "IPropertyTable.h"
+#include "IPropertyTableCellPresenter.h"
+#include "IPropertyTableWidgetHandle.h"
 #include "PropertyChangeListener.h"
-#include "Toolkits/AssetEditorToolkit.h"
 #include "PropertyEditorToolkit.h"
+#include "PropertyRowGenerator.h"
+#include "SDetailsView.h"
+#include "SPropertyTreeViewImpl.h"
+#include "SSingleProperty.h"
+#include "SStructureDetailsView.h"
 
+#include "Engine/UserDefinedEnum.h"
+#include "Engine/UserDefinedStruct.h"
+#include "Framework/Application/SlateApplication.h"
+#include "Interfaces/IMainFrameModule.h"
+#include "Modules/ModuleManager.h"
+#include "Presentation/PropertyEditor/PropertyEditor.h"
 #include "Presentation/PropertyTable/PropertyTable.h"
-#include "IPropertyTableCellPresenter.h"
+#include "Toolkits/AssetEditorToolkit.h"
+#include "UObject/UnrealType.h"
+#include "UserInterface/PropertyTable/PropertyTableWidgetHandle.h"
+#include "UserInterface/PropertyTable/SPropertyTable.h"
 #include "UserInterface/PropertyTable/TextPropertyTableCellPresenter.h"
-
-#include "SStructureDetailsView.h"
 #include "Widgets/Colors/SColorPicker.h"
-#include "PropertyRowGenerator.h"
+#include "Widgets/Layout/SBorder.h"
 
 
 IMPLEMENT_MODULE( FPropertyEditorModule, PropertyEditor );
@@ -514,7 +513,156 @@
 	}
 }
 
+void FPropertyEditorModule::RegisterSectionMapping(FName ClassName, FName SectionName, FName CategoryName)
+{
+	if (!ClassName.IsValid() || ClassName.IsNone())
+	{
+		return;
+	}
 
+	checkf(!CategoryName.ToString().Contains(TEXT("|")), TEXT("Cannnot register a section mapping for a subcategory. Class: '%s', Section: '%s', Category: '%s'"), *ClassName.ToString(), *SectionName.ToString(), *CategoryName.ToString());
 
+	FClassSectionMapping& SectionMapping = ClassSectionMappings.FindOrAdd(ClassName);
+
+	const FName* ExistingSection = SectionMapping.CategoryToSectionMap.Find(CategoryName);
+	if (ExistingSection != nullptr)
+	{
+		checkf(!ExistingSection, TEXT("This category has already been mapped to a section for this class. Class: '%s', Category: '%s', Existing Section: '%s', New Section: '%s'"), *ClassName.ToString(), *CategoryName.ToString(), *ExistingSection->ToString(), *SectionName.ToString());
+	}
+
+	SectionMapping.CategoryToSectionMap.Add(CategoryName, SectionName);
+}
+
+FName FPropertyEditorModule::FindSectionForCategory(const UStruct* Struct, FName CategoryName) const
+{
+	if (Struct == nullptr)
+	{
+		return NAME_None;
+	}
+
+	TSet<const UStruct*> SearchedStructs;
+	return FindSectionForCategoryHelper(Struct, CategoryName, SearchedStructs);
+}
+
+FName FPropertyEditorModule::FindSectionForCategoryHelper(const UStruct* Struct, FName CategoryName, TSet<const UStruct*>& SearchedStructs) const
+{
+	if (Struct == nullptr)
+	{
+		return NAME_None;
+	}
+
+	bool bAlreadySearched = false;
+	SearchedStructs.Add(Struct, &bAlreadySearched);
+
+	if (bAlreadySearched)
+	{
+		return NAME_None;
+	}
+
+	const UStruct* CurrentStruct = Struct;
+	while (CurrentStruct != nullptr)
+	{
+		// check this class' sections
+		const FClassSectionMapping* SectionMapping = ClassSectionMappings.Find(Struct->GetFName());
+		if (SectionMapping != nullptr)
+		{
+			const FName* SectionName = SectionMapping->CategoryToSectionMap.Find(CategoryName);
+			if (SectionName != nullptr)
+			{
+				return *SectionName;
+			}
+		}
+
+		// check all struct properties' sections
+		for (TFieldIterator<FStructProperty> It(CurrentStruct); It; ++It)
+		{
+			const FStructProperty* Property = *It;
+			const FName SectionName = FindSectionForCategoryHelper(Property->Struct, CategoryName, SearchedStructs);
+			if (!SectionName.IsNone())
+			{
+				return SectionName;
+			}
+		}
+
+		// check all object properties' sections
+		for (TFieldIterator<FObjectPropertyBase> It(CurrentStruct); It; ++It)
+		{
+			const FObjectPropertyBase* Property = *It;
+			if (Property->HasAnyPropertyFlags(CPF_InstancedReference | CPF_ContainsInstancedReference))
+			{
+				const FName SectionName = FindSectionForCategoryHelper(Property->PropertyClass, CategoryName, SearchedStructs);
+				if (!SectionName.IsNone())
+				{
+					return SectionName;
+				}
+			}
+		}
+
+		CurrentStruct = CurrentStruct->GetSuperStruct();
+	}
+
+	return NAME_None;
+}
+
+void FPropertyEditorModule::GetAllSections(const UStruct* Struct, TSet<FName>& OutSectionNames) const
+{
+	if (Struct == nullptr)
+	{
+		return;
+	}
+
+	TSet<const UStruct*> ProcessedStructs;
+	GetAllSectionsHelper(Struct, ProcessedStructs, OutSectionNames);
+}
+
+void FPropertyEditorModule::GetAllSectionsHelper(const UStruct* Struct, TSet<const UStruct*>& ProcessedStructs, TSet<FName>& OutSectionNames) const
+{
+	if (Struct == nullptr)
+	{
+		return;
+	}
+
+	bool bAlreadyProcessed = false;
+	ProcessedStructs.Add(Struct, &bAlreadyProcessed);
+
+	if (bAlreadyProcessed)
+	{
+		return;
+	}
+
+	const UStruct* CurrentStruct = Struct;
+	while (CurrentStruct != nullptr)
+	{
+		// add all sections defined for this class
+		const FClassSectionMapping* SectionMapping = ClassSectionMappings.Find(Struct->GetFName());
+		if (SectionMapping != nullptr)
+		{
+			for (const TPair<FName, FName>& Pair : SectionMapping->CategoryToSectionMap)
+			{
+				OutSectionNames.Add(Pair.Value);
+			}
+		}
+
+		// check all struct properties and add their sections
+		for (TFieldIterator<FStructProperty> It(CurrentStruct); It; ++It)
+		{
+			const FStructProperty* Property = *It;
+			GetAllSectionsHelper(Property->Struct, ProcessedStructs, OutSectionNames);
+		}
+
+		// check all object properties and add their sections
+		for (TFieldIterator<FObjectPropertyBase> It(CurrentStruct); It; ++It)
+		{
+			const FObjectPropertyBase* Property = *It;
+			if (Property->HasAnyPropertyFlags(CPF_InstancedReference | CPF_ContainsInstancedReference))
+			{
+				GetAllSectionsHelper(Property->PropertyClass, ProcessedStructs, OutSectionNames);
+			}
+		}
+
+		CurrentStruct = CurrentStruct->GetSuperStruct();
+	}
+}
+
 bool FPropertyEditorModule::HasUnlockedDetailViews() const
 {
diff --git a/Engine/Source/Editor/PropertyEditor/Private/SDetailsView.cpp b/Engine/Source/Editor/PropertyEditor/Private/SDetailsView.cpp
--- a/Engine/Source/Editor/PropertyEditor/Private/SDetailsView.cpp
+++ b/Engine/Source/Editor/PropertyEditor/Private/SDetailsView.cpp
@@ -23,6 +23,7 @@
 #include "Widgets/Input/SButton.h"
 #include "Widgets/Input/SComboButton.h"
 #include "Widgets/Input/SSearchBox.h"
+#include "Widgets/Input/SSegmentedControl.h"
 
 #define LOCTEXT_NAMESPACE "SDetailsView"
 
@@ -229,7 +230,7 @@ void SDetailsView::Construct(const FArguments& InArgs, const FDetailsViewArgs& I
 		[
 			// Create the search box
 			SAssignNew(SearchBox, SSearchBox)
-			.HintText(LOCTEXT("SearchDetailsHint", "Search Details"))
+			.HintText(LOCTEXT("SearchDetailsHint", "Search"))
 			.OnTextChanged(this, &SDetailsView::OnFilterTextChanged)
 			.OnTextCommitted(this, &SDetailsView::OnFilterTextCommitted)
 			.AddMetaData<FTagMetaData>(TEXT("Details.Search"))
@@ -305,20 +306,28 @@ void SDetailsView::Construct(const FArguments& InArgs, const FDetailsViewArgs& I
 			];
 	}
 
-	FilterRow = SNew(SVerticalBox)
+	TSharedRef<SVerticalBox> FilterRowVBox = SNew(SVerticalBox)
 		.Visibility( this, &SDetailsView::GetFilterBoxVisibility )
 		+ SVerticalBox::Slot()
 		.AutoHeight()
 		[
 			FilterRowHBox
-		]
-		+ SVerticalBox::Slot()
-		.AutoHeight()
-		[
-			SNew(SImage)
-			.Image(FAppStyle::Get().GetBrush("DetailsView.AdvancedDropdownBorder.Open"))
 		];
 
+	if (DetailsViewArgs.bShowSectionSelector)
+	{
+		FilterRowVBox->AddSlot()
+			.Padding(30, 2, 30, 7)
+			.AutoHeight()
+			[
+				SAssignNew(SectionSelectorBox, SBox)
+			];
+
+		RebuildSectionSelector();
+	}
+
+	FilterRow = FilterRowVBox;
+
 	if( !DetailsViewArgs.bCustomFilterAreaLocation )
 	{
 		VerticalBox->AddSlot()
@@ -356,6 +365,7 @@ void SDetailsView::Construct(const FArguments& InArgs, const FDetailsViewArgs& I
 		[
 			SNew(STextBlock)
 			.Text(LOCTEXT("AllItemsFiltered", "All results have been filtered. Try changing your active filters above."))
+			.AutoWrapText(true)
 			.Visibility_Lambda([this]() { return ((this->GetFilterBoxVisibility() == EVisibility::Visible) && !this->CurrentFilter.IsEmptyFilter() && RootTreeNodes.Num() == 0) ? EVisibility::HitTestInvisible : EVisibility::Collapsed; })
 		]
 		+ SOverlay::Slot()
@@ -724,6 +734,8 @@ void SDetailsView::SetObjectArrayPrivate(const TArray<UObject*>& InObjects)
 
 	OnObjectArrayChanged.ExecuteIfBound(Title, InObjects);
 
+	RebuildSectionSelector();
+
 	double ElapsedTime = FPlatformTime::Seconds() - StartTime;
 }
 
@@ -1156,4 +1168,138 @@ void SDetailsView::SetCustomBuilderFavorite(FStringView Path, bool IsFavorite)
 	}
 }
 
+void SDetailsView::RebuildSectionSelector()
+{
+	if (!SectionSelectorBox.IsValid())
+	{
+		return;
+	}
+
+	const TSet<FName> Sections = GetAllSectionNames();
+	if (Sections.IsEmpty())
+	{
+		// we've selected something that has no sections - rather than show just "All", reset and hide the box
+		SectionSelectorBox->SetContent(SNullWidget::NullWidget);
+		SectionSelectorBox->SetVisibility(EVisibility::Collapsed);
+		return;
+	}
+
+	int SegmentsPerLine = 5;
+
+	const int NumSections = Sections.Num() + 1; // implicit "All"
+	const int NumLines = (NumSections / SegmentsPerLine) + 1;
+	if (NumLines > 1)
+	{
+		// the segmented control fills MaxSegmentsPerLine segments at a go, this attempts to normalize that
+		// eg. 6 will get divided into 3 + 3 instead of 5 + 1, 12 into 4 + 4 + 4 instead of 5 + 5 + 2
+
+		SegmentsPerLine = (NumSections / NumLines) + (NumSections % NumLines == 0 ? 0 : 1);
+	}
+
+	TSharedRef<SSegmentedControl<FName>> SectionSelectorWidget = SNew(SSegmentedControl<FName>)
+		.TextStyle(FAppStyle::Get(), "SmallText")
+		.MaxSegmentsPerLine(SegmentsPerLine)
+		.OnValueChanged(this, &SDetailsView::OnSectionSelectionChanged)
+		.Value(this, &SDetailsView::GetSelectedSection);
+
+
+	for (const FName& SectionName : Sections)
+	{
+		FText SectionDisplayName = GetSectionDisplayName(SectionName);
+
+		SectionSelectorWidget->AddSlot(SectionName, false)
+			.Text(SectionDisplayName);
+	}
+
+	SectionSelectorWidget->AddSlot(FName("All"), false)
+		.Text(NSLOCTEXT("UObjectSection", "All", "All"));
+
+	SectionSelectorWidget->RebuildChildren();
+
+	SectionSelectorBox->SetContent(SectionSelectorWidget);
+	SectionSelectorBox->SetVisibility(EVisibility::Visible);
+}
+
+void SDetailsView::OnSectionSelectionChanged(FName NewSelection)
+{
+	CurrentFilter.VisibleSections.Reset();
+
+	if (NewSelection != "All")
+	{
+		CurrentFilter.VisibleSections.Add(NewSelection);
+	}
+
+	UpdateFilteredDetails();
+}
+
+FName SDetailsView::GetSelectedSection() const
+{
+	FName CurrentSection("All");
+	for (FName Section : CurrentFilter.VisibleSections)
+	{
+		CurrentSection = Section;
+		break;
+	}
+
+	return CurrentSection;
+}
+
+TSet<FName> SDetailsView::GetAllSectionNames() const 
+{
+	static const FName PropertyEditor("PropertyEditor");
+	FPropertyEditorModule& PropertyModule = FModuleManager::GetModuleChecked<FPropertyEditorModule>(PropertyEditor);
+
+	TSet<FName> SectionNames;
+
+	// for every category, check every base struct and find the associated section
+	// if one exists, add it to the set of valid section names
+	for (const TSharedRef<FDetailTreeNode>& RootNode : RootTreeNodes)
+	{
+		if (RootNode->GetNodeType() == EDetailNodeType::Category)
+		{
+			FName CategoryName = RootNode->GetNodeName();
+
+			for (const TSharedPtr<FComplexPropertyNode>& RootPropertyNode : RootPropertyNodes)
+			{
+				const UStruct* RootBaseStruct = RootPropertyNode->GetBaseStructure();
+
+				FName SectionName = PropertyModule.FindSectionForCategory(RootBaseStruct, CategoryName);
+				if (!SectionName.IsNone())
+				{
+					SectionNames.Add(SectionName);
+				}
+			}
+		}
+	}
+
+
+	return MoveTemp(SectionNames);
+}
+
+FText SDetailsView::GetSectionDisplayName(FName SectionName) const
+{
+	static const FTextKey SectionLocalizationNamespace = TEXT("UObjectSection");
+	static const FName SectionMetaDataKey = TEXT("Section");
+
+	FText SectionDisplayName;
+
+	const FString SectionString = SectionName.ToString();
+	if (FText::FindText(SectionLocalizationNamespace, SectionString, /*OUT*/SectionDisplayName, &SectionString))
+	{
+		// Category names in English are typically gathered in their non-pretty form (eg "UserInterface" rather than "User Interface"), so skip 
+		// applying the localized variant if the text matches the raw category name, as in this case the pretty printer will do a better job
+		if (SectionString.Equals(SectionDisplayName.ToString(), ESearchCase::CaseSensitive))
+		{
+			SectionDisplayName = FText();
+		}
+	}
+
+	if (SectionDisplayName.IsEmpty())
+	{
+		SectionDisplayName = FText::AsCultureInvariant(FName::NameToDisplayString(SectionString, false));
+	}
+
+	return SectionDisplayName;
+}
+
 #undef LOCTEXT_NAMESPACE
diff --git a/Engine/Source/Editor/PropertyEditor/Private/SDetailsView.h b/Engine/Source/Editor/PropertyEditor/Private/SDetailsView.h
--- a/Engine/Source/Editor/PropertyEditor/Private/SDetailsView.h
+++ b/Engine/Source/Editor/PropertyEditor/Private/SDetailsView.h
@@ -3,14 +3,15 @@
 #pragma once
 
 #include "CoreMinimal.h"
-#include "Layout/Visibility.h"
+#include "AssetSelection.h"
+#include "IDetailsView.h"
+#include "PropertyNode.h"
+#include "SDetailsViewBase.h"
+
 #include "Input/Reply.h"
+#include "Layout/Visibility.h"
 #include "Widgets/DeclarativeSyntaxSupport.h"
-#include "PropertyNode.h"
-#include "IDetailsView.h"
-#include "AssetSelection.h"
 #include "Widgets/Layout/SScrollBar.h"
-#include "SDetailsViewBase.h"
 
 class AActor;
 class IDetailRootObjectCustomization;
@@ -35,7 +36,6 @@ public:
 	/** Move the scrolling offset (by item), but do not refresh the tree*/
 	void MoveScrollOffset(int32 DeltaOffset) override;
 
-
 	/**
 	 * Constructs the property view widgets                   
 	 */
@@ -173,6 +173,21 @@ private:
 	/** Called after an undo or redo operation occurs in the editor. */
 	void OnPostUndoRedo();
 
+	/** Get all section names for the objects currently selected in the view. */
+	TSet<FName> GetAllSectionNames() const;
+
+	/** Get a display name for the given section. */
+	FText GetSectionDisplayName(FName SectionName) const;
+	
+	/** Rebuild the section selector widget after a selection has been changed. */
+	void RebuildSectionSelector();
+
+	/** Refilter the details view after the user has selected a new section. */
+	void OnSectionSelectionChanged(FName NewSelection);
+
+	/** Get the currently selected section. */
+	FName GetSelectedSection() const;
+
 private:
 	/** The filter for objects viewed by this details panel */
 	TSharedPtr<FDetailsViewObjectFilter> ObjectFilter;
@@ -201,4 +216,6 @@ private:
 	bool bViewingClassDefaultObject;
 	/** Delegate handle for unregistering from the PostUndoRedo event. */
 	FDelegateHandle PostUndoRedoDelegateHandle;
+	/** The section selector widget to show if DetailsViewArgs.bShowSectionSelector is true. */
+	TSharedPtr<SBox> SectionSelectorBox;
 };
diff --git a/Engine/Source/Editor/PropertyEditor/Private/SDetailsViewBase.h b/Engine/Source/Editor/PropertyEditor/Private/SDetailsViewBase.h
--- a/Engine/Source/Editor/PropertyEditor/Private/SDetailsViewBase.h
+++ b/Engine/Source/Editor/PropertyEditor/Private/SDetailsViewBase.h
@@ -3,26 +3,27 @@
 #pragma once
 
 #include "CoreMinimal.h"
-#include "Misc/Attribute.h"
-#include "Layout/Visibility.h"
-#include "PropertyPath.h"
-#include "Input/Reply.h"
 #include "AssetThumbnail.h"
+#include "IDetailsView.h"
+#include "IDetailsViewPrivate.h"
 #include "IPropertyUtilities.h"
+#include "DetailFilter.h"
 #include "DetailTreeNode.h"
-#include "Widgets/Views/STableViewBase.h"
-#include "Widgets/Views/STableRow.h"
-#include "PropertyNode.h"
-#include "Widgets/SWindow.h"
+#include "DetailsViewConfig.h"
+#include "PropertyCustomizationHelpers.h"
 #include "PropertyEditorModule.h"
+#include "PropertyNode.h"
+#include "PropertyPath.h"
+#include "PropertyRowGenerator.h"
 #include "Framework/Commands/UICommandList.h"
+#include "Input/Reply.h"
+#include "Layout/Visibility.h"
+#include "Misc/Attribute.h"
+#include "Widgets/SWindow.h"
 #include "Widgets/Layout/SSplitter.h"
+#include "Widgets/Views/STableViewBase.h"
+#include "Widgets/Views/STableRow.h"
 #include "Widgets/Views/STreeView.h"
-#include "IDetailsView.h"
-#include "IDetailsViewPrivate.h"
-#include "PropertyRowGenerator.h"
-#include "PropertyCustomizationHelpers.h"
-#include "DetailsViewConfig.h"
 
 class FDetailCategoryImpl;
 class FDetailLayoutBuilderImpl;
@@ -37,49 +38,6 @@ struct FDetailsViewObjectRoot;
 
 typedef STreeView< TSharedRef<class FDetailTreeNode> > SDetailTree;
 
-/** Represents a filter which controls the visibility of items in the details view */
-struct FDetailFilter
-{
-	FDetailFilter() :
-		bShowOnlyModified(false)
-		, bShowAllAdvanced(false)
-		, bShowAllChildrenIfCategoryMatches(true)
-		, bShowOnlyKeyable(false)
-		, bShowOnlyAnimated(false)
-		, bShowFavoritesCategory(false)
-		, bShowOnlyWhitelisted(false)
-	{}
-
-	bool IsEmptyFilter() const 
-	{ 
-		return FilterStrings.Num() == 0 
-			&& bShowAllAdvanced == false 
-			&& bShowOnlyWhitelisted == false 
-			&& bShowAllChildrenIfCategoryMatches == false 
-			&& bShowOnlyKeyable == false 
-			&& bShowOnlyAnimated == false;
-	}
-
-	/** Any user search terms that items must match */
-	TArray<FString> FilterStrings;
-	/** If we should only show modified properties */
-	bool bShowOnlyModified;
-	/** If we should show all advanced properties */
-	bool bShowAllAdvanced;
-	/** If we should show all the children if their category name matches the search */
-	bool bShowAllChildrenIfCategoryMatches;
-	/** If we should only show keyable properties */
-	bool bShowOnlyKeyable;
-	/** If we should only show animated properties */
-	bool bShowOnlyAnimated;
-	/** If we should show the favorites category. */
-	bool bShowFavoritesCategory;
-	/** If we should only show whitelisted properties */
-	bool bShowOnlyWhitelisted;
-	/** The set of whitelisted properties to show. */
-	TSet<FPropertyPath> WhitelistedProperties;
-};
-
 class SDetailsViewBase : public IDetailsViewPrivate
 {
 public:
@@ -375,7 +333,7 @@ protected:
 	/** The current detail layout based on objects in this details panel.  There is one layout for each top level object node.*/
 	FDetailLayoutList DetailLayouts;
 	/** Row for searching and view options */
-	TSharedPtr<SWidget>  FilterRow;
+	TSharedPtr<SWidget> FilterRow;
 	/** Search box */
 	TSharedPtr<SSearchBox> SearchBox;
 	/** Customization instances that need to be destroyed when safe to do so */
diff --git a/Engine/Source/Editor/PropertyEditor/Public/PropertyEditorModule.h b/Engine/Source/Editor/PropertyEditor/Public/PropertyEditorModule.h
--- a/Engine/Source/Editor/PropertyEditor/Public/PropertyEditorModule.h
+++ b/Engine/Source/Editor/PropertyEditor/Public/PropertyEditorModule.h
@@ -2,13 +2,14 @@
 #pragma once
 
 #include "CoreMinimal.h"
-#include "Widgets/SWidget.h"
-#include "Widgets/SWindow.h"
-#include "Modules/ModuleInterface.h"
-#include "UObject/StructOnScope.h"
 #include "IDetailsView.h"
-#include "PropertyEditorDelegates.h"
 #include "IPropertyTypeCustomization.h"
+#include "PropertyEditorDelegates.h"
+
+#include "Modules/ModuleInterface.h"
+#include "UObject/StructOnScope.h"
+#include "Widgets/SWidget.h"
+#include "Widgets/SWindow.h"
 
 class FAssetEditorToolkit;
 class FNotifyHook;
@@ -184,13 +185,28 @@ public:
 	virtual void RegisterCustomPropertyTypeLayout( FName PropertyTypeName, FOnGetPropertyTypeCustomizationInstance PropertyTypeLayoutDelegate, TSharedPtr<IPropertyTypeIdentifier> Identifier = nullptr);
 
 	/**
-	 * Unregisters a custom detail layout for a properrty type
+	 * Unregisters a custom detail layout for a property type
 	 *
 	 * @param PropertyTypeName 	The name of the property type that was registered
 	 * @param Identifier 		An identifier to use to differentiate between two customizations on the same type
 	 */
 	virtual void UnregisterCustomPropertyTypeLayout( FName PropertyTypeName, TSharedPtr<IPropertyTypeIdentifier> InIdentifier = nullptr);
 
+	/**
+	 * Register a section mapping for a class.
+	 * 
+	 * @param ClassName		The class to add a section mapping for.
+	 * @param SectionName	The section to add a category to.
+	 * @param CategoryName	The category to add.
+	 */
+	virtual void RegisterSectionMapping(FName ClassName, FName SectionName, FName CategoryName);
+
+	/** Find the section that the given category in the given struct should be a part of. */
+	virtual FName FindSectionForCategory(const UStruct* Struct, FName CategoryName) const;
+
+	/** Get all registered sections for the given struct (including the default section). */
+	virtual void GetAllSections(const UStruct* Struct, TSet<FName>& OutSectionNames) const;
+
 	/**
 	 * Customization modules should call this when that module has been unloaded, loaded, etc...
 	 * so the property module can clean up its data.  Needed to support dynamic reloading of modules
@@ -299,6 +315,10 @@ private:
 	virtual TSharedRef<SPropertyTreeViewImpl> CreatePropertyView( UObject* InObject, bool bAllowFavorites, bool bIsLockable, bool bHiddenPropertyVisibility, bool bAllowSearch, bool ShowTopLevelNodes, FNotifyHook* InNotifyHook, float InNameColumnWidth, FOnPropertySelectionChanged OnPropertySelectionChanged, FOnPropertyClicked OnPropertyMiddleClicked, FConstructExternalColumnHeaders ConstructExternalColumnHeaders, FConstructExternalColumnCell ConstructExternalColumnCell );
 
 	TSharedPtr<FAssetThumbnailPool> GetThumbnailPool();
+
+	void GetAllSectionsHelper(const UStruct* Struct, TSet<const UStruct*>& ProcessedStructs, TSet<FName>& OutSectionNames) const;
+	FName FindSectionForCategoryHelper(const UStruct* Struct, FName CategoryName, TSet<const UStruct*>& SearchedStructs) const;
+
 private:
 	/** All created detail views */
 	TArray< TWeakPtr<class SDetailsView> > AllDetailViews;
@@ -308,6 +328,16 @@ private:
 	FCustomDetailLayoutNameMap ClassNameToDetailLayoutNameMap;
 	/** A mapping of property names to property type layout delegates, called when querying for custom property layouts */
 	FCustomPropertyTypeLayoutMap GlobalPropertyTypeToLayoutMap;
+
+	/** A mapping of categories to section names for a given class. */
+	struct FClassSectionMapping
+	{
+		/** A mapping of category to the section that it should be added to. */
+		TMap<FName, FName> CategoryToSectionMap;
+	};
+
+	/** A mapping of class names to section mappings. */
+	TMap<FName, FClassSectionMapping> ClassSectionMappings;
 	/** Event to be called when a property editor is opened */
 	FPropertyEditorOpenedEvent PropertyEditorOpened;
 	/** Mapping of registered floating UStructs to their struct proxy so they show correctly in the details panel */
