Fixed EditCondition enum comparisons failing due to looking up the value pointer from the UnderlyingNumericProperty instead of the Property.

diff --git a/Engine/Source/Editor/PropertyEditor/Private/EditConditionContext.cpp b/Engine/Source/Editor/PropertyEditor/Private/EditConditionContext.cpp
--- a/Engine/Source/Editor/PropertyEditor/Private/EditConditionContext.cpp
+++ b/Engine/Source/Editor/PropertyEditor/Private/EditConditionContext.cpp
@@ -298,12 +298,7 @@ TOptional<FString> FEditConditionContext::GetEnumValue(const FString& PropertyNa
 		EnumType = ByteProperty->GetIntPropertyEnum();
 	}
 
-	if (EnumType == nullptr)
-	{
-		return TOptional<FString>();
-	}
-	
-	if (NumericProperty == nullptr || !NumericProperty->IsInteger())
+	if (EnumType == nullptr || NumericProperty == nullptr || !NumericProperty->IsInteger())
 	{
 		return TOptional<FString>();
 	}
@@ -323,13 +318,13 @@ TOptional<FString> FEditConditionContext::GetEnumValue(const FString& PropertyNa
 	TOptional<int64> Result;
 	for (int32 Index = 0; Index < ComplexParentNode->GetInstancesNum(); ++Index)
 	{
-		const uint8* ValuePtr = GetPropertyValuePtr(NumericProperty, PinnedNode, ParentNode, ComplexParentNode, Index);
+		const uint8* ValuePtr = GetPropertyValuePtr(Property, PinnedNode, ParentNode, ComplexParentNode, Index);
 		if (ValuePtr == nullptr)
 		{
 			return TOptional<FString>();
 		}
 
-		int64 Value = NumericProperty->GetSignedIntPropertyValue(ValuePtr);
+		const int64 Value = NumericProperty->GetSignedIntPropertyValue(ValuePtr);
 		if (!Result.IsSet())
 		{
 			Result = Value;
@@ -421,12 +416,7 @@ TOptional<FString> FEditConditionContext::GetTypeName(const FString& PropertyNam
 
 	if (const FEnumProperty* EnumProperty = CastField<FEnumProperty>(Property))
 	{
-		const UEnum* EnumType = EnumProperty->GetEnum();
-		if (EnumType == nullptr)
-		{
-			return TOptional<FString>();
-		}
-		return EnumType->GetName();
+		return EnumProperty->GetEnum()->GetName();
 	}
 	else if (const FByteProperty* ByteProperty = CastField<FByteProperty>(Property))
 	{
@@ -448,7 +438,7 @@ TOptional<int64> FEditConditionContext::GetIntegerValueOfEnum(const FString& Enu
 		return TOptional<int64>();
 	}
 
-	int64 EnumValue = EnumType->GetValueByName(FName(*MemberName));
+	const int64 EnumValue = EnumType->GetValueByName(FName(*MemberName));
 	if (EnumValue == INDEX_NONE)
 	{
 		return TOptional<int64>();
