Blocking restricted classes from showing in a class selector, preventing a user from setting up a ChildActorComponent that would produce an recursive cycle

diff --git a/Engine/Source/Editor/Kismet/Private/BlueprintDetailsCustomization.cpp b/Engine/Source/Editor/Kismet/Private/BlueprintDetailsCustomization.cpp
--- a/Engine/Source/Editor/Kismet/Private/BlueprintDetailsCustomization.cpp
+++ b/Engine/Source/Editor/Kismet/Private/BlueprintDetailsCustomization.cpp
@@ -6791,25 +6791,20 @@ void FChildActorComponentDetails::CustomizeDetails(IDetailLayoutBuilder& DetailB
 	{
 		if (BlueprintEditorPtr.IsValid())
 		{
-			// only restrict for the components view (you can successfully add 
-			// a self child component in the execution graphs)
-			if (BlueprintEditorPtr.Pin()->GetCurrentMode() == FBlueprintEditorApplicationModes::BlueprintComponentsMode)
+			if (UBlueprint* Blueprint = BlueprintEditorPtr.Pin()->GetBlueprintObj())
 			{
-				if (UBlueprint* Blueprint = BlueprintEditorPtr.Pin()->GetBlueprintObj())
-				{
-					FText RestrictReason = LOCTEXT("NoSelfChildActors", "Cannot append a child-actor of this blueprint type (could cause infinite recursion).");
-					TSharedPtr<FPropertyRestriction> ClassRestriction = MakeShareable(new FPropertyRestriction(MoveTemp(RestrictReason)));
-
-					ClassRestriction->AddDisabledValue(Blueprint->GetName());
-					ClassRestriction->AddDisabledValue(Blueprint->GetPathName());
-					if (Blueprint->GeneratedClass)
-					{
-						ClassRestriction->AddDisabledValue(Blueprint->GeneratedClass->GetName());
-						ClassRestriction->AddDisabledValue(Blueprint->GeneratedClass->GetPathName());
-					}
+				static FText RestrictReason = LOCTEXT("NoSelfChildActors", "Cannot append a child-actor of this blueprint type (could cause infinite recursion).");
+				TSharedPtr<FPropertyRestriction> ClassRestriction = MakeShareable(new FPropertyRestriction(RestrictReason));
 
-					ActorClassProperty->AddRestriction(ClassRestriction.ToSharedRef());
+				ClassRestriction->AddDisabledValue(Blueprint->GetName());
+				ClassRestriction->AddDisabledValue(Blueprint->GetPathName());
+				if (Blueprint->GeneratedClass)
+				{
+					ClassRestriction->AddDisabledValue(Blueprint->GeneratedClass->GetName());
+					ClassRestriction->AddDisabledValue(Blueprint->GeneratedClass->GetPathName());
 				}
+
+				ActorClassProperty->AddRestriction(ClassRestriction.ToSharedRef());
 			}
 		}
 
diff --git a/Engine/Source/Editor/PropertyEditor/Private/UserInterface/PropertyEditor/SPropertyEditorClass.cpp b/Engine/Source/Editor/PropertyEditor/Private/UserInterface/PropertyEditor/SPropertyEditorClass.cpp
--- a/Engine/Source/Editor/PropertyEditor/Private/UserInterface/PropertyEditor/SPropertyEditorClass.cpp
+++ b/Engine/Source/Editor/PropertyEditor/Private/UserInterface/PropertyEditor/SPropertyEditorClass.cpp
@@ -11,6 +11,7 @@
 #include "ClassViewerModule.h"
 #include "ClassViewerFilter.h"
 #include "UObject/UObjectIterator.h"
+#include "PropertyRestriction.h"
 
 #define LOCTEXT_NAMESPACE "PropertyEditor"
 
@@ -69,6 +70,18 @@ private:
 	}
 };
 
+static UClass* FindOrLoadClass(const FString& ClassName)
+{
+	UClass* Class = FindObject<UClass>(ANY_PACKAGE, *ClassName);
+
+	if (!Class)
+	{
+		Class = LoadObject<UClass>(nullptr, *ClassName);
+	}
+
+	return Class;
+}
+
 void SPropertyEditorClass::GetDesiredWidth(float& OutMinDesiredWidth, float& OutMaxDesiredWidth)
 {
 	OutMinDesiredWidth = 200.0f;
@@ -140,12 +153,7 @@ void SPropertyEditorClass::Construct(const FArguments& InArgs, const TSharedPtr<
 
 				for (const FString& ClassName : ClassFilterNames)
 				{
-					UClass* Class = FindObject<UClass>(ANY_PACKAGE, *ClassName);
-
-					if (!Class)
-					{
-						Class = LoadObject<UClass>(nullptr, *ClassName);
-					}
+					UClass* Class = FindOrLoadClass(ClassName);
 
 					if (Class)
 					{
@@ -170,8 +178,29 @@ void SPropertyEditorClass::Construct(const FArguments& InArgs, const TSharedPtr<
 			}
 		};
 
+		// Filter based on UPROPERTY meta data
 		FillClassFilters(AllowedClassFilters, "AllowedClasses");
 		FillClassFilters(DisallowedClassFilters, "DisallowedClasses");
+
+		// Filter based on restrictions
+		for (const TSharedRef<const FPropertyRestriction>& ClassRestriction : PropertyNode->GetRestrictions())
+		{
+			for (TArray<FString>::TConstIterator Iter= ClassRestriction.Get().GetHiddenValuesIterator(); Iter; ++Iter)
+			{
+				if (UClass* HiddenClass = FindOrLoadClass(*Iter))
+				{
+					DisallowedClassFilters.Add(HiddenClass);
+				}
+			}
+
+			for (TArray<FString>::TConstIterator Iter = ClassRestriction.Get().GetDisabledValuesIterator(); Iter; ++Iter)
+			{
+				if (UClass* DisabledClass = FindOrLoadClass(*Iter))
+				{
+					DisallowedClassFilters.Add(DisabledClass);
+				}
+			}
+		}
 	}
 	else
 	{
@@ -332,11 +361,7 @@ void SPropertyEditorClass::SendToObjects(const FString& NewValue)
 	}
 	else if (!NewValue.IsEmpty() && NewValue != TEXT("None"))
 	{
-		UClass* NewClass = FindObject<UClass>(ANY_PACKAGE, *NewValue);
-		if(!NewClass)
-		{
-			NewClass = LoadObject<UClass>(nullptr, *NewValue);
-		}
+		UClass* NewClass = FindOrLoadClass(NewValue);
 		OnSetClass.Execute(NewClass);
 	}
 	else
diff --git a/Engine/Source/Editor/PropertyEditor/Public/PropertyRestriction.h b/Engine/Source/Editor/PropertyEditor/Public/PropertyRestriction.h
--- a/Engine/Source/Editor/PropertyEditor/Public/PropertyRestriction.h
+++ b/Engine/Source/Editor/PropertyEditor/Public/PropertyRestriction.h
@@ -24,6 +24,16 @@ public:
 	void RemoveDisabledValue(FString Value);
 	void RemoveAll();
 
+	TArray<FString>::TConstIterator GetHiddenValuesIterator() const 
+	{
+		return HiddenValues.CreateConstIterator();
+	}
+
+	TArray<FString>::TConstIterator GetDisabledValuesIterator() const
+	{
+		return DisabledValues.CreateConstIterator();
+	}
+
 private:
 
 	TArray<FString> HiddenValues;
