Removed unnecessary typename retrieval in EditConditionParser.

Added comment clarifying intent in EditConditionContext.

diff --git a/Engine/Source/Editor/PropertyEditor/Private/EditConditionContext.cpp b/Engine/Source/Editor/PropertyEditor/Private/EditConditionContext.cpp
--- a/Engine/Source/Editor/PropertyEditor/Private/EditConditionContext.cpp
+++ b/Engine/Source/Editor/PropertyEditor/Private/EditConditionContext.cpp
@@ -318,6 +318,8 @@ TOptional<FString> FEditConditionContext::GetEnumValue(const FString& PropertyNa
 	TOptional<int64> Result;
 	for (int32 Index = 0; Index < ComplexParentNode->GetInstancesNum(); ++Index)
 	{
+		// NOTE: this very intentionally fetches the value from Property, not NumericProperty, 
+		// because the underlying property of an enum does not return a valid value
 		const uint8* ValuePtr = GetPropertyValuePtr(Property, PinnedNode, ParentNode, ComplexParentNode, Index);
 		if (ValuePtr == nullptr)
 		{
diff --git a/Engine/Source/Editor/PropertyEditor/Private/EditConditionParser.cpp b/Engine/Source/Editor/PropertyEditor/Private/EditConditionParser.cpp
--- a/Engine/Source/Editor/PropertyEditor/Private/EditConditionParser.cpp
+++ b/Engine/Source/Editor/PropertyEditor/Private/EditConditionParser.cpp
@@ -213,8 +213,6 @@ static FExpressionResult ApplyPropertyIsNull(const EditConditionParserTokens::FP
 		return MakeError(FText::Format(LOCTEXT("InvalidOperand", "EditCondition attempted to use an invalid operand \"{0}\"."), FText::FromString(Property.PropertyName)));
 	}
 
-	const FString& TypeNameValue = TypeName.GetValue();
-	
 	TOptional<UObject*> Ptr = Context.GetPointerValue(Property.PropertyName);
 	if (!Ptr.IsSet())
 	{
