Details view style refresh:
- rows are now 26 pixels high
- right column is narrower unless the Sequencer is invoked
- padding has been removed to make more rows fit on the screen at once
- Advanced dropdown is now styled like a normal group and thus the bottom node is gone

diff --git a/Engine/Source/Editor/EditorStyle/Private/StarshipStyle.cpp b/Engine/Source/Editor/EditorStyle/Private/StarshipStyle.cpp
--- a/Engine/Source/Editor/EditorStyle/Private/StarshipStyle.cpp
+++ b/Engine/Source/Editor/EditorStyle/Private/StarshipStyle.cpp
@@ -3219,9 +3219,9 @@ void FStarshipEditorStyle::FStyle::SetupPropertyEditorStyles()
 		Set( "DetailsView.EditConfigProperties",      new IMAGE_BRUSH("Icons/icon_PropertyMatrix_16px",  Icon16x16, FLinearColor::White ) );
 
 		Set( "DetailsView.CollapsedCategory",         new FSlateColorBrush(FStyleColors::Header));
-		Set( "DetailsView.CollapsedCategory_Hovered", new FSlateColorBrush(FStyleColors::Header));
+		Set( "DetailsView.CollapsedCategory_Hovered", new FSlateColorBrush(FStyleColors::Hover));
 		Set( "DetailsView.CategoryTop",               new FSlateColorBrush(FStyleColors::Header));
-		Set( "DetailsView.CategoryTop_Hovered",       new FSlateColorBrush(FStyleColors::Header));
+		Set( "DetailsView.CategoryTop_Hovered",       new FSlateColorBrush(FStyleColors::Hover));
 		Set( "DetailsView.CategoryBottom",            new FSlateColorBrush(FStyleColors::Recessed));
 		
 		// these are not actually displayed as white, see PropertyEditorConstants::GetRowBackgroundColor
@@ -3239,10 +3239,7 @@ void FStarshipEditorStyle::FStyle::SetupPropertyEditorStyles()
 
 		Set( "DetailsView.GridLine", new FSlateColorBrush(FStyleColors::Recessed) );
 
-		Set( "DetailsView.AdvancedDropdownBorder",      new FSlateColorBrush(FStyleColors::Panel));
-		Set( "DetailsView.AdvancedDropdownBorder.Open", new IMAGE_BRUSH("Common/ScrollBoxShadowTop", FVector2D(64,8)));
-
-		Set( "DetailsView.CategoryFontStyle", FStyleFonts::Get().NormalBold);
+		Set( "DetailsView.CategoryFontStyle", FStyleFonts::Get().SmallBold);
 		Set( "DetailsView.CategoryTextStyle", FTextBlockStyle(NormalText)
 			.SetFont(GetFontStyle("DetailsView.CategoryFontStyle"))
 			.SetColorAndOpacity(FStyleColors::ForegroundHeader)
@@ -3259,26 +3256,25 @@ void FStarshipEditorStyle::FStyle::SetupPropertyEditorStyles()
 			.SetHoveredForeground(FStyleColors::ForegroundHover)
 			.SetPressedForeground(FStyleColors::ForegroundHover)
 			.SetDisabledForeground(FStyleColors::Foreground)
-			.SetNormalPadding(FMargin(0, 2, 0, 2))
-			.SetPressedPadding(FMargin(0, 3, 0, 1));
+			.SetNormalPadding(FMargin(2, 2, 2, 2))
+			.SetPressedPadding(FMargin(2, 3, 2, 1));
 
 		Set("DetailsView.ExtensionToolBar.Button", DetailsExtensionMenuButton);
 
 		FToolBarStyle DetailsExtensionToolBarStyle = FToolBarStyle(FStarshipCoreStyle::GetCoreStyle().GetWidgetStyle<FToolBarStyle>("SlimToolBar"))
-			.SetBackground(FSlateNoResource())
+			.SetButtonStyle(DetailsExtensionMenuButton)
+			.SetExpandBrush(CORE_IMAGE_BRUSH_SVG("Starship/Common/ellipsis-vertical-narrow", FVector2D(4, 16)))
 			.SetIconSize(Icon16x16)
+			.SetBackground(FSlateNoResource())
 			.SetLabelPadding(FMargin(0))
 			.SetComboButtonPadding(FMargin(0))
-			.SetBlockPadding(FMargin(2, 0))
-			.SetIndentedBlockPadding(FMargin(0))
+			.SetBlockPadding(FMargin(0))
 			.SetIndentedBlockPadding(FMargin(0))
 			.SetBackgroundPadding(FMargin(0))
-			.SetButtonPadding(FMargin(2, 0))
-			.SetCheckBoxPadding(FMargin(3, 0))
-			.SetButtonStyle(DetailsExtensionMenuButton)
+			.SetButtonPadding(FMargin(0))
+			.SetCheckBoxPadding(FMargin(0))
 			.SetSeparatorBrush(FSlateNoResource())
-			.SetSeparatorPadding(FMargin(4, 0))
-			.SetExpandBrush(CORE_IMAGE_BRUSH_SVG("Starship/Common/ellipsis-vertical-narrow", FVector2D(6, 24)));
+			.SetSeparatorPadding(FMargin(0));
 
 		Set("DetailsView.ExtensionToolBar", DetailsExtensionToolBarStyle);
 
diff --git a/Engine/Source/Editor/PropertyEditor/Private/DetailAdvancedDropdownNode.cpp b/Engine/Source/Editor/PropertyEditor/Private/DetailAdvancedDropdownNode.cpp
--- a/Engine/Source/Editor/PropertyEditor/Private/DetailAdvancedDropdownNode.cpp
+++ b/Engine/Source/Editor/PropertyEditor/Private/DetailAdvancedDropdownNode.cpp
@@ -1,23 +1,26 @@
 // Copyright Epic Games, Inc. All Rights Reserved.
 
 #include "DetailAdvancedDropdownNode.h"
+
+#include "PropertyCustomizationHelpers.h"
+#include "SDetailRowIndent.h"
+#include "SDetailTableRowBase.h"
+#include "Styling/StyleColors.h"
+#include "UserInterface/PropertyEditor/PropertyEditorConstants.h"
 #include "Widgets/Images/SImage.h"
 #include "Widgets/Input/SButton.h"
-#include "SDetailTableRowBase.h"
-#include "PropertyCustomizationHelpers.h"
 
 class SAdvancedDropdownRow : public SDetailTableRowBase
 {
 public:
-	SLATE_BEGIN_ARGS( SAdvancedDropdownRow )
-		: _IsExpanded( false )
-		, _IsButtonEnabled( true )
-		, _ShouldShowAdvancedButton( false )
+	SLATE_BEGIN_ARGS(SAdvancedDropdownRow)
+		: _IsExpanded(false)
+		, _IsButtonEnabled(true)
 	{}
-		SLATE_ATTRIBUTE( bool, IsExpanded )
-		SLATE_ATTRIBUTE( bool, IsButtonEnabled )
-		SLATE_ARGUMENT( bool, ShouldShowAdvancedButton )
-		SLATE_EVENT( FOnClicked, OnClicked )
+		SLATE_ATTRIBUTE(bool, IsExpanded)
+		SLATE_ATTRIBUTE(bool, IsButtonEnabled)
+		SLATE_ARGUMENT(bool, IsVisible)
+		SLATE_EVENT(FOnClicked, OnClicked)
 	SLATE_END_ARGS()
 
 	/**
@@ -25,91 +28,107 @@ public:
 	 *
 	 * @param InArgs   A declaration from which to construct the widget
 	 */
-	void Construct( const FArguments& InArgs, IDetailsViewPrivate* InDetailsView, const TSharedRef<STableViewBase>& InOwnerTableView, bool bIsTopNode, bool bInDisplayShowAdvancedMessage)
+	void Construct( const FArguments& InArgs, IDetailsViewPrivate* InDetailsView, const TSharedRef<STableViewBase>& InOwnerTableView)
 	{
 		IsExpanded = InArgs._IsExpanded;
 		DetailsView = InDetailsView;
-		bDisplayShowAdvancedMessage = bInDisplayShowAdvancedMessage;
 
 		TSharedPtr<SWidget> ContentWidget;
-		if( bIsTopNode )
-		{
-			ContentWidget = SNew(SImage)
-				.Image(FEditorStyle::GetBrush("DetailsView.AdvancedDropdownBorder.Open"));
-		}
-		else if( InArgs._ShouldShowAdvancedButton )
+
+		if (InArgs._IsVisible)
 		{
-			ContentWidget = 		
-				SNew( SVerticalBox )
-				+ SVerticalBox::Slot()
-				.HAlign( HAlign_Center )
-				.AutoHeight()
+			ContentWidget = SNew(SHorizontalBox)
+			+ SHorizontalBox::Slot()
+			.HAlign(HAlign_Left)
+			.VAlign(VAlign_Fill)
+			.AutoWidth()
+			[
+				SNew(SDetailRowIndent, SharedThis(this))
+			]
+			+ SHorizontalBox::Slot()
+			.HAlign(HAlign_Left)
+			.VAlign(VAlign_Center)
+			.Padding(2, 0, 0, 0)
+			.AutoWidth()
+			[
+				SAssignNew(ExpanderButton, SButton)
+				.ButtonStyle(FCoreStyle::Get(), "NoBorder")
+				.VAlign(VAlign_Center)
+				.HAlign(HAlign_Center)
+				.ClickMethod(EButtonClickMethod::MouseDown)
+				.OnClicked(InArgs._OnClicked)
+				.IsEnabled(InArgs._IsButtonEnabled)
+				.ContentPadding(0)
+				.IsFocusable(false)
+				.ToolTipText(this, &SAdvancedDropdownRow::GetAdvancedPulldownToolTipText )
 				[
-					SNew( STextBlock )
-					.Text( NSLOCTEXT("DetailsView", "NoSimpleProperties", "Click the arrow to display advanced properties") )
-					.Font( IDetailLayoutBuilder::GetDetailFont() )
-					.Visibility( this, &SAdvancedDropdownRow::OnGetHelpTextVisibility )
-					.ColorAndOpacity(FLinearColor(1,1,1,.5))
+					SNew(SImage)
+					.Image(this, &SAdvancedDropdownRow::GetExpanderImage)
+					.ColorAndOpacity(FSlateColor::UseSubduedForeground())
 				]
-				+ SVerticalBox::Slot()
-				.AutoHeight()
+			]
+			+ SHorizontalBox::Slot()
+			.VAlign(VAlign_Center)
+			.HAlign(HAlign_Fill)
+			.Padding(4, 0, 0, 0)
+			[
+				SNew(SBox)
+				.VAlign(VAlign_Center)
+				.HeightOverride(PropertyEditorConstants::PropertyRowHeight)
 				[
-					SAssignNew(ExpanderButton, SButton)
-					.ButtonStyle(FEditorStyle::Get(), "NoBorder")
-					.HAlign(HAlign_Center)
-					.ContentPadding(2)
-					.OnClicked(InArgs._OnClicked)
-					.IsEnabled(InArgs._IsButtonEnabled)
-					.ToolTipText(this, &SAdvancedDropdownRow::GetAdvancedPulldownToolTipText )
-					[
-						SNew(SImage)
-						.Image(this, &SAdvancedDropdownRow::GetAdvancedPulldownImage)
-					]
-				];
-		}
-		else
-		{
-			ContentWidget = SNew(SSpacer);
-		}
-
-		const FSlateBrush* ScrollBarWellBrush = FAppStyle::Get().GetBrush("DetailsView.GridLine");
-		const FSlateBrush* BackgroundBrush = FAppStyle::Get().GetBrush("DetailsView.AdvancedDropdownBorder");
-
-		TWeakPtr<STableViewBase> OwnerTableViewWeak = InOwnerTableView;
-		auto GetScrollbarWellBrush = [OwnerTableViewWeak, ScrollBarWellBrush, BackgroundBrush]()
-		{
-			return SDetailTableRowBase::IsScrollBarVisible(OwnerTableViewWeak) ? ScrollBarWellBrush : BackgroundBrush;
-		};
-
-		ChildSlot
-		[
-			SNew(SBorder)
-			.BorderImage(FAppStyle::Get().GetBrush("DetailsView.GridLine"))
-			.Padding(bIsTopNode ? FMargin(0) : FMargin(0, 0, 0, 1))
+					SNew(STextBlock)
+					.Text(NSLOCTEXT("PropertyEditor", "Advanced", "Advanced"))
+					.Font(FAppStyle::Get().GetFontStyle(PropertyEditorConstants::PropertyFontStyle))
+					.TextStyle(FAppStyle::Get(), "DetailsView.CategoryTextStyle")
+				]
+			];
+
+			TWeakPtr<STableViewBase> OwnerTableViewWeak = InOwnerTableView;
+			auto GetScrollbarWellBrush = [this, OwnerTableViewWeak]()
+			{
+				return SDetailTableRowBase::IsScrollBarVisible(OwnerTableViewWeak) ?
+					FAppStyle::Get().GetBrush("DetailsView.GridLine") : 
+					FAppStyle::Get().GetBrush("DetailsView.CategoryMiddle");
+			};
+
+			auto GetScrollbarWellTint = [this, OwnerTableViewWeak]()
+			{
+				return SDetailTableRowBase::IsScrollBarVisible(OwnerTableViewWeak) ?
+					FStyleColors::White : 
+					this->GetRowBackgroundColor();
+			};
+
+			ChildSlot
 			[
-				SNew(SHorizontalBox)
-				+ SHorizontalBox::Slot()
-				.VAlign(VAlign_Fill)
-				.HAlign(HAlign_Fill)
+				SNew(SBorder)
+				.BorderImage(FAppStyle::Get().GetBrush("DetailsView.GridLine"))
+				.Padding(FMargin(0, 0, 0, 1))
 				[
-					SNew(SBorder)
-					.BorderImage(FAppStyle::Get().GetBrush("DetailsView.AdvancedDropdownBorder"))
-					.Padding(bIsTopNode ? FMargin(0, 0, 0, 2) : FMargin(0, 3, 0, 2))
+					SNew(SHorizontalBox)
+					+ SHorizontalBox::Slot()
+					.HAlign(HAlign_Fill)
 					[
-						ContentWidget.ToSharedRef()
+						SNew(SBorder)
+						.BorderImage(FAppStyle::Get().GetBrush("DetailsView.CategoryMiddle"))
+						.BorderBackgroundColor(this, &SAdvancedDropdownRow::GetRowBackgroundColor)
+						.Padding(0)
+						[
+							ContentWidget.ToSharedRef()
+						]
+					]
+					+ SHorizontalBox::Slot()
+					.HAlign(HAlign_Right)
+					.VAlign(VAlign_Fill)
+					.AutoWidth()
+					[
+						SNew(SBorder)
+						.BorderImage_Lambda(GetScrollbarWellBrush)
+						.BorderBackgroundColor_Lambda(GetScrollbarWellTint)
+						.Padding(FMargin(0, 0, SDetailTableRowBase::ScrollBarPadding, 0))
 					]
 				]
-				+ SHorizontalBox::Slot()
-				.VAlign(VAlign_Fill)
-				.HAlign(HAlign_Right)
-				.AutoWidth()
-				[
-					SNew(SBorder)
-					.BorderImage_Lambda(GetScrollbarWellBrush)
-					.Padding(FMargin(0, 0, SDetailTableRowBase::ScrollBarPadding, 0))
-				]
-			]
-		];
+			];
+		}
 
 		STableRow< TSharedPtr< FDetailTreeNode > >::ConstructInternal(
 			STableRow::FArguments()
@@ -121,28 +140,52 @@ public:
 
 private:
 
-	EVisibility OnGetHelpTextVisibility() const
-	{
-		return bDisplayShowAdvancedMessage && !IsExpanded.Get() ? EVisibility::Visible : EVisibility::Collapsed;
-	}
-
 	FText GetAdvancedPulldownToolTipText() const
 	{
 		return IsExpanded.Get() ? NSLOCTEXT("DetailsView", "HideAdvanced", "Hide Advanced") : NSLOCTEXT("DetailsView", "ShowAdvanced", "Show Advanced");
 	}
 
-	const FSlateBrush* GetAdvancedPulldownImage() const
+	const FSlateBrush* GetExpanderImage() const
 	{
-		if( ExpanderButton->IsHovered() )
+		const bool bIsItemExpanded = IsExpanded.Get();
+
+		FName ResourceName;
+		if (bIsItemExpanded)
 		{
-			return IsExpanded.Get() ? FEditorStyle::GetBrush("DetailsView.PulldownArrow.Up.Hovered") : FEditorStyle::GetBrush("DetailsView.PulldownArrow.Down.Hovered");
+			if (ExpanderButton->IsHovered())
+			{
+				static const FName ExpandedHoveredName = "TreeArrow_Expanded_Hovered";
+				ResourceName = ExpandedHoveredName;
+			}
+			else
+			{
+				static const FName ExpandedName = "TreeArrow_Expanded";
+				ResourceName = ExpandedName;
+			}
 		}
 		else
 		{
-			return IsExpanded.Get() ? FEditorStyle::GetBrush("DetailsView.PulldownArrow.Up") : FEditorStyle::GetBrush("DetailsView.PulldownArrow.Down");
+			if (ExpanderButton->IsHovered())
+			{
+				static const FName CollapsedHoveredName = "TreeArrow_Collapsed_Hovered";
+				ResourceName = CollapsedHoveredName;
+			}
+			else
+			{
+				static const FName CollapsedName = "TreeArrow_Collapsed";
+				ResourceName = CollapsedName;
+			}
 		}
+
+		return FAppStyle::Get().GetBrush(ResourceName);
 	}
 
+	FSlateColor GetRowBackgroundColor() const
+	{
+		return PropertyEditorConstants::GetRowBackgroundColor(0, this->IsHovered());
+	}
+
+
 private:
 	TAttribute<bool> IsExpanded;
 	TSharedPtr<SButton> ExpanderButton;
@@ -152,11 +195,10 @@ private:
 
 TSharedRef< ITableRow > FAdvancedDropdownNode::GenerateWidgetForTableView( const TSharedRef<STableViewBase>& OwnerTable, bool bAllowFavoriteSystem)
 {
-	return SNew( SAdvancedDropdownRow, ParentCategory.GetDetailsView(), OwnerTable, bIsTopNode, bDisplayShowAdvancedMessage )
-		.OnClicked( this, &FAdvancedDropdownNode::OnAdvancedDropDownClicked )
-		.IsButtonEnabled( IsEnabled )
-		.IsExpanded( IsExpanded )
-		.ShouldShowAdvancedButton( bShouldShowAdvancedButton );
+	return SNew(SAdvancedDropdownRow, ParentCategory.GetDetailsView(), OwnerTable)
+		.OnClicked(this, &FAdvancedDropdownNode::OnAdvancedDropDownClicked)
+		.IsButtonEnabled(IsEnabled)
+		.IsExpanded(IsExpanded);
 }
 
 bool FAdvancedDropdownNode::GenerateStandaloneWidget(FDetailWidgetRow& OutRow) const
diff --git a/Engine/Source/Editor/PropertyEditor/Private/DetailAdvancedDropdownNode.h b/Engine/Source/Editor/PropertyEditor/Private/DetailAdvancedDropdownNode.h
--- a/Engine/Source/Editor/PropertyEditor/Private/DetailAdvancedDropdownNode.h
+++ b/Engine/Source/Editor/PropertyEditor/Private/DetailAdvancedDropdownNode.h
@@ -15,25 +15,14 @@
 class FAdvancedDropdownNode : public FDetailTreeNode, public TSharedFromThis<FAdvancedDropdownNode>
 {
 public:
-	FAdvancedDropdownNode( TSharedRef<FDetailCategoryImpl> InParentCategory, const TAttribute<bool>& InExpanded, const TAttribute<bool>& InEnabled, bool bInShouldShowAdvancedButton, bool bInDisplayShowAdvancedMessage )
-		: ParentCategory( InParentCategory.Get() )
-		, IsEnabled( InEnabled )
-		, IsExpanded( InExpanded )
-		, bShouldShowAdvancedButton( bInShouldShowAdvancedButton )
-		, bIsTopNode( false )
-		, bDisplayShowAdvancedMessage( bInDisplayShowAdvancedMessage )
+	FAdvancedDropdownNode(TSharedRef<FDetailCategoryImpl> InParentCategory, const TAttribute<bool>& InExpanded, const TAttribute<bool>& InEnabled, bool bIsVisible)
+		: ParentCategory(InParentCategory.Get())
+		, IsEnabled(InEnabled)
+		, IsExpanded(InExpanded)
 	{
 		SetParentNode(InParentCategory);
 	}
 
-	FAdvancedDropdownNode( TSharedRef<FDetailCategoryImpl> InParentCategory, bool bInIsTopNode )
-		: ParentCategory( InParentCategory.Get() )
-		, bShouldShowAdvancedButton( false )
-		, bIsTopNode( bInIsTopNode )
-		, bDisplayShowAdvancedMessage( false  )
-	{
-		SetParentNode(InParentCategory);
-	}
 private:
 	/** IDetailTreeNode Interface */
 	virtual IDetailsView* GetNodeDetailsView() const override { return ParentCategory.GetNodeDetailsView(); }
@@ -59,7 +48,4 @@ private:
 	FDetailCategoryImpl& ParentCategory;
 	TAttribute<bool> IsEnabled;
 	TAttribute<bool> IsExpanded;
-	bool bShouldShowAdvancedButton;
-	bool bIsTopNode;
-	bool bDisplayShowAdvancedMessage;
 };
diff --git a/Engine/Source/Editor/PropertyEditor/Private/DetailCategoryBuilderImpl.cpp b/Engine/Source/Editor/PropertyEditor/Private/DetailCategoryBuilderImpl.cpp
--- a/Engine/Source/Editor/PropertyEditor/Private/DetailCategoryBuilderImpl.cpp
+++ b/Engine/Source/Editor/PropertyEditor/Private/DetailCategoryBuilderImpl.cpp
@@ -917,11 +917,10 @@ void FDetailCategoryImpl::GenerateChildrenForLayouts()
 	// Generate nodes for advanced dropdowns
 	if (AdvancedChildNodes.Num() > 0)
 	{
-		TAttribute<bool> ShowAdvanced(this, &FDetailCategoryImpl::ShouldShowAdvanced);
 		TAttribute<bool> IsEnabled(this, &FDetailCategoryImpl::IsAdvancedDropdownEnabled);
+		TAttribute<bool> IsExpanded(this, &FDetailCategoryImpl::ShouldShowAdvanced);
 
-		AdvancedDropdownNodeTop = MakeShared<FAdvancedDropdownNode>(AsShared(), true);
-		AdvancedDropdownNodeBottom = MakeShared<FAdvancedDropdownNode>(AsShared(), ShowAdvanced, IsEnabled, AdvancedChildNodes.Num() > 0, SimpleChildNodes.Num() == 0);
+		AdvancedDropdownNode = MakeShared<FAdvancedDropdownNode>(AsShared(), IsExpanded, IsEnabled, AdvancedChildNodes.Num() > 0);
 	}
 }
 
@@ -948,33 +947,31 @@ void FDetailCategoryImpl::GetGeneratedChildren(FDetailNodeList& OutChildren, boo
 		}
 	}
 
-	if (!bIgnoreAdvancedDropdown && ShouldShowAdvanced())
+	if (!bIgnoreAdvancedDropdown)
 	{
-		if (AdvancedDropdownNodeTop.IsValid())
+		if (AdvancedChildNodes.Num() > 0 && AdvancedDropdownNode.IsValid())
 		{
-			OutChildren.Add(AdvancedDropdownNodeTop.ToSharedRef());
+			OutChildren.Add(AdvancedDropdownNode.ToSharedRef());
 		}
 
-		for (TSharedRef<FDetailTreeNode>& Child : AdvancedChildNodes)
+		if (ShouldShowAdvanced())
 		{
-			if (bIgnoreVisibility || Child->GetVisibility() == ENodeVisibility::Visible)
+			for (TSharedRef<FDetailTreeNode>& Child : AdvancedChildNodes)
 			{
-				if (Child->ShouldShowOnlyChildren())
+				if (bIgnoreVisibility || Child->GetVisibility() == ENodeVisibility::Visible)
 				{
-					Child->GetChildren(OutChildren);
-				}
-				else
-				{
-					OutChildren.Add(Child);
+					if (Child->ShouldShowOnlyChildren())
+					{
+						Child->GetChildren(OutChildren);
+					}
+					else
+					{
+						OutChildren.Add(Child);
+					}
 				}
 			}
 		}
 	}
-
-	if (!bIgnoreAdvancedDropdown && AdvancedDropdownNodeBottom.IsValid())
-	{
-		OutChildren.Add(AdvancedDropdownNodeBottom.ToSharedRef());
-	}
 }
 
 void FDetailCategoryImpl::FilterNode(const FDetailFilter& InFilter)
@@ -1048,8 +1045,7 @@ void FDetailCategoryImpl::GenerateLayout()
 	// Reset all children
 	SimpleChildNodes.Empty();
 	AdvancedChildNodes.Empty();
-	AdvancedDropdownNodeTop.Reset();
-	AdvancedDropdownNodeBottom.Reset();
+	AdvancedDropdownNode.Reset();
 	InlinePropertyNode.Reset();
 
 	GenerateChildrenForLayouts();
diff --git a/Engine/Source/Editor/PropertyEditor/Private/DetailCategoryBuilderImpl.h b/Engine/Source/Editor/PropertyEditor/Private/DetailCategoryBuilderImpl.h
--- a/Engine/Source/Editor/PropertyEditor/Private/DetailCategoryBuilderImpl.h
+++ b/Engine/Source/Editor/PropertyEditor/Private/DetailCategoryBuilderImpl.h
@@ -383,10 +383,8 @@ private:
 	TArray< TSharedRef<FDetailTreeNode> > SimpleChildNodes;
 	/** All Advanced child nodes */
 	TArray< TSharedRef<FDetailTreeNode> > AdvancedChildNodes;
-	/** Advanced dropdown node (always shown) */
-	TSharedPtr<FDetailTreeNode> AdvancedDropdownNodeBottom;
-	/** Advanced dropdown node that is shown if the advanced dropdown is expanded */
-	TSharedPtr<FDetailTreeNode> AdvancedDropdownNodeTop;
+	/** Advanced dropdown node. */
+	TSharedPtr<FDetailTreeNode> AdvancedDropdownNode;
 	/** Delegate called when expansion of the category changes */
 	FOnBooleanValueChanged OnExpansionChangedDelegate;
 	/** The display name of the category */
diff --git a/Engine/Source/Editor/PropertyEditor/Private/PropertyEditorHelpers.cpp b/Engine/Source/Editor/PropertyEditor/Private/PropertyEditorHelpers.cpp
--- a/Engine/Source/Editor/PropertyEditor/Private/PropertyEditorHelpers.cpp
+++ b/Engine/Source/Editor/PropertyEditor/Private/PropertyEditorHelpers.cpp
@@ -91,7 +91,7 @@ void SPropertyValueWidget::Construct( const FArguments& InArgs, TSharedPtr<FProp
 				.AutoWidth()
 				.HAlign(HAlign_Center)
 				.VAlign(VAlign_Center)
-				.Padding( 2.0f, 1.0f )
+				.Padding( 2.0f, 0.0f )
 				[ 
 					RequiredButtons[ButtonIndex]
 				];
diff --git a/Engine/Source/Editor/PropertyEditor/Private/SDetailCategoryTableRow.cpp b/Engine/Source/Editor/PropertyEditor/Private/SDetailCategoryTableRow.cpp
--- a/Engine/Source/Editor/PropertyEditor/Private/SDetailCategoryTableRow.cpp
+++ b/Engine/Source/Editor/PropertyEditor/Private/SDetailCategoryTableRow.cpp
@@ -14,8 +14,6 @@ void SDetailCategoryTableRow::Construct(const FArguments& InArgs, TSharedRef<FDe
 	bIsInnerCategory = InArgs._InnerCategory;
 	bShowBorder = InArgs._ShowBorder;
 
-	const float VerticalPadding = bIsInnerCategory ? 6 : 8;
-
 	FDetailColumnSizeData& ColumnSizeData = InOwnerTreeNode->GetDetailsView()->GetColumnSizeData();
 
 	TSharedRef<SHorizontalBox> HeaderBox = SNew(SHorizontalBox)
@@ -29,18 +27,17 @@ void SDetailCategoryTableRow::Construct(const FArguments& InArgs, TSharedRef<FDe
 		+ SHorizontalBox::Slot()
 		.HAlign(HAlign_Left)
 		.VAlign(VAlign_Center)
-		.Padding(5, 0, 0, 0)
+		.Padding(2, 0, 0, 0)
 		.AutoWidth()
 		[
 			SNew(SDetailExpanderArrow, SharedThis(this))
 		]
 		+ SHorizontalBox::Slot()
 		.VAlign(VAlign_Center)
-		.Padding(12, VerticalPadding, 0, VerticalPadding)
+		.Padding(4, 0, 0, 0)
 		.FillWidth(1)
 		[
 			SNew(STextBlock)
-			.TransformPolicy(ETextTransformPolicy::ToUpper)
 			.Text(InArgs._DisplayName)
 			.Font(FAppStyle::Get().GetFontStyle(bIsInnerCategory ? PropertyEditorConstants::PropertyFontStyle : PropertyEditorConstants::CategoryFontStyle))
 			.TextStyle(FAppStyle::Get(), "DetailsView.CategoryTextStyle")
@@ -80,18 +77,25 @@ void SDetailCategoryTableRow::Construct(const FArguments& InArgs, TSharedRef<FDe
 			SNew(SHorizontalBox)
 			+ SHorizontalBox::Slot()
 			.HAlign(HAlign_Fill)
+			.Padding(0)
 			[
 				SNew(SBorder)
 				.BorderImage(this, &SDetailCategoryTableRow::GetBackgroundImage)
 				.BorderBackgroundColor(this, &SDetailCategoryTableRow::GetInnerBackgroundColor)
 				.Padding(0)
 				[
-					HeaderBox
+					SNew(SBox)
+					.MinDesiredHeight(PropertyEditorConstants::PropertyRowHeight)
+					.VAlign(VAlign_Center)
+					[
+						HeaderBox
+					]
 				]
 			]
 			+ SHorizontalBox::Slot()
 			.HAlign(HAlign_Right)
 			.VAlign(VAlign_Fill)
+			.Padding(0)
 			.AutoWidth()
 			[
 				SNew(SBorder)
@@ -124,14 +128,8 @@ const FSlateBrush* SDetailCategoryTableRow::GetBackgroundImage() const
 			return FAppStyle::Get().GetBrush("DetailsView.CategoryMiddle");
 		}
 
-		if (IsHovered())
-		{
-			return FAppStyle::Get().GetBrush("DetailsView.CategoryTop_Hovered");
-		}
-		else
-		{
-			return FAppStyle::Get().GetBrush("DetailsView.CategoryTop");
-		}
+		// intentionally no hover on outer categories
+		return FAppStyle::Get().GetBrush("DetailsView.CategoryTop");
 	}
 
 	return nullptr;
@@ -149,7 +147,7 @@ FSlateColor SDetailCategoryTableRow::GetInnerBackgroundColor() const
 
 		IndentLevel = FMath::Max(IndentLevel - 1, 0);
 
-		return PropertyEditorConstants::GetRowBackgroundColor(IndentLevel);
+		return PropertyEditorConstants::GetRowBackgroundColor(IndentLevel, this->IsHovered());
 	}
 
 	return FSlateColor(FLinearColor::White);
diff --git a/Engine/Source/Editor/PropertyEditor/Private/SDetailExpanderArrow.cpp b/Engine/Source/Editor/PropertyEditor/Private/SDetailExpanderArrow.cpp
--- a/Engine/Source/Editor/PropertyEditor/Private/SDetailExpanderArrow.cpp
+++ b/Engine/Source/Editor/PropertyEditor/Private/SDetailExpanderArrow.cpp
@@ -13,7 +13,7 @@ void SDetailExpanderArrow::Construct(const FArguments& InArgs, TSharedRef<SDetai
 	ChildSlot
 	[
 		SNew(SConstrainedBox)
-		.MinWidth(16)
+		.MinWidth(12)
 		.Visibility(this, &SDetailExpanderArrow::GetExpanderVisibility)
 		[
 			SAssignNew(ExpanderArrow, SButton)
@@ -22,7 +22,7 @@ void SDetailExpanderArrow::Construct(const FArguments& InArgs, TSharedRef<SDetai
 			.HAlign(HAlign_Center)
 			.ClickMethod(EButtonClickMethod::MouseDown)
 			.OnClicked(this, &SDetailExpanderArrow::OnExpanderClicked)
-			.ContentPadding(FMargin(0))
+			.ContentPadding(0)
 			.IsFocusable(false)
 			[
 				SNew(SImage)
@@ -41,7 +41,7 @@ EVisibility SDetailExpanderArrow::GetExpanderVisibility() const
 		return EVisibility::Collapsed;
 	}
 
-	return RowPtr->DoesItemHaveChildren() ? EVisibility::Visible : EVisibility::Collapsed;
+	return RowPtr->DoesItemHaveChildren() ? EVisibility::Visible : EVisibility::Hidden;
 }
 
 const FSlateBrush* SDetailExpanderArrow::GetExpanderImage() const
diff --git a/Engine/Source/Editor/PropertyEditor/Private/SDetailRowIndent.cpp b/Engine/Source/Editor/PropertyEditor/Private/SDetailRowIndent.cpp
--- a/Engine/Source/Editor/PropertyEditor/Private/SDetailRowIndent.cpp
+++ b/Engine/Source/Editor/PropertyEditor/Private/SDetailRowIndent.cpp
@@ -62,16 +62,12 @@ FOptionalSize SDetailRowIndent::GetIndentWidth() const
 	{
 		IndentLevel = RowPtr->GetIndentLevelForBackgroundColor();
 	}
+
 	return IndentLevel * 16.0f;
 }
 
 FSlateColor SDetailRowIndent::GetRowBackgroundColor(int32 IndentLevel) const
 {
 	TSharedPtr<SDetailTableRowBase> RowPtr = Row.Pin();
-	if (RowPtr.IsValid() && RowPtr->IsHovered())
-	{
-		return FAppStyle::Get().GetSlateColor("Colors.Header");
-	}
-
-	return PropertyEditorConstants::GetRowBackgroundColor(IndentLevel);
+	return PropertyEditorConstants::GetRowBackgroundColor(IndentLevel, RowPtr.IsValid() && RowPtr->IsHovered());
 }
diff --git a/Engine/Source/Editor/PropertyEditor/Private/SDetailSingleItemRow.cpp b/Engine/Source/Editor/PropertyEditor/Private/SDetailSingleItemRow.cpp
--- a/Engine/Source/Editor/PropertyEditor/Private/SDetailSingleItemRow.cpp
+++ b/Engine/Source/Editor/PropertyEditor/Private/SDetailSingleItemRow.cpp
@@ -8,7 +8,6 @@
 #include "Editor.h"
 #include "EditorMetadataOverrides.h"
 #include "IDetailDragDropHandler.h"
-#include "IDetailKeyframeHandler.h"
 #include "IDetailPropertyExtensionHandler.h"
 #include "ObjectPropertyNode.h"
 #include "PropertyEditorConstants.h"
@@ -31,8 +30,8 @@
 
 namespace DetailWidgetConstants
 {
-	const FMargin LeftRowPadding( 20.0f, 2.5f, 10.0f, 2.5f );
-	const FMargin RightRowPadding( 12.0f, 2.5f, 2.0f, 2.5f );
+	const FMargin LeftRowPadding( 20.0f, 0.0f, 10.0f, 0.0f );
+	const FMargin RightRowPadding( 12.0f, 0.0f, 2.0f, 0.0f );
 }
 
 namespace SDetailSingleItemRow_Helper
@@ -452,7 +451,7 @@ void SDetailSingleItemRow::Construct( const FArguments& InArgs, FDetailLayoutCus
 			NameColumnBox->AddSlot()
 				.HAlign(HAlign_Left)
 				.VAlign(VAlign_Center)
-				.Padding(8,0,0,0)
+				.Padding(2,0,0,0)
 				.AutoWidth()
 				[
 					SNew(SDetailExpanderArrow, SharedThis(this))
@@ -461,7 +460,7 @@ void SDetailSingleItemRow::Construct( const FArguments& InArgs, FDetailLayoutCus
 			NameColumnBox->AddSlot()
 				.VAlign(VAlign_Center)
 				.HAlign(HAlign_Left)
-				.Padding(3,0,0,0)
+				.Padding(2,0,0,0)
 				.AutoWidth()
 				[
 					SNew(SEditConditionWidget)
@@ -474,18 +473,38 @@ void SDetailSingleItemRow::Construct( const FArguments& InArgs, FDetailLayoutCus
 				NameColumnBox->AddSlot()
 					.HAlign(WidgetRow.NameWidget.HorizontalAlignment)
 					.VAlign(WidgetRow.NameWidget.VerticalAlignment)
-					.Padding(5,0,0,0)
+					.Padding(2,0,0,0)
 					[
 						NameWidget.ToSharedRef()
 					];
-		
-				Splitter->AddSlot()
-					.Value(ColumnSizeData.NameColumnWidth)
-					.OnSlotResized(ColumnSizeData.OnNameColumnResized)
+			}
+			else
+			{
+				// create whole row widget, which takes up both the Name and Value columns:
+				// | Name | Value | Right |
+				NameColumnBox->SetEnabled(IsEnabledAttribute);
+				NameColumnBox->AddSlot()
+					.HAlign(WidgetRow.WholeRowWidget.HorizontalAlignment)
+					.VAlign(WidgetRow.WholeRowWidget.VerticalAlignment)
+					.Padding(2,0,0,0)
 					[
-						NameColumnBox
+						WidgetRow.WholeRowWidget.Widget
 					];
+			}
 
+			Splitter->AddSlot()
+				.Value(ColumnSizeData.PropertyColumnWidth)
+				.OnSlotResized(ColumnSizeData.OnPropertyColumnResized)
+				[
+					SNew(SBox)
+					.MinDesiredHeight(PropertyEditorConstants::PropertyRowHeight)
+					[
+						NameColumnBox
+					]
+				];
+
+			if (bHasMultipleColumns)
+			{
 				// create Value column:
 				// | Name | Value | Right |
 				Splitter->AddSlot()
@@ -512,24 +531,6 @@ void SDetailSingleItemRow::Construct( const FArguments& InArgs, FDetailLayoutCus
 						]
 					];
 			}
-			else
-			{
-				NameColumnBox->SetEnabled(IsEnabledAttribute);
-				NameColumnBox->AddSlot()
-					.HAlign(WidgetRow.WholeRowWidget.HorizontalAlignment)
-					.VAlign(WidgetRow.WholeRowWidget.VerticalAlignment)
-					.Padding(5,0,0,0)
-					[
-						WidgetRow.WholeRowWidget.Widget
-					];
-
-				Splitter->AddSlot()
-					.Value(ColumnSizeData.PropertyColumnWidth)
-					.OnSlotResized(ColumnSizeData.OnPropertyColumnResized)
-					[
-						NameColumnBox
-					];
-			}
 
 			TArray<FPropertyRowExtensionButton> ExtensionButtons;
 
@@ -577,16 +578,11 @@ void SDetailSingleItemRow::Construct( const FArguments& InArgs, FDetailLayoutCus
 				SNew(SBorder)
 				.BorderImage(FAppStyle::Get().GetBrush("DetailsView.CategoryMiddle"))
 				.BorderBackgroundColor(this, &SDetailSingleItemRow::GetOuterBackgroundColor)
+				.HAlign(HAlign_Fill)
+				.VAlign(VAlign_Center)
 				.Padding(0)
 				[
-					SNew(SHorizontalBox)
-					+ SHorizontalBox::Slot()
-					.HAlign(HAlign_Fill)
-					.VAlign(VAlign_Center)
-					.Padding(5,0,5,0)
-					[
-						ToolbarBuilder.MakeWidget()
-					]
+					ToolbarBuilder.MakeWidget()
 				]
 			];
 		}
@@ -687,17 +683,12 @@ void SDetailSingleItemRow::OnResetToDefaultClicked() const
 /** Get the background color of the outer part of the row, which contains the edit condition and extension widgets. */
 FSlateColor SDetailSingleItemRow::GetOuterBackgroundColor() const
 {
-	if (IsHighlighted())
+	if (IsHighlighted() || DragOperation.IsValid())
 	{
 		return FAppStyle::Get().GetSlateColor("Colors.Hover");
 	}
 
-	if (IsHovered())
-	{
-		return FAppStyle::Get().GetSlateColor("Colors.Header");
-	}
-
-	return FAppStyle::Get().GetSlateColor("Colors.Panel");
+	return PropertyEditorConstants::GetRowBackgroundColor(0, this->IsHovered());
 }
 
 /** Get the background color of the inner part of the row, which contains the name and value widgets. */
@@ -708,18 +699,8 @@ FSlateColor SDetailSingleItemRow::GetInnerBackgroundColor() const
 		return FAppStyle::Get().GetSlateColor("Colors.Hover");
 	}
 
-	if (IsHovered())
-	{
-		return FAppStyle::Get().GetSlateColor("Colors.Header");
-	}
-	
-	if (DragOperation.IsValid())
-	{
-		return FAppStyle::Get().GetSlateColor("Colors.Hover");
-	}
-
-	int32 IndentLevel = GetIndentLevelForBackgroundColor();
-	return PropertyEditorConstants::GetRowBackgroundColor(IndentLevel);
+	const int32 IndentLevel = GetIndentLevelForBackgroundColor();
+	return PropertyEditorConstants::GetRowBackgroundColor(IndentLevel, this->IsHovered());
 }
 
 bool SDetailSingleItemRow::OnContextMenuOpening(FMenuBuilder& MenuBuilder)
@@ -1167,7 +1148,6 @@ void SDetailSingleItemRow::CreateGlobalExtensionWidgets(TArray<FPropertyRowExten
 	FOnGenerateGlobalRowExtensionArgs Args;
 	Args.OwnerTreeNode = OwnerTreeNode;
 
-	
 	if (Customization->HasPropertyNode())
 	{
 		Args.PropertyHandle = PropertyEditorHelpers::GetPropertyHandle(Customization->GetPropertyNode().ToSharedRef(), nullptr, nullptr);
diff --git a/Engine/Source/Editor/PropertyEditor/Private/SDetailsViewBase.cpp b/Engine/Source/Editor/PropertyEditor/Private/SDetailsViewBase.cpp
--- a/Engine/Source/Editor/PropertyEditor/Private/SDetailsViewBase.cpp
+++ b/Engine/Source/Editor/PropertyEditor/Private/SDetailsViewBase.cpp
@@ -520,21 +520,11 @@ void SDetailsViewBase::SetKeyframeHandler( TSharedPtr<class IDetailKeyframeHandl
 	KeyframeHandler = InKeyframeHandler;
 }
 
-TSharedPtr<IDetailKeyframeHandler> SDetailsViewBase::GetKeyframeHandler() 
-{
-	return KeyframeHandler;
-}
-
 void SDetailsViewBase::SetExtensionHandler(TSharedPtr<class IDetailPropertyExtensionHandler> InExtensionHandler)
 {
 	ExtensionHandler = InExtensionHandler;
 }
 
-TSharedPtr<IDetailPropertyExtensionHandler> SDetailsViewBase::GetExtensionHandler()
-{
-	return ExtensionHandler;
-}
-
 void SDetailsViewBase::SetGenericLayoutDetailsDelegate(FOnGetDetailCustomizationInstance OnGetGenericDetails)
 {
 	GenericLayoutDelegate = OnGetGenericDetails;
diff --git a/Engine/Source/Editor/PropertyEditor/Private/SDetailsViewBase.h b/Engine/Source/Editor/PropertyEditor/Private/SDetailsViewBase.h
--- a/Engine/Source/Editor/PropertyEditor/Private/SDetailsViewBase.h
+++ b/Engine/Source/Editor/PropertyEditor/Private/SDetailsViewBase.h
@@ -94,9 +94,9 @@ public:
 	virtual FIsPropertyEditingEnabled& GetIsPropertyEditingEnabledDelegate() override { return IsPropertyEditingEnabledDelegate; }
 	virtual bool IsPropertyEditingEnabled() const override;
 	virtual void SetKeyframeHandler(TSharedPtr<class IDetailKeyframeHandler> InKeyframeHandler) override;
-	virtual TSharedPtr<IDetailKeyframeHandler> GetKeyframeHandler() override;
+	virtual TSharedPtr<IDetailKeyframeHandler> GetKeyframeHandler() const override { return KeyframeHandler; }
 	virtual void SetExtensionHandler(TSharedPtr<class IDetailPropertyExtensionHandler> InExtensionHandler) override;
-	virtual TSharedPtr<IDetailPropertyExtensionHandler> GetExtensionHandler() override;
+	virtual TSharedPtr<IDetailPropertyExtensionHandler> GetExtensionHandler() const override { return ExtensionHandler; }
 	virtual bool IsPropertyVisible(const struct FPropertyAndParent& PropertyAndParent) const override;
 	virtual bool IsPropertyReadOnly(const struct FPropertyAndParent& PropertyAndParent) const override;
 	virtual bool IsCustomRowVisible(FName InRowName, FName InParentName) const override;
diff --git a/Engine/Source/Editor/PropertyEditor/Private/UserInterface/PropertyEditor/PropertyEditorConstants.cpp b/Engine/Source/Editor/PropertyEditor/Private/UserInterface/PropertyEditor/PropertyEditorConstants.cpp
--- a/Engine/Source/Editor/PropertyEditor/Private/UserInterface/PropertyEditor/PropertyEditorConstants.cpp
+++ b/Engine/Source/Editor/PropertyEditor/Private/UserInterface/PropertyEditor/PropertyEditorConstants.cpp
@@ -5,24 +5,19 @@
 #include "Styling/StyleColors.h"
 
 const FName PropertyEditorConstants::PropertyFontStyle( TEXT("PropertyWindow.NormalFont") );
-const FName PropertyEditorConstants::CategoryFontStyle( TEXT("PropertyWindow.BoldFont") );
+const FName PropertyEditorConstants::CategoryFontStyle( TEXT("DetailsView.CategoryFontStyle") );
 
 const FName PropertyEditorConstants::MD_Bitmask( TEXT("Bitmask") );
 const FName PropertyEditorConstants::MD_BitmaskEnum( TEXT("BitmaskEnum") );
 const FName PropertyEditorConstants::MD_UseEnumValuesAsMaskValuesInEditor( TEXT("UseEnumValuesAsMaskValuesInEditor") );
 
-const FSlateBrush* PropertyEditorConstants::GetOverlayBrush( const TSharedRef< class FPropertyEditor > PropertyEditor )
+const FSlateBrush* PropertyEditorConstants::GetOverlayBrush(const TSharedRef<class FPropertyEditor> PropertyEditor )
 {
 	return FEditorStyle::GetBrush( TEXT("PropertyWindow.NoOverlayColor") );
 }
 
-FSlateColor PropertyEditorConstants::GetRowBackgroundColor(int32 IndentLevel) 
+FSlateColor PropertyEditorConstants::GetRowBackgroundColor(int32 IndentLevel, bool IsHovered) 
 {
-	if (IndentLevel == 0)
-	{
-		return FAppStyle::Get().GetSlateColor("Colors.Panel");
-	}
-
 	int32 ColorIndex = 0;
 	int32 Increment = 1;
 
@@ -38,12 +33,16 @@ FSlateColor PropertyEditorConstants::GetRowBackgroundColor(int32 IndentLevel)
 
 	static const uint8 ColorOffsets[] =
 	{
-		2, (2 + 4), (6 + 6), (12 + 8)
+		0, 4, (4 + 2), (6 + 4), (10 + 6)
 	};
 
-	FColor BaseColor = FAppStyle::Get().GetSlateColor("Colors.Panel").GetSpecifiedColor().ToFColor(true);
+	const FSlateColor BaseSlateColor = IsHovered ? 
+		FAppStyle::Get().GetSlateColor("Colors.Header") : 
+		FAppStyle::Get().GetSlateColor("Colors.Panel");
+
+	const FColor BaseColor = BaseSlateColor.GetSpecifiedColor().ToFColor(true);
 
-	FColor ColorWithOffset(
+	const FColor ColorWithOffset(
 		BaseColor.R + ColorOffsets[ColorIndex], 
 		BaseColor.G + ColorOffsets[ColorIndex], 
 		BaseColor.B + ColorOffsets[ColorIndex]);
diff --git a/Engine/Source/Editor/PropertyEditor/Private/UserInterface/PropertyEditor/PropertyEditorConstants.h b/Engine/Source/Editor/PropertyEditor/Private/UserInterface/PropertyEditor/PropertyEditorConstants.h
--- a/Engine/Source/Editor/PropertyEditor/Private/UserInterface/PropertyEditor/PropertyEditorConstants.h
+++ b/Engine/Source/Editor/PropertyEditor/Private/UserInterface/PropertyEditor/PropertyEditorConstants.h
@@ -18,7 +18,9 @@ public:
 	static const FName MD_BitmaskEnum;
 	static const FName MD_UseEnumValuesAsMaskValuesInEditor;
 
-	static const FSlateBrush* GetOverlayBrush( const TSharedRef< class FPropertyEditor > PropertyEditor );
+	static constexpr float PropertyRowHeight = 26.0f;
 
-	static FSlateColor GetRowBackgroundColor(int32 IndentLevel);
+	static const FSlateBrush* GetOverlayBrush(const TSharedRef<class FPropertyEditor> PropertyEditor);
+
+	static FSlateColor GetRowBackgroundColor(int32 IndentLevel, bool IsHovered);
 };
diff --git a/Engine/Source/Editor/PropertyEditor/Private/UserInterface/PropertyEditor/SPropertyEditorAsset.cpp b/Engine/Source/Editor/PropertyEditor/Private/UserInterface/PropertyEditor/SPropertyEditorAsset.cpp
--- a/Engine/Source/Editor/PropertyEditor/Private/UserInterface/PropertyEditor/SPropertyEditorAsset.cpp
+++ b/Engine/Source/Editor/PropertyEditor/Private/UserInterface/PropertyEditor/SPropertyEditorAsset.cpp
@@ -454,7 +454,7 @@ void SPropertyEditorAsset::Construct(const FArguments& InArgs, const TSharedPtr<
 	
 	TSharedPtr<SHorizontalBox> CustomContentBox;
 
-	if(ShouldDisplayThumbnail(InArgs, ObjectClass))
+	if (ShouldDisplayThumbnail(InArgs, ObjectClass))
 	{
 		FObjectOrAssetData Value; 
 		GetValue( Value );
@@ -475,13 +475,13 @@ void SPropertyEditorAsset::Construct(const FArguments& InArgs, const TSharedPtr<
 		}
 
 		ValueContentBox->AddSlot()
-		.Padding(FMargin(0,5,5,5))
+		.Padding(0,3,5,0)
 		.AutoWidth()
 		.VAlign(VAlign_Center)
 		[
 			SNew(SBorder)
 			.Visibility(EVisibility::SelfHitTestInvisible)
-			.Padding(FMargin(0.0f, 0.0f, 4.0f, 4.0f))
+			.Padding(FMargin(0, 0, 4, 4))
 			.BorderImage(FAppStyle::Get().GetBrush("PropertyEditor.AssetTileItem.DropShadow"))
 			[
 				SNew(SOverlay)
@@ -513,7 +513,8 @@ void SPropertyEditorAsset::Construct(const FArguments& InArgs, const TSharedPtr<
 
 	
 		ValueContentBox->AddSlot()
-		.Padding(FMargin(0.0f, 5.0f))
+		.Padding(0)
+		.VAlign(VAlign_Center)
 		[
 			SNew(SVerticalBox)
 			+ SVerticalBox::Slot()
@@ -527,8 +528,7 @@ void SPropertyEditorAsset::Construct(const FArguments& InArgs, const TSharedPtr<
 			.AutoHeight()
 			[
 				SAssignNew(CustomContentBox, SHorizontalBox)
-
-				+SHorizontalBox::Slot()
+				+ SHorizontalBox::Slot()
 				.AutoWidth()
 				[
 					SAssignNew(ButtonBoxWrapper, SBox)
diff --git a/Engine/Source/Editor/PropertyEditor/Public/DetailsViewArgs.h b/Engine/Source/Editor/PropertyEditor/Public/DetailsViewArgs.h
--- a/Engine/Source/Editor/PropertyEditor/Public/DetailsViewArgs.h
+++ b/Engine/Source/Editor/PropertyEditor/Public/DetailsViewArgs.h
@@ -102,8 +102,8 @@ public:
 		, ViewIdentifier(NAME_None)
 		, NotifyHook(nullptr)
 		, NameAreaSettings(ActorsUseNameArea)
-		, ColumnWidth(.65f)
-		, RightColumnMinWidth(30)
+		, ColumnWidth(.7f)
+		, RightColumnMinWidth(22)
 		, bUpdatesFromSelection(false)
 		, bLockable(false)
 		, bAllowSearch(true)
diff --git a/Engine/Source/Editor/PropertyEditor/Public/IDetailsView.h b/Engine/Source/Editor/PropertyEditor/Public/IDetailsView.h
--- a/Engine/Source/Editor/PropertyEditor/Public/IDetailsView.h
+++ b/Engine/Source/Editor/PropertyEditor/Public/IDetailsView.h
@@ -158,11 +158,11 @@ public:
 	virtual void SetIsPropertyEditingEnabledDelegate( FIsPropertyEditingEnabled IsPropertyEditingEnabled ) = 0;
 	virtual FIsPropertyEditingEnabled& GetIsPropertyEditingEnabledDelegate() = 0;
 
-	virtual void SetKeyframeHandler( TSharedPtr<class IDetailKeyframeHandler> InKeyframeHandler ) = 0;
-	virtual TSharedPtr<IDetailKeyframeHandler> GetKeyframeHandler() = 0;
+	virtual void SetKeyframeHandler( TSharedPtr<IDetailKeyframeHandler> InKeyframeHandler ) = 0;
+	virtual TSharedPtr<IDetailKeyframeHandler> GetKeyframeHandler() const = 0;
 
-	virtual void SetExtensionHandler(TSharedPtr<class IDetailPropertyExtensionHandler> InExtensionHandler) = 0;
-	virtual TSharedPtr<class IDetailPropertyExtensionHandler> GetExtensionHandler() = 0;
+	virtual void SetExtensionHandler(TSharedPtr<IDetailPropertyExtensionHandler> InExtensionHandler) = 0;
+	virtual TSharedPtr<IDetailPropertyExtensionHandler> GetExtensionHandler() const = 0;
 
 	/**
 	 * @return true if property editing is enabled (based on the FIsPropertyEditingEnabled delegate)
