The right column in structure details views was not taking into account the RightColumnMinWidth setting in FDetailsViewArgs, causing 0 space to be allocated for it and a floating ellipsis to be displayed on top of the row.

Changed to use an overlay instead of an HBox for the scrollbar, which frees up ~16 pixels on the right.

diff --git a/Engine/Source/Editor/PropertyEditor/Private/SStructureDetailsView.cpp b/Engine/Source/Editor/PropertyEditor/Private/SStructureDetailsView.cpp
--- a/Engine/Source/Editor/PropertyEditor/Private/SStructureDetailsView.cpp
+++ b/Engine/Source/Editor/PropertyEditor/Private/SStructureDetailsView.cpp
@@ -35,7 +35,10 @@ UStruct* SStructureDetailsView::GetBaseScriptStruct() const
 void SStructureDetailsView::Construct(const FArguments& InArgs)
 {
 	DetailsViewArgs = InArgs._DetailsViewArgs;
-	
+
+	ColumnSizeData.SetValueColumnWidth(DetailsViewArgs.ColumnWidth);
+	ColumnSizeData.RightColumnMinWidth = DetailsViewArgs.RightColumnMinWidth;
+
 	CustomName = InArgs._CustomName;
 
 	// Create the root property now
@@ -185,13 +188,13 @@ void SStructureDetailsView::Construct(const FArguments& InArgs)
 			.FillHeight(1)
 			.Padding(0)
 			[
-				SNew( SHorizontalBox )
-				+ SHorizontalBox::Slot()
+				SNew( SOverlay )
+				+ SOverlay::Slot()
 				[
 					DetailTree.ToSharedRef()
 				]
-				+ SHorizontalBox::Slot()
-				.AutoWidth()
+				+ SOverlay::Slot()
+				.HAlign(HAlign_Right)
 				[
 					SNew( SBox )
 					.WidthOverride( 16.0f )
