Fixed bug where details rows were being sized incorrectly - the Name + Value width was accidentally being used for the Name column as a result of an earlier refactor.

diff --git a/Engine/Source/Editor/PropertyEditor/Private/SDetailSingleItemRow.cpp b/Engine/Source/Editor/PropertyEditor/Private/SDetailSingleItemRow.cpp
--- a/Engine/Source/Editor/PropertyEditor/Private/SDetailSingleItemRow.cpp
+++ b/Engine/Source/Editor/PropertyEditor/Private/SDetailSingleItemRow.cpp
@@ -477,34 +477,20 @@ void SDetailSingleItemRow::Construct( const FArguments& InArgs, FDetailLayoutCus
 					[
 						NameWidget.ToSharedRef()
 					];
-			}
-			else
-			{
-				// create whole row widget, which takes up both the Name and Value columns:
+
+				// create Name column:
 				// | Name | Value | Right |
-				NameColumnBox->SetEnabled(IsEnabledAttribute);
-				NameColumnBox->AddSlot()
-					.HAlign(WidgetRow.WholeRowWidget.HorizontalAlignment)
-					.VAlign(WidgetRow.WholeRowWidget.VerticalAlignment)
-					.Padding(2,0,0,0)
+				Splitter->AddSlot()
+					.Value(ColumnSizeData.NameColumnWidth)
+					.OnSlotResized(ColumnSizeData.OnNameColumnResized)
 					[
-						WidgetRow.WholeRowWidget.Widget
+						SNew(SBox)
+						.MinDesiredHeight(PropertyEditorConstants::PropertyRowHeight)
+						[
+							NameColumnBox
+						]
 					];
-			}
-
-			Splitter->AddSlot()
-				.Value(ColumnSizeData.PropertyColumnWidth)
-				.OnSlotResized(ColumnSizeData.OnPropertyColumnResized)
-				[
-					SNew(SBox)
-					.MinDesiredHeight(PropertyEditorConstants::PropertyRowHeight)
-					[
-						NameColumnBox
-					]
-				];
 
-			if (bHasMultipleColumns)
-			{
 				// create Value column:
 				// | Name | Value | Right |
 				Splitter->AddSlot()
@@ -531,6 +517,30 @@ void SDetailSingleItemRow::Construct( const FArguments& InArgs, FDetailLayoutCus
 						]
 					];
 			}
+			else
+			{
+				// create whole row widget, which takes up both the Name and Value columns:
+				// | Name | Value | Right |
+				NameColumnBox->SetEnabled(IsEnabledAttribute);
+				NameColumnBox->AddSlot()
+					.HAlign(WidgetRow.WholeRowWidget.HorizontalAlignment)
+					.VAlign(WidgetRow.WholeRowWidget.VerticalAlignment)
+					.Padding(2,0,0,0)
+					[
+						WidgetRow.WholeRowWidget.Widget
+					];
+
+				Splitter->AddSlot()
+					.Value(ColumnSizeData.PropertyColumnWidth)
+					.OnSlotResized(ColumnSizeData.OnPropertyColumnResized)
+					[
+						SNew(SBox)
+						.MinDesiredHeight(PropertyEditorConstants::PropertyRowHeight)
+						[
+							NameColumnBox
+						]
+					];
+			}
 
 			TArray<FPropertyRowExtensionButton> ExtensionButtons;
 
