Fix styling issues carried over from UE4 in the SME

diff --git a/Engine/Source/Editor/PropertyEditor/Private/SPerQualityLevelPropertiesWidget.cpp b/Engine/Source/Editor/PropertyEditor/Private/SPerQualityLevelPropertiesWidget.cpp
--- a/Engine/Source/Editor/PropertyEditor/Private/SPerQualityLevelPropertiesWidget.cpp
+++ b/Engine/Source/Editor/PropertyEditor/Private/SPerQualityLevelPropertiesWidget.cpp
@@ -49,24 +49,24 @@ TSharedRef<SWidget> SOverridePropertiesWidget::MakeOverrideWidget(FName InName,
 		[
 			SNew(SVerticalBox)
 			+ SVerticalBox::Slot()
-		.AutoHeight()
-		[
-			SAssignNew(HorizontalBox, SHorizontalBox)
-			+ SHorizontalBox::Slot()
-		.FillWidth(1.0f)
-		.Padding(0.0f, 0.0f, 2.0f, 2.0f)
-		.VAlign(VAlign_Center)
-		[
-			SNew(STextBlock)
-			.Font(IDetailLayoutBuilder::GetDetailFont())
-		.Text(InDisplayText)
-		]
-		]
-	+ SVerticalBox::Slot()
-		.AutoHeight()
-		[
-			OnGenerateWidget.Execute(InName)
-		]
+			.AutoHeight()
+			[
+				SAssignNew(HorizontalBox, SHorizontalBox)
+				+ SHorizontalBox::Slot()
+				.FillWidth(1.0f)
+				.Padding(0.0f, 0.0f, 2.0f, 2.0f)
+				.VAlign(VAlign_Center)
+				[
+					SNew(STextBlock)
+					.Font(IDetailLayoutBuilder::GetDetailFont())
+				.Text(InDisplayText)
+				]
+			]
+			+ SVerticalBox::Slot()
+			.AutoHeight()
+			[
+				OnGenerateWidget.Execute(InName)
+			]
 		];
 
 	if (InName != NAME_None)
@@ -78,17 +78,17 @@ TSharedRef<SWidget> SOverridePropertiesWidget::MakeOverrideWidget(FName InName,
 			[
 				SNew(SButton)
 				.ContentPadding(2.0f)
-			.ButtonStyle(FEditorStyle::Get(), "HoverHintOnly")
-			.OnClicked(this, &SOverridePropertiesWidget::RemoveEntry, InName)
-			.ToolTipText(FText::Format(NSLOCTEXT("SOverridePropertiesWidget", "RemoveOverrideFor", "Remove Override for {0}"), InDisplayText))
-			.ForegroundColor(FSlateColor::UseForeground())
-			.HAlign(HAlign_Center)
-			.VAlign(VAlign_Center)
-			.Content()
-			[
-				SNew(SImage)
-				.Image(FEditorStyle::GetBrush("Icons.Cross"))
-			]
+				.ButtonStyle(FEditorStyle::Get(), "HoverHintOnly")
+				.OnClicked(this, &SOverridePropertiesWidget::RemoveEntry, InName)
+				.ToolTipText(FText::Format(NSLOCTEXT("SOverridePropertiesWidget", "RemoveOverrideFor", "Remove Override for {0}"), InDisplayText))
+				.ForegroundColor(FSlateColor::UseForeground())
+				.HAlign(HAlign_Center)
+				.VAlign(VAlign_Center)
+				.Content()
+				[
+					SNew(SImage)
+					.Image(FEditorStyle::GetBrush("Icons.Delete"))
+				]
 			];
 	}
 	else
@@ -100,20 +100,20 @@ TSharedRef<SWidget> SOverridePropertiesWidget::MakeOverrideWidget(FName InName,
 			[
 				SNew(SComboButton)
 				.Visibility_Lambda([this]() { return bAddedMenuItem ? EVisibility::Visible : EVisibility::Hidden; })
-			.ButtonStyle(FEditorStyle::Get(), "HoverHintOnly")
-			.ContentPadding(2.0f)
-			.ForegroundColor(FSlateColor::UseForeground())
-			.HasDownArrow(false)
-			.ButtonContent()
-			[
-				SNew(SImage)
-				.Image(FEditorStyle::GetBrush("PropertyWindow.Button_AddToArray"))
-			]
-		.MenuContent()
-			[
-				InAddMenuBuilder.MakeWidget()
-			]
-		.ToolTipText(NSLOCTEXT("SOverridePropertiesWidget", "AddOverrideToolTip", "Add an override for a specific quality level"))
+				.ButtonStyle(FEditorStyle::Get(), "HoverHintOnly")
+				.ContentPadding(2.0f)
+				.ForegroundColor(FSlateColor::UseForeground())
+				.HasDownArrow(false)
+				.ButtonContent()
+				[
+					SNew(SImage)
+					.Image(FAppStyle::Get().GetBrush("Icons.PlusCircle"))
+				]
+				.MenuContent()
+				[
+					InAddMenuBuilder.MakeWidget()
+				]
+				.ToolTipText(NSLOCTEXT("SOverridePropertiesWidget", "AddOverrideToolTip", "Add an override for a specific quality level"))
 			];
 	}
 
diff --git a/Engine/Source/Editor/StaticMeshEditor/Private/StaticMeshEditorTools.cpp b/Engine/Source/Editor/StaticMeshEditor/Private/StaticMeshEditorTools.cpp
--- a/Engine/Source/Editor/StaticMeshEditor/Private/StaticMeshEditorTools.cpp
+++ b/Engine/Source/Editor/StaticMeshEditor/Private/StaticMeshEditorTools.cpp
@@ -3235,23 +3235,24 @@ void FLevelOfDetailSettingsLayout::AddToDetailsPanel( IDetailLayoutBuilder& Deta
 	}
 
 	LODSettingsCategory.AddCustomRow(LOCTEXT("QualityLevelMinLOD", "Quality Level Min LOD"))
+	.Visibility(bDisablePerPlatformMinLod ? EVisibility::Visible : EVisibility::Collapsed)
 	.RowTag("QualityLevelMinLOD")
 	.NameContent()
 		[
 			SNew(STextBlock)
 			.Font(IDetailLayoutBuilder::GetDetailFont())
-		.Text(LOCTEXT("QualityLevelMinLOD", "Quality Level Min LOD"))
+			.Text(LOCTEXT("QualityLevelMinLOD", "Quality Level Min LOD"))
 		]
-	.ValueContent()
+		.ValueContent()
 		.MinDesiredWidth((float)(StaticMesh->GetQualityLevelMinLOD().PerQuality.Num() + 1)*125.0f)
 		.MaxDesiredWidth((float)((int32)QualityLevelProperty::EQualityLevels::Num + 1)*125.0f)
 		[
 			SNew(SPerQualityLevelPropertiesWidget)
 			.IsEnabled(FLevelOfDetailSettingsLayout::GetLODCount() > 1 && bDisablePerPlatformMinLod)
-		.OnGenerateWidget(this, &FLevelOfDetailSettingsLayout::GetMinQualityLevelLODWidget)
-		.OnAddEntry(this, &FLevelOfDetailSettingsLayout::AddMinLODQualityLevelOverride)
-		.OnRemoveEntry(this, &FLevelOfDetailSettingsLayout::RemoveMinLODQualityLevelOverride)
-		.EntryNames(this, &FLevelOfDetailSettingsLayout::GetMinQualityLevelLODOverrideNames)
+			.OnGenerateWidget(this, &FLevelOfDetailSettingsLayout::GetMinQualityLevelLODWidget)
+			.OnAddEntry(this, &FLevelOfDetailSettingsLayout::AddMinLODQualityLevelOverride)
+			.OnRemoveEntry(this, &FLevelOfDetailSettingsLayout::RemoveMinLODQualityLevelOverride)
+			.EntryNames(this, &FLevelOfDetailSettingsLayout::GetMinQualityLevelLODOverrideNames)
 		];
 
 	{
@@ -3319,6 +3320,8 @@ void FLevelOfDetailSettingsLayout::AddToDetailsPanel( IDetailLayoutBuilder& Deta
 		SNew(SButton)
 		.OnClicked(this, &FLevelOfDetailSettingsLayout::OnApply)
 		.IsEnabled(this, &FLevelOfDetailSettingsLayout::IsApplyNeeded)
+		.VAlign(VAlign_Center)
+		.HAlign(HAlign_Center)
 		[
 			SNew( STextBlock )
 			.Text(LOCTEXT("ApplyChanges", "Apply Changes"))
@@ -3623,10 +3626,12 @@ void FLevelOfDetailSettingsLayout::AddLODLevelCategories( IDetailLayoutBuilder&
 				.ValueContent()
 				.HAlign(HAlign_Left)
 				[
-					SNew(SButton)
+ 					SNew(SButton)
 					.OnClicked(this, &FLevelOfDetailSettingsLayout::OnRemoveLOD, LODIndex)
 					.IsEnabled(this, &FLevelOfDetailSettingsLayout::CanRemoveLOD, LODIndex)
 					.ToolTipText( LOCTEXT("RemoveLOD_ToolTip", "Removes this LOD from the Static Mesh") )
+					.VAlign(VAlign_Center)
+					.HAlign(HAlign_Center)
 					[
 						SNew(STextBlock)
 						.Text( LOCTEXT("RemoveLOD", "Remove LOD") )
@@ -4792,6 +4797,8 @@ void FNaniteSettingsLayout::AddToDetailsPanel(IDetailLayoutBuilder& DetailBuilde
 		SNew(SButton)
 		.OnClicked(this, &FNaniteSettingsLayout::OnApply)
 		.IsEnabled(this, &FNaniteSettingsLayout::IsApplyNeeded)
+		.VAlign(VAlign_Center)
+		.HAlign(HAlign_Center)
 		[
 			SNew(STextBlock)
 			.Text(LOCTEXT("ApplyChanges", "Apply Changes"))
