NameContent() in details rows now correctly respects user-defined HAlign, allowing for labels to contain an HBox with both left and right aligned slots.

diff --git a/Engine/Source/Editor/PropertyEditor/Private/DetailPropertyRow.cpp b/Engine/Source/Editor/PropertyEditor/Private/DetailPropertyRow.cpp
--- a/Engine/Source/Editor/PropertyEditor/Private/DetailPropertyRow.cpp
+++ b/Engine/Source/Editor/PropertyEditor/Private/DetailPropertyRow.cpp
@@ -738,7 +738,7 @@ void FDetailPropertyRow::SetWidgetRowProperties(FDetailWidgetRow& Row) const
 void FDetailPropertyRow::MakeNameOrKeyWidget( FDetailWidgetRow& Row, const TSharedPtr<FDetailWidgetRow> InCustomRow ) const
 {
 	EVerticalAlignment VerticalAlignment = VAlign_Center;
-	EHorizontalAlignment HorizontalAlignment = HAlign_Left;
+	EHorizontalAlignment HorizontalAlignment = HAlign_Fill;
 
 	// We will only use key widgets for non-struct keys
 	const bool bHasKeyNode = PropertyKeyEditor.IsValid() && !PropertyHandle->HasMetaData(TEXT("ReadOnlyKeys"));
@@ -797,7 +797,6 @@ void FDetailPropertyRow::MakeNameOrKeyWidget( FDetailWidgetRow& Row, const TShar
 
 	SHorizontalBox::FSlot* SlotPointer = nullptr;
 	NameHorizontalBox->AddSlot()
-	.HAlign(HAlign_Left)
 	.Expose(SlotPointer)
 	[
 		NameWidget.ToSharedRef()
@@ -809,7 +808,7 @@ void FDetailPropertyRow::MakeNameOrKeyWidget( FDetailWidgetRow& Row, const TShar
 	}
 	else if (InCustomRow.IsValid())
 	{
-		//Allow custom name slot to fill all the area. If the user add a SHorizontalBox with left and right align slot
+		// Allow custom name slots to fill all of the area. Eg., the user adds a SHorizontalBox with left and right align slots.
 		SlotPointer->SetFillWidth(1.0f);
 	}
 	else
