Details value widgets are no longer sometimes stretched due to property rows having a minimum height.

diff --git a/Engine/Source/Editor/PropertyEditor/Private/SDetailSingleItemRow.cpp b/Engine/Source/Editor/PropertyEditor/Private/SDetailSingleItemRow.cpp
--- a/Engine/Source/Editor/PropertyEditor/Private/SDetailSingleItemRow.cpp
+++ b/Engine/Source/Editor/PropertyEditor/Private/SDetailSingleItemRow.cpp
@@ -484,11 +484,7 @@ void SDetailSingleItemRow::Construct( const FArguments& InArgs, FDetailLayoutCus
 					.Value(ColumnSizeData.NameColumnWidth)
 					.OnSlotResized(ColumnSizeData.OnNameColumnResized)
 					[
-						SNew(SBox)
-						.MinDesiredHeight(PropertyEditorConstants::PropertyRowHeight)
-						[
-							NameColumnBox
-						]
+						NameColumnBox
 					];
 
 				// create Value column:
@@ -534,11 +530,7 @@ void SDetailSingleItemRow::Construct( const FArguments& InArgs, FDetailLayoutCus
 					.Value(ColumnSizeData.PropertyColumnWidth)
 					.OnSlotResized(ColumnSizeData.OnPropertyColumnResized)
 					[
-						SNew(SBox)
-						.MinDesiredHeight(PropertyEditorConstants::PropertyRowHeight)
-						[
-							NameColumnBox
-						]
+						NameColumnBox
 					];
 			}
 
@@ -630,33 +622,38 @@ void SDetailSingleItemRow::Construct( const FArguments& InArgs, FDetailLayoutCus
 		.BorderImage(FAppStyle::Get().GetBrush("DetailsView.GridLine"))
 		.Padding(FMargin(0,0,0,1))
 		[
-			SNew( SHorizontalBox )
-			+ SHorizontalBox::Slot()
-			.HAlign(HAlign_Fill)
+			SNew(SBox)
+			.MinDesiredHeight(PropertyEditorConstants::PropertyRowHeight)
 			[
-				SNew( SBorder )
-				.BorderImage(FAppStyle::Get().GetBrush("DetailsView.Highlight"))
-				.Padding_Lambda(GetHighlightBorderPadding)
+				SNew( SHorizontalBox )
+				+ SHorizontalBox::Slot()
+				.HAlign(HAlign_Fill)
+				.VAlign(VAlign_Fill)
 				[
 					SNew( SBorder )
-					.BorderImage(FAppStyle::Get().GetBrush("DetailsView.CategoryMiddle"))
-					.BorderBackgroundColor(this, &SDetailSingleItemRow::GetOuterBackgroundColor)
-					.Padding(0)
-					.Clipping(EWidgetClipping::ClipToBounds)
+					.BorderImage(FAppStyle::Get().GetBrush("DetailsView.Highlight"))
+					.Padding_Lambda(GetHighlightBorderPadding)
 					[
-						Widget
+						SNew( SBorder )
+						.BorderImage(FAppStyle::Get().GetBrush("DetailsView.CategoryMiddle"))
+						.BorderBackgroundColor(this, &SDetailSingleItemRow::GetOuterBackgroundColor)
+						.Padding(0)
+						.Clipping(EWidgetClipping::ClipToBounds)
+						[
+							Widget
+						]
 					]
 				]
-			]
-			+ SHorizontalBox::Slot()
-			.HAlign(HAlign_Right)
-			.VAlign(VAlign_Fill)
-			.AutoWidth()
-			[
-				SNew( SBorder )
-				.BorderImage_Lambda(GetScrollbarWellBrush)
-				.BorderBackgroundColor_Lambda(GetScrollbarWellTint)
-				.Padding(FMargin(0, 0, SDetailTableRowBase::ScrollBarPadding, 0))
+				+ SHorizontalBox::Slot()
+				.HAlign(HAlign_Right)
+				.VAlign(VAlign_Fill)
+				.AutoWidth()
+				[
+					SNew( SBorder )
+					.BorderImage_Lambda(GetScrollbarWellBrush)
+					.BorderBackgroundColor_Lambda(GetScrollbarWellTint)
+					.Padding(FMargin(0, 0, SDetailTableRowBase::ScrollBarPadding, 0))
+				]
 			]
 		]
 	];
diff --git a/Engine/Source/Editor/PropertyEditor/Public/DetailWidgetRow.h b/Engine/Source/Editor/PropertyEditor/Public/DetailWidgetRow.h
--- a/Engine/Source/Editor/PropertyEditor/Public/DetailWidgetRow.h
+++ b/Engine/Source/Editor/PropertyEditor/Public/DetailWidgetRow.h
@@ -109,7 +109,7 @@ public:
 
 	FDetailWidgetRow()
 		: NameWidget( *this, 0.0f, 0.0f, HAlign_Left, VAlign_Center )
-		, ValueWidget( *this, DefaultValueMinWidth, DefaultValueMaxWidth, HAlign_Left, VAlign_Fill )
+		, ValueWidget( *this, DefaultValueMinWidth, DefaultValueMaxWidth, HAlign_Left, VAlign_Center )
 		, ExtensionWidget( *this, 0.0f, 0.0f, HAlign_Right, VAlign_Center)
 		, WholeRowWidget( *this, 0.0f, 0.0f, HAlign_Fill, VAlign_Fill )
 		, VisibilityAttr( EVisibility::Visible )
