Advanced dropdowns no longer randomly disappear when scrolling.

This was due to the IsVisible variable not being initialized correctly.

Advanced group is now clickable everywhere.

diff --git a/Engine/Source/Editor/PropertyEditor/Private/DetailAdvancedDropdownNode.cpp b/Engine/Source/Editor/PropertyEditor/Private/DetailAdvancedDropdownNode.cpp
--- a/Engine/Source/Editor/PropertyEditor/Private/DetailAdvancedDropdownNode.cpp
+++ b/Engine/Source/Editor/PropertyEditor/Private/DetailAdvancedDropdownNode.cpp
@@ -16,6 +16,7 @@ public:
 	SLATE_BEGIN_ARGS(SAdvancedDropdownRow)
 		: _IsExpanded(false)
 		, _IsButtonEnabled(true)
+		, _IsVisible(true)
 	{}
 		SLATE_ATTRIBUTE(bool, IsExpanded)
 		SLATE_ATTRIBUTE(bool, IsButtonEnabled)
@@ -32,6 +33,7 @@ public:
 	{
 		IsExpanded = InArgs._IsExpanded;
 		DetailsView = InDetailsView;
+		OnClicked = InArgs._OnClicked;
 
 		TSharedPtr<SWidget> ContentWidget;
 
@@ -185,10 +187,23 @@ private:
 		return PropertyEditorConstants::GetRowBackgroundColor(0, this->IsHovered());
 	}
 
+	FReply OnMouseButtonDown(const FGeometry& MyGeometry, const FPointerEvent& MouseEvent) override
+	{
+		if (MouseEvent.GetEffectingButton() == EKeys::LeftMouseButton)
+		{
+			if (OnClicked.IsBound())
+			{
+				return OnClicked.Execute();
+			}
+		}
+		
+		return FReply::Unhandled();
+	}
 
 private:
 	TAttribute<bool> IsExpanded;
 	TSharedPtr<SButton> ExpanderButton;
+	FOnClicked OnClicked;
 	bool bDisplayShowAdvancedMessage;
 	IDetailsViewPrivate* DetailsView;
 };
@@ -198,7 +213,8 @@ TSharedRef< ITableRow > FAdvancedDropdownNode::GenerateWidgetForTableView( const
 	return SNew(SAdvancedDropdownRow, ParentCategory.GetDetailsView(), OwnerTable)
 		.OnClicked(this, &FAdvancedDropdownNode::OnAdvancedDropDownClicked)
 		.IsButtonEnabled(IsEnabled)
-		.IsExpanded(IsExpanded);
+		.IsExpanded(IsExpanded)
+		.IsVisible(IsVisible);
 }
 
 bool FAdvancedDropdownNode::GenerateStandaloneWidget(FDetailWidgetRow& OutRow) const
diff --git a/Engine/Source/Editor/PropertyEditor/Private/DetailAdvancedDropdownNode.h b/Engine/Source/Editor/PropertyEditor/Private/DetailAdvancedDropdownNode.h
--- a/Engine/Source/Editor/PropertyEditor/Private/DetailAdvancedDropdownNode.h
+++ b/Engine/Source/Editor/PropertyEditor/Private/DetailAdvancedDropdownNode.h
@@ -19,6 +19,7 @@ public:
 		: ParentCategory(InParentCategory.Get())
 		, IsEnabled(InEnabled)
 		, IsExpanded(InExpanded)
+		, IsVisible(bIsVisible)
 	{
 		SetParentNode(InParentCategory);
 	}
@@ -48,4 +49,5 @@ private:
 	FDetailCategoryImpl& ParentCategory;
 	TAttribute<bool> IsEnabled;
 	TAttribute<bool> IsExpanded;
+	bool IsVisible;
 };
