Section buttons changed from a segmented control to a wrap box of buttons instead. Sections can now be multi-selected by holding Ctrl.

Section selector visibility can now be toggled from the options dropdown and is persisted between sessions.

diff --git a/Engine/Source/Editor/EditorStyle/Private/StarshipStyle.cpp b/Engine/Source/Editor/EditorStyle/Private/StarshipStyle.cpp
--- a/Engine/Source/Editor/EditorStyle/Private/StarshipStyle.cpp
+++ b/Engine/Source/Editor/EditorStyle/Private/StarshipStyle.cpp
@@ -3215,6 +3215,15 @@ void FStarshipEditorStyle::FStyle::SetupPropertyEditorStyles()
 		);
 
 		Set( "DetailsView.GridLine", new FSlateColorBrush(FStyleColors::Recessed) );
+		Set( "DetailsView.SectionButton", FCheckBoxStyle( FStarshipCoreStyle::GetCoreStyle().GetWidgetStyle<FCheckBoxStyle>("ToggleButtonCheckbox"))
+			.SetUncheckedImage(FSlateRoundedBoxBrush(FStyleColors::Header, 4.0f, FStyleColors::Input, 1.0f))
+			.SetUncheckedHoveredImage(FSlateRoundedBoxBrush(FStyleColors::Hover, 4.0f, FStyleColors::Input, 1.0f))
+			.SetUncheckedPressedImage(FSlateRoundedBoxBrush(FStyleColors::Hover, 4.0f, FStyleColors::Input, 1.0f))
+			.SetCheckedImage(FSlateRoundedBoxBrush(FStyleColors::Primary, 4.0f, FStyleColors::Input, 1.0f))
+			.SetCheckedHoveredImage(FSlateRoundedBoxBrush(FStyleColors::PrimaryHover, 4.0f, FStyleColors::Input, 1.0f))
+			.SetCheckedPressedImage(FSlateRoundedBoxBrush(FStyleColors::PrimaryHover, 4.0f, FStyleColors::Input, 1.0f))
+			.SetPadding(FMargin(16, 4))
+		);
 
 		Set( "DetailsView.CategoryFontStyle", FStyleFonts::Get().SmallBold);
 		Set( "DetailsView.CategoryTextStyle", FTextBlockStyle(NormalText)
diff --git a/Engine/Source/Editor/PropertyEditor/Private/DetailCategoryBuilderImpl.cpp b/Engine/Source/Editor/PropertyEditor/Private/DetailCategoryBuilderImpl.cpp
--- a/Engine/Source/Editor/PropertyEditor/Private/DetailCategoryBuilderImpl.cpp
+++ b/Engine/Source/Editor/PropertyEditor/Private/DetailCategoryBuilderImpl.cpp
@@ -986,7 +986,8 @@ void FDetailCategoryImpl::FilterNode(const FDetailFilter& InFilter)
 		return;
 	}
 
-	if (!InFilter.VisibleSections.IsEmpty())
+	// only apply the section filter if the user hasn't typed anything
+	if (InFilter.FilterStrings.IsEmpty() && !InFilter.VisibleSections.IsEmpty())
 	{
 		const UStruct* BaseStruct = GetParentBaseStructure();
 		if (BaseStruct != nullptr)
@@ -994,8 +995,29 @@ void FDetailCategoryImpl::FilterNode(const FDetailFilter& InFilter)
 			static FName PropertyEditor("PropertyEditor");
 			const FPropertyEditorModule& PropertyModule = FModuleManager::GetModuleChecked<FPropertyEditorModule>(PropertyEditor);
 
-			FName SectionName = PropertyModule.FindSectionForCategory(BaseStruct, CategoryName);
-			if (SectionName.IsNone() || !InFilter.VisibleSections.Contains(SectionName))
+			TArray<TSharedPtr<FPropertySection>> PropertySections = PropertyModule.FindSectionsForCategory(BaseStruct, CategoryName);
+			if (PropertySections.IsEmpty())
+			{
+				return;
+			}
+
+			// if this property is not in any visible section, hide it
+			bool bFound = false;
+			for (FName VisibleSection : InFilter.VisibleSections)
+			{
+				auto Predicate = [VisibleSection](const TSharedPtr<FPropertySection>& Elem)
+				{
+					return Elem->GetName() == VisibleSection;
+				};
+
+				if (PropertySections.ContainsByPredicate(Predicate))
+				{
+					bFound = true;
+					break;
+				}
+			}
+
+			if (!bFound)
 			{
 				return;
 			}
diff --git a/Engine/Source/Editor/PropertyEditor/Private/DetailsViewConfig.h b/Engine/Source/Editor/PropertyEditor/Private/DetailsViewConfig.h
--- a/Engine/Source/Editor/PropertyEditor/Private/DetailsViewConfig.h
+++ b/Engine/Source/Editor/PropertyEditor/Private/DetailsViewConfig.h
@@ -41,11 +41,12 @@ public:
 	UPROPERTY()
 	bool bShowOnlyModified { false };
 
+	/** Show sections. */
 	UPROPERTY()
-	float ValueColumnWidth { 0 };
+	bool bShowSections { true };
 
 	UPROPERTY()
-	float RightColumnWidth { 0 };
+	float ValueColumnWidth { 0 };
 };
 
 UCLASS(EditorConfig="DetailsView")
diff --git a/Engine/Source/Editor/PropertyEditor/Private/SDetailsView.cpp b/Engine/Source/Editor/PropertyEditor/Private/SDetailsView.cpp
--- a/Engine/Source/Editor/PropertyEditor/Private/SDetailsView.cpp
+++ b/Engine/Source/Editor/PropertyEditor/Private/SDetailsView.cpp
@@ -2,20 +2,20 @@
 
 
 #include "SDetailsView.h"
+
+#include "Classes/EditorStyleSettings.h"
 #include "DetailLayoutBuilderImpl.h"
 #include "DetailsViewGenericObjectFilter.h"
 #include "DetailsViewPropertyGenerationUtilities.h"
 #include "Editor.h"
 #include "EditorMetadataOverrides.h"
-#include "ObjectPropertyNode.h"
-#include "PropertyEditorHelpers.h"
-#include "SDetailNameArea.h"
-
-#include "Classes/EditorStyleSettings.h"
 #include "Framework/MultiBox/MultiBoxBuilder.h"
 #include "GameFramework/Actor.h"
 #include "Misc/ConfigCacheIni.h"
 #include "Modules/ModuleManager.h"
+#include "ObjectPropertyNode.h"
+#include "PropertyEditorHelpers.h"
+#include "SDetailNameArea.h"
 #include "Styling/StyleColors.h"
 #include "UserInterface/PropertyDetails/PropertyDetailsUtilities.h"
 #include "Widgets/Colors/SColorPicker.h"
@@ -24,6 +24,7 @@
 #include "Widgets/Input/SComboButton.h"
 #include "Widgets/Input/SSearchBox.h"
 #include "Widgets/Input/SSegmentedControl.h"
+#include "Widgets/Layout/SWrapBox.h"
 
 #define LOCTEXT_NAMESPACE "SDetailsView"
 
@@ -35,6 +36,8 @@ SDetailsView::~SDetailsView()
 		SaveExpandedItems(RootNode.ToSharedRef());
 	}
 
+	SaveViewConfig();
+
 	FEditorDelegates::PostUndoRedo.Remove(PostUndoRedoDelegateHandle);
 };
 
@@ -47,6 +50,17 @@ void SDetailsView::Construct(const FArguments& InArgs, const FDetailsViewArgs& I
 {
 	DetailsViewArgs = InDetailsViewArgs;
 
+	const FDetailsViewConfig& ViewConfig = GetMutableViewConfig();
+	if (ViewConfig.ValueColumnWidth != 0 && ViewConfig.ValueColumnWidth != DetailsViewArgs.ColumnWidth)
+	{
+		DetailsViewArgs.ColumnWidth = ViewConfig.ValueColumnWidth;
+	}
+
+	if (DetailsViewArgs.bShowSectionSelector)
+	{
+		DetailsViewArgs.bShowSectionSelector = ViewConfig.bShowSections;
+	}
+
 	ColumnSizeData.SetValueColumnWidth(DetailsViewArgs.ColumnWidth);
 	ColumnSizeData.RightColumnMinWidth = DetailsViewArgs.RightColumnMinWidth;
 
@@ -185,7 +199,20 @@ void SDetailsView::Construct(const FArguments& InArgs, const FDetailsViewArgs& I
 			),
 			NAME_None,
 			EUserInterfaceActionType::ToggleButton 
-			);
+		);
+
+		DetailViewOptions.AddMenuEntry(
+			LOCTEXT("ShowSections", "Show Sections"),
+			LOCTEXT("ShowSections_ToolTip", "Shows the sections list."),
+			FSlateIcon(),
+			FUIAction( 
+				FExecuteAction::CreateSP(this, &SDetailsView::OnShowSectionsClicked),
+				FCanExecuteAction(),
+				FIsActionChecked::CreateSP(this, &SDetailsView::IsShowSectionsChecked)
+			),
+			NAME_None,
+			EUserInterfaceActionType::ToggleButton 
+		);
 
 		DetailViewOptions.AddMenuEntry(
 			LOCTEXT("CollapseAll", "Collapse All Categories"),
@@ -314,17 +341,15 @@ void SDetailsView::Construct(const FArguments& InArgs, const FDetailsViewArgs& I
 			FilterRowHBox
 		];
 
-	if (DetailsViewArgs.bShowSectionSelector)
-	{
-		FilterRowVBox->AddSlot()
-			.Padding(30, 2, 30, 7)
-			.AutoHeight()
-			[
-				SAssignNew(SectionSelectorBox, SBox)
-			];
+	FilterRowVBox->AddSlot()
+		.Padding(30, 2, 30, 7)
+		.AutoHeight()
+		[
+			SAssignNew(SectionSelectorBox, SWrapBox)
+			.UseAllottedSize(true)
+		];
 
-		RebuildSectionSelector();
-	}
+	RebuildSectionSelector();
 
 	FilterRow = FilterRowVBox;
 
@@ -984,6 +1009,21 @@ void SDetailsView::OnShowHiddenPropertiesWhilePlayingClicked()
 	ForceRefresh();
 }
 
+void SDetailsView::OnShowSectionsClicked()
+{
+	DetailsViewArgs.bShowSectionSelector = !DetailsViewArgs.bShowSectionSelector;
+	GetMutableViewConfig().bShowSections = DetailsViewArgs.bShowSectionSelector;
+	SaveViewConfig();
+
+	if (!DetailsViewArgs.bShowSectionSelector)
+	{
+		CurrentFilter.VisibleSections.Reset();
+	}
+
+	RebuildSectionSelector();
+	UpdateFilteredDetails();
+}
+
 FSlateColor SDetailsView::GetToggleFavoritesColor() const
 {
 	if (DetailsViewArgs.bAllowFavoriteSystem && CurrentFilter.bShowFavoritesCategory)
@@ -1170,86 +1210,115 @@ void SDetailsView::SetCustomBuilderFavorite(FStringView Path, bool IsFavorite)
 
 void SDetailsView::RebuildSectionSelector()
 {
-	if (!SectionSelectorBox.IsValid())
+	SectionSelectorBox->ClearChildren();
+	SectionSelectorBox->SetVisibility(EVisibility::Collapsed);
+
+	if (!DetailsViewArgs.bShowSectionSelector)
 	{
 		return;
 	}
 
-	const TSet<FName> Sections = GetAllSectionNames();
-	if (Sections.IsEmpty())
+	const TMap<FName, FText> AllSections = GetAllSections();
+	if (AllSections.IsEmpty())
 	{
-		// we've selected something that has no sections - rather than show just "All", reset and hide the box
-		SectionSelectorBox->SetContent(SNullWidget::NullWidget);
-		SectionSelectorBox->SetVisibility(EVisibility::Collapsed);
+		// we've selected something that has no sections - rather than show just "All", hide the box
 		return;
 	}
 
-	int SegmentsPerLine = 5;
-
-	const int NumSections = Sections.Num() + 1; // implicit "All"
-	const int NumLines = (NumSections / SegmentsPerLine) + 1;
-	if (NumLines > 1)
+	auto CreateSection = [this](FName SectionName, FText SectionDisplayName) 
+		-> TSharedRef<SWidget>
 	{
-		// the segmented control fills MaxSegmentsPerLine segments at a go, this attempts to normalize that
-		// eg. 6 will get divided into 3 + 3 instead of 5 + 1, 12 into 4 + 4 + 4 instead of 5 + 5 + 2
-
-		SegmentsPerLine = (NumSections / NumLines) + (NumSections % NumLines == 0 ? 0 : 1);
-	}
-
-	TSharedRef<SSegmentedControl<FName>> SectionSelectorWidget = SNew(SSegmentedControl<FName>)
-		.TextStyle(FAppStyle::Get(), "SmallText")
-		.MaxSegmentsPerLine(SegmentsPerLine)
-		.OnValueChanged(this, &SDetailsView::OnSectionSelectionChanged)
-		.Value(this, &SDetailsView::GetSelectedSection);
-
+		return SNew(SBox)
+			.Padding(FMargin(0, 0, 2, 0))
+			[
+				SNew(SCheckBox)
+				.Style(FAppStyle::Get(), "DetailsView.SectionButton")
+				.OnCheckStateChanged(this, &SDetailsView::OnSectionCheckedChanged, SectionName)
+				.IsChecked(this, &SDetailsView::IsSectionChecked, SectionName)
+				[
+					SNew(STextBlock)
+					.TextStyle(FAppStyle::Get(), "SmallText")
+					.Text(SectionDisplayName)
+				]
+			];
+	};
 
-	for (const FName& SectionName : Sections)
+	for (const TPair<FName, FText>& Section : AllSections)
 	{
-		FText SectionDisplayName = GetSectionDisplayName(SectionName);
-
-		SectionSelectorWidget->AddSlot(SectionName, false)
-			.Text(SectionDisplayName);
+		SectionSelectorBox->AddSlot()
+		[
+			CreateSection(Section.Key, Section.Value)
+		];
 	}
 
-	SectionSelectorWidget->AddSlot(FName("All"), false)
-		.Text(NSLOCTEXT("UObjectSection", "All", "All"));
-
-	SectionSelectorWidget->RebuildChildren();
+	SectionSelectorBox->AddSlot()
+	[
+		CreateSection(FName("All"), NSLOCTEXT("UObjectSection", "All", "All"))
+	];
 
-	SectionSelectorBox->SetContent(SectionSelectorWidget);
 	SectionSelectorBox->SetVisibility(EVisibility::Visible);
 }
 
-void SDetailsView::OnSectionSelectionChanged(FName NewSelection)
+void SDetailsView::OnSectionCheckedChanged(ECheckBoxState State, FName SectionName)
 {
-	CurrentFilter.VisibleSections.Reset();
+	const bool IsControlDown = FSlateApplication::Get().GetModifierKeys().IsControlDown();
 
-	if (NewSelection != "All")
+	if (State == ECheckBoxState::Unchecked)
 	{
-		CurrentFilter.VisibleSections.Add(NewSelection);
+		if (IsControlDown)
+		{
+			CurrentFilter.VisibleSections.Remove(SectionName);
+		}
+		else
+		{
+			CurrentFilter.VisibleSections.Reset();
+
+			if (SectionName != "All")
+			{
+				CurrentFilter.VisibleSections.Add(SectionName);
+			}
+		}
+	}
+	else if (State == ECheckBoxState::Checked)
+	{
+		if (!IsControlDown)
+		{
+			CurrentFilter.VisibleSections.Reset();
+		}
+
+		if (SectionName != "All")
+		{
+			CurrentFilter.VisibleSections.Add(SectionName);
+		}
 	}
 
 	UpdateFilteredDetails();
 }
 
-FName SDetailsView::GetSelectedSection() const
+ECheckBoxState SDetailsView::IsSectionChecked(FName Section) const
 {
-	FName CurrentSection("All");
-	for (FName Section : CurrentFilter.VisibleSections)
+	if (CurrentFilter.VisibleSections.IsEmpty() && Section == "All")
 	{
-		CurrentSection = Section;
-		break;
+		return ECheckBoxState::Checked;
 	}
 
-	return CurrentSection;
+	for (FName VisibleSection : CurrentFilter.VisibleSections)
+	{
+		if (Section == VisibleSection)
+		{
+			return ECheckBoxState::Checked;
+		}
+	}
+
+	return ECheckBoxState::Unchecked;
 }
 
-TSet<FName> SDetailsView::GetAllSectionNames() const 
+TMap<FName, FText> SDetailsView::GetAllSections() const 
 {
 	static const FName PropertyEditor("PropertyEditor");
 	FPropertyEditorModule& PropertyModule = FModuleManager::GetModuleChecked<FPropertyEditorModule>(PropertyEditor);
 
-	TSet<FName> SectionNames;
+	TMap<FName, FText> AllSections;
 
 	// for every category, check every base struct and find the associated section
 	// if one exists, add it to the set of valid section names
@@ -1257,49 +1326,23 @@ TSet<FName> SDetailsView::GetAllSectionNames() const
 	{
 		if (RootNode->GetNodeType() == EDetailNodeType::Category)
 		{
-			FName CategoryName = RootNode->GetNodeName();
+			const FName CategoryName = RootNode->GetNodeName();
 
 			for (const TSharedPtr<FComplexPropertyNode>& RootPropertyNode : RootPropertyNodes)
 			{
 				const UStruct* RootBaseStruct = RootPropertyNode->GetBaseStructure();
 
-				FName SectionName = PropertyModule.FindSectionForCategory(RootBaseStruct, CategoryName);
-				if (!SectionName.IsNone())
+				TArray<TSharedPtr<FPropertySection>> SectionsForCategory = PropertyModule.FindSectionsForCategory(RootBaseStruct, CategoryName);
+				for (const TSharedPtr<FPropertySection>& Section : SectionsForCategory)
 				{
-					SectionNames.Add(SectionName);
+					AllSections.Add(Section->GetName(), Section->GetDisplayName());
 				}
 			}
 		}
 	}
 
 
-	return MoveTemp(SectionNames);
-}
-
-FText SDetailsView::GetSectionDisplayName(FName SectionName) const
-{
-	static const FTextKey SectionLocalizationNamespace = TEXT("UObjectSection");
-	static const FName SectionMetaDataKey = TEXT("Section");
-
-	FText SectionDisplayName;
-
-	const FString SectionString = SectionName.ToString();
-	if (FText::FindText(SectionLocalizationNamespace, SectionString, /*OUT*/SectionDisplayName, &SectionString))
-	{
-		// Category names in English are typically gathered in their non-pretty form (eg "UserInterface" rather than "User Interface"), so skip 
-		// applying the localized variant if the text matches the raw category name, as in this case the pretty printer will do a better job
-		if (SectionString.Equals(SectionDisplayName.ToString(), ESearchCase::CaseSensitive))
-		{
-			SectionDisplayName = FText();
-		}
-	}
-
-	if (SectionDisplayName.IsEmpty())
-	{
-		SectionDisplayName = FText::AsCultureInvariant(FName::NameToDisplayString(SectionString, false));
-	}
-
-	return SectionDisplayName;
+	return MoveTemp(AllSections);
 }
 
 #undef LOCTEXT_NAMESPACE
diff --git a/Engine/Source/Editor/PropertyEditor/Private/SDetailsView.h b/Engine/Source/Editor/PropertyEditor/Private/SDetailsView.h
--- a/Engine/Source/Editor/PropertyEditor/Private/SDetailsView.h
+++ b/Engine/Source/Editor/PropertyEditor/Private/SDetailsView.h
@@ -5,11 +5,10 @@
 #include "CoreMinimal.h"
 #include "AssetSelection.h"
 #include "IDetailsView.h"
-#include "PropertyNode.h"
-#include "SDetailsViewBase.h"
-
 #include "Input/Reply.h"
 #include "Layout/Visibility.h"
+#include "PropertyNode.h"
+#include "SDetailsViewBase.h"
 #include "Widgets/DeclarativeSyntaxSupport.h"
 #include "Widgets/Layout/SScrollBar.h"
 
@@ -164,6 +163,12 @@ private:
 	/** Called when show hidden properties while playing is clicked */
 	void OnShowHiddenPropertiesWhilePlayingClicked();
 
+	/** @return true if Show Sections is checked. */
+	bool IsShowSectionsChecked() const { return DetailsViewArgs.bShowSectionSelector; }
+
+	/** Called when Show Sections is clicked */
+	void OnShowSectionsClicked();
+
 	/** Get the color of the toggle favorites button. */
 	FSlateColor GetToggleFavoritesColor() const;
 
@@ -173,20 +178,17 @@ private:
 	/** Called after an undo or redo operation occurs in the editor. */
 	void OnPostUndoRedo();
 
-	/** Get all section names for the objects currently selected in the view. */
-	TSet<FName> GetAllSectionNames() const;
+	/** Get all section names and display names for the objects currently selected in the view. */
+	TMap<FName, FText> GetAllSections() const;
 
-	/** Get a display name for the given section. */
-	FText GetSectionDisplayName(FName SectionName) const;
-	
 	/** Rebuild the section selector widget after a selection has been changed. */
 	void RebuildSectionSelector();
 
 	/** Refilter the details view after the user has selected a new section. */
-	void OnSectionSelectionChanged(FName NewSelection);
+	void OnSectionCheckedChanged(ECheckBoxState State, FName NewSelection);
 
 	/** Get the currently selected section. */
-	FName GetSelectedSection() const;
+	ECheckBoxState IsSectionChecked(FName Section) const;
 
 private:
 	/** The filter for objects viewed by this details panel */
@@ -217,5 +219,5 @@ private:
 	/** Delegate handle for unregistering from the PostUndoRedo event. */
 	FDelegateHandle PostUndoRedoDelegateHandle;
 	/** The section selector widget to show if DetailsViewArgs.bShowSectionSelector is true. */
-	TSharedPtr<SBox> SectionSelectorBox;
+	TSharedPtr<class SWrapBox> SectionSelectorBox;
 };
diff --git a/Engine/Source/Editor/PropertyEditor/Public/DetailsViewArgs.h b/Engine/Source/Editor/PropertyEditor/Public/DetailsViewArgs.h
--- a/Engine/Source/Editor/PropertyEditor/Public/DetailsViewArgs.h
+++ b/Engine/Source/Editor/PropertyEditor/Public/DetailsViewArgs.h
@@ -48,7 +48,7 @@ struct FDetailsViewArgs
 	class FNotifyHook* NotifyHook;
 	/** Settings for displaying the name area (@see ENameAreaSettings) */
 	int32 NameAreaSettings;
-	/** The default column width */
+	/** The default value column width, as a percentage, 0-1. */
 	float ColumnWidth;
 	/** The minimum width of the right column in Slate units. */
 	float RightColumnMinWidth;
