CIS fix. Reverted changes to files that unintentionally went in as part of CL 17595848.

diff --git a/Engine/Source/Editor/PropertyEditor/Private/DetailCategoryBuilderImpl.cpp b/Engine/Source/Editor/PropertyEditor/Private/DetailCategoryBuilderImpl.cpp
--- a/Engine/Source/Editor/PropertyEditor/Private/DetailCategoryBuilderImpl.cpp
+++ b/Engine/Source/Editor/PropertyEditor/Private/DetailCategoryBuilderImpl.cpp
@@ -995,29 +995,8 @@ void FDetailCategoryImpl::FilterNode(const FDetailFilter& InFilter)
 			static FName PropertyEditor("PropertyEditor");
 			const FPropertyEditorModule& PropertyModule = FModuleManager::GetModuleChecked<FPropertyEditorModule>(PropertyEditor);
 
-			TArray<TSharedPtr<FPropertySection>> PropertySections = PropertyModule.FindSectionsForCategory(BaseStruct, CategoryName);
-			if (PropertySections.IsEmpty())
-			{
-				return;
-			}
-
-			// if this property is not in any visible section, hide it
-			bool bFound = false;
-			for (FName VisibleSection : InFilter.VisibleSections)
-			{
-				auto Predicate = [VisibleSection](const TSharedPtr<FPropertySection>& Elem)
-				{
-					return Elem->GetName() == VisibleSection;
-				};
-
-				if (PropertySections.ContainsByPredicate(Predicate))
-				{
-					bFound = true;
-					break;
-				}
-			}
-
-			if (!bFound)
+			const FName SectionName = PropertyModule.FindSectionForCategory(BaseStruct, CategoryName);
+			if (SectionName.IsNone() || !InFilter.VisibleSections.Contains(SectionName))
 			{
 				return;
 			}
diff --git a/Engine/Source/Editor/PropertyEditor/Private/SDetailsView.cpp b/Engine/Source/Editor/PropertyEditor/Private/SDetailsView.cpp
--- a/Engine/Source/Editor/PropertyEditor/Private/SDetailsView.cpp
+++ b/Engine/Source/Editor/PropertyEditor/Private/SDetailsView.cpp
@@ -1313,7 +1313,7 @@ ECheckBoxState SDetailsView::IsSectionChecked(FName Section) const
 	return ECheckBoxState::Unchecked;
 }
 
-TMap<FName, FText> SDetailsView::GetAllSections() const 
+TMap<FName, FText> SDetailsView::GetAllSections() const
 {
 	static const FName PropertyEditor("PropertyEditor");
 	FPropertyEditorModule& PropertyModule = FModuleManager::GetModuleChecked<FPropertyEditorModule>(PropertyEditor);
@@ -1332,17 +1332,42 @@ TMap<FName, FText> SDetailsView::GetAllSections() const
 			{
 				const UStruct* RootBaseStruct = RootPropertyNode->GetBaseStructure();
 
-				TArray<TSharedPtr<FPropertySection>> SectionsForCategory = PropertyModule.FindSectionsForCategory(RootBaseStruct, CategoryName);
-				for (const TSharedPtr<FPropertySection>& Section : SectionsForCategory)
+				const FName SectionName = PropertyModule.FindSectionForCategory(RootBaseStruct, CategoryName);
+				if (!SectionName.IsNone())
 				{
-					AllSections.Add(Section->GetName(), Section->GetDisplayName());
+					AllSections.Add(SectionName, GetSectionDisplayName(SectionName));
 				}
 			}
 		}
 	}
 
-
 	return MoveTemp(AllSections);
 }
 
+FText SDetailsView::GetSectionDisplayName(FName SectionName) const
+{
+	static const FTextKey SectionLocalizationNamespace = TEXT("UObjectSection");
+	static const FName SectionMetaDataKey = TEXT("Section");
+
+	FText SectionDisplayName;
+
+	const FString SectionString = SectionName.ToString();
+	if (FText::FindText(SectionLocalizationNamespace, SectionString, /*OUT*/SectionDisplayName, &SectionString))
+	{
+		// Category names in English are typically gathered in their non-pretty form (eg "UserInterface" rather than "User Interface"), so skip 
+		// applying the localized variant if the text matches the raw category name, as in this case the pretty printer will do a better job
+		if (SectionString.Equals(SectionDisplayName.ToString(), ESearchCase::CaseSensitive))
+		{
+			SectionDisplayName = FText();
+		}
+	}
+
+	if (SectionDisplayName.IsEmpty())
+	{
+		SectionDisplayName = FText::AsCultureInvariant(FName::NameToDisplayString(SectionString, false));
+	}
+
+	return SectionDisplayName;
+}
+
 #undef LOCTEXT_NAMESPACE
diff --git a/Engine/Source/Editor/PropertyEditor/Private/SDetailsView.h b/Engine/Source/Editor/PropertyEditor/Private/SDetailsView.h
--- a/Engine/Source/Editor/PropertyEditor/Private/SDetailsView.h
+++ b/Engine/Source/Editor/PropertyEditor/Private/SDetailsView.h
@@ -181,6 +181,9 @@ private:
 	/** Get all section names and display names for the objects currently selected in the view. */
 	TMap<FName, FText> GetAllSections() const;
 
+	/** Get the display name for the given section. */
+	FText GetSectionDisplayName(FName SectionName) const;
+
 	/** Rebuild the section selector widget after a selection has been changed. */
 	void RebuildSectionSelector();
 
