Fixed crash on shutdown in SDetailsView destructor when opening a different project.

SaveViewConfig() called GetMutableDefault<UDetailViewConfig>() which dereferenced null because all UClass instances have been destroyed.

diff --git a/Engine/Source/Editor/PropertyEditor/Private/SDetailsView.cpp b/Engine/Source/Editor/PropertyEditor/Private/SDetailsView.cpp
--- a/Engine/Source/Editor/PropertyEditor/Private/SDetailsView.cpp
+++ b/Engine/Source/Editor/PropertyEditor/Private/SDetailsView.cpp
@@ -36,8 +36,6 @@ SDetailsView::~SDetailsView()
 		SaveExpandedItems(RootNode.ToSharedRef());
 	}
 
-	SaveViewConfig();
-
 	FEditorDelegates::PostUndoRedo.Remove(PostUndoRedoDelegateHandle);
 };
 
