Fixed crash when customizations added layouts while iterating through a category's LayoutMap with a ranged-for.

diff --git a/Engine/Source/Editor/PropertyEditor/Private/DetailCategoryBuilderImpl.cpp b/Engine/Source/Editor/PropertyEditor/Private/DetailCategoryBuilderImpl.cpp
--- a/Engine/Source/Editor/PropertyEditor/Private/DetailCategoryBuilderImpl.cpp
+++ b/Engine/Source/Editor/PropertyEditor/Private/DetailCategoryBuilderImpl.cpp
@@ -904,13 +904,17 @@ void FDetailCategoryImpl::GenerateNodesFromCustomizations(const TArray<FDetailLa
 
 void FDetailCategoryImpl::GenerateChildrenForLayouts()
 {
-	for (const FDetailLayout& Layout : LayoutMap)
+	// note: this can't be a ranged-for, because the map may have items added to it by customizations during iteration
+	for (int32 LayoutIndex = 0; LayoutIndex < LayoutMap.Num(); ++LayoutIndex)
 	{
+		const FDetailLayout& Layout = LayoutMap[LayoutIndex];
 		GenerateNodesFromCustomizations(Layout.GetSimpleLayouts(), SimpleChildNodes);
 	}
 
-	for (const FDetailLayout& Layout : LayoutMap)
+	// note: this can't be a ranged-for, because the map may have items added to it by customizations during iteration
+	for (int32 LayoutIndex = 0; LayoutIndex < LayoutMap.Num(); ++LayoutIndex)
 	{
+		const FDetailLayout& Layout = LayoutMap[LayoutIndex];
 		GenerateNodesFromCustomizations(Layout.GetAdvancedLayouts(), AdvancedChildNodes);
 	}
 
diff --git a/Engine/Source/Editor/PropertyEditor/Private/DetailCategoryBuilderImpl.h b/Engine/Source/Editor/PropertyEditor/Private/DetailCategoryBuilderImpl.h
--- a/Engine/Source/Editor/PropertyEditor/Private/DetailCategoryBuilderImpl.h
+++ b/Engine/Source/Editor/PropertyEditor/Private/DetailCategoryBuilderImpl.h
@@ -129,6 +129,11 @@ public:
 	FORCEINLINE RangedForIteratorType      end()         { return Layouts.end(); }
 	FORCEINLINE RangedForConstIteratorType end()   const { return Layouts.end(); }
 
+	FDetailLayout& operator[](int32 Index) { return Layouts[Index]; }
+	const FDetailLayout& operator[](int32 Index) const { return Layouts[Index]; }
+
+	int32 Num() const { return Layouts.Num(); }
+
 	/**
 	 * @return Whether or not we need to display a group border around a list of details.
 	 */
