Details views now dynamically resize the right column when the sequencer is enabled.

SSplitter slots now use a TAttribute for their MinSize, allowing it to be changed dynamically.

diff --git a/Engine/Source/Editor/PropertyEditor/Private/SDetailsView.cpp b/Engine/Source/Editor/PropertyEditor/Private/SDetailsView.cpp
--- a/Engine/Source/Editor/PropertyEditor/Private/SDetailsView.cpp
+++ b/Engine/Source/Editor/PropertyEditor/Private/SDetailsView.cpp
@@ -60,7 +60,7 @@ void SDetailsView::Construct(const FArguments& InArgs, const FDetailsViewArgs& I
 	}
 
 	ColumnSizeData.SetValueColumnWidth(DetailsViewArgs.ColumnWidth);
-	ColumnSizeData.RightColumnMinWidth = DetailsViewArgs.RightColumnMinWidth;
+	ColumnSizeData.SetRightColumnMinWidth(DetailsViewArgs.RightColumnMinWidth);
 
 	SetObjectFilter(InDetailsViewArgs.ObjectFilter);
 	SetClassViewerFilters(InDetailsViewArgs.ClassViewerFilters);
diff --git a/Engine/Source/Editor/PropertyEditor/Private/SDetailsViewBase.cpp b/Engine/Source/Editor/PropertyEditor/Private/SDetailsViewBase.cpp
--- a/Engine/Source/Editor/PropertyEditor/Private/SDetailsViewBase.cpp
+++ b/Engine/Source/Editor/PropertyEditor/Private/SDetailsViewBase.cpp
@@ -517,7 +517,23 @@ bool SDetailsViewBase::IsPropertyEditingEnabled() const
 
 void SDetailsViewBase::SetKeyframeHandler( TSharedPtr<class IDetailKeyframeHandler> InKeyframeHandler )
 {
+	// if we don't have a keyframe handler and a valid handler is set, add width to the right column
+	// if we do have a keyframe handler and an invalid handler is set, remove width
+	float ExtraWidth = 0;
+	if (!KeyframeHandler.IsValid() && InKeyframeHandler.IsValid())
+	{
+		ExtraWidth = 22;
+	}
+	else if (KeyframeHandler.IsValid() && !InKeyframeHandler.IsValid())
+	{
+		ExtraWidth = -22;
+	}
+	
+	const float NewWidth = ColumnSizeData.GetRightColumnMinWidth() + ExtraWidth;
+	ColumnSizeData.SetRightColumnMinWidth(NewWidth);
+
 	KeyframeHandler = InKeyframeHandler;
+	RefreshTree();
 }
 
 void SDetailsViewBase::SetExtensionHandler(TSharedPtr<class IDetailPropertyExtensionHandler> InExtensionHandler)
@@ -859,6 +875,13 @@ void SDetailsViewBase::Tick( const FGeometry& AllottedGeometry, const double InC
 {
 	HandlePendingCleanup();
 
+	FDetailsViewConfig& ViewConfig = GetMutableViewConfig();
+	if (ViewConfig.ValueColumnWidth != ColumnSizeData.ValueColumnWidth.Get(0))
+	{
+		ViewConfig.ValueColumnWidth = ColumnSizeData.ValueColumnWidth.Get(0);
+		SaveViewConfig();
+	}
+
 	FRootPropertyNodeList& RootPropertyNodes = GetRootNodes();
 
 	bool bHadDeferredActions = DeferredActions.Num() > 0;
diff --git a/Engine/Source/Editor/PropertyEditor/Private/SStructureDetailsView.cpp b/Engine/Source/Editor/PropertyEditor/Private/SStructureDetailsView.cpp
--- a/Engine/Source/Editor/PropertyEditor/Private/SStructureDetailsView.cpp
+++ b/Engine/Source/Editor/PropertyEditor/Private/SStructureDetailsView.cpp
@@ -37,7 +37,7 @@ void SStructureDetailsView::Construct(const FArguments& InArgs)
 	DetailsViewArgs = InArgs._DetailsViewArgs;
 
 	ColumnSizeData.SetValueColumnWidth(DetailsViewArgs.ColumnWidth);
-	ColumnSizeData.RightColumnMinWidth = DetailsViewArgs.RightColumnMinWidth;
+	ColumnSizeData.SetRightColumnMinWidth(DetailsViewArgs.RightColumnMinWidth);
 
 	CustomName = InArgs._CustomName;
 
diff --git a/Engine/Source/Editor/PropertyEditor/Public/IDetailPropertyRow.h b/Engine/Source/Editor/PropertyEditor/Public/IDetailPropertyRow.h
--- a/Engine/Source/Editor/PropertyEditor/Public/IDetailPropertyRow.h
+++ b/Engine/Source/Editor/PropertyEditor/Public/IDetailPropertyRow.h
@@ -29,6 +29,7 @@ public:
 		ValueColumnWidth.BindRaw(this, &FDetailColumnSizeData::GetValueColumnWidth);
 		PropertyColumnWidth.BindRaw(this, &FDetailColumnSizeData::GetPropertyColumnWidth);
 		RightColumnWidth.BindRaw(this, &FDetailColumnSizeData::GetRightColumnWidth);
+		RightColumnMinWidth.BindRaw(this, &FDetailColumnSizeData::GetRightColumnMinWidth);
 		HoveredSplitterIndex.BindRaw(this, &FDetailColumnSizeData::GetHoveredSplitterIndex);
 		OnValueColumnResized.BindRaw(this, &FDetailColumnSizeData::SetValueColumnWidth);
 		OnRightColumnResized.BindRaw(this, &FDetailColumnSizeData::OnSetRightColumnWidth);
@@ -39,12 +40,12 @@ public:
 		OnNameColumnResized.BindLambda([](float) {}); 
 	}
 
-	float RightColumnMinWidth;
-
+	// don't change these attributes directly - they're here only to be forwarded to the details view
 	TAttribute<float> NameColumnWidth;
 	TAttribute<float> ValueColumnWidth;
 	TAttribute<float> PropertyColumnWidth;
 	TAttribute<float> RightColumnWidth;
+	TAttribute<float> RightColumnMinWidth;
 	TAttribute<int32> HoveredSplitterIndex;
 	SSplitter::FOnSlotResized OnNameColumnResized; 
 	SSplitter::FOnSlotResized OnValueColumnResized;
@@ -54,13 +55,20 @@ public:
 
 	void SetValueColumnWidth(float NewWidth)
 	{ 
-		ensure(NewWidth <= 1.0f); 
+		ensureAlways(NewWidth <= 1.0f);
 		ValueColumnWidthValue = NewWidth;
-	} 
+	}
+
+	float GetRightColumnMinWidth() const { return RightColumnMinWidthValue; }
+	void SetRightColumnMinWidth(float NewWidth)
+	{ 
+		RightColumnMinWidthValue = NewWidth;
+	}
 
 private:
 	float ValueColumnWidthValue;
 	float RightColumnWidthValue;
+	float RightColumnMinWidthValue;
 	int HoveredSplitterIndexValue;
 
 	float GetNameColumnWidth() const { return 1.0f - (ValueColumnWidthValue + RightColumnWidthValue); }
@@ -68,11 +76,9 @@ private:
 	float GetRightColumnWidth() const { return RightColumnWidthValue; }
 	float GetPropertyColumnWidth() const { return 1.0f - RightColumnWidthValue; }
 	int32 GetHoveredSplitterIndex() const { return HoveredSplitterIndexValue; }
-
 	
 	void OnSetRightColumnWidth(float NewWidth)
 	{
-		ensure(NewWidth <= 1.0f);
 		RightColumnWidthValue = NewWidth;
 	}
 
diff --git a/Engine/Source/Runtime/Slate/Private/Widgets/Layout/SSplitter.cpp b/Engine/Source/Runtime/Slate/Private/Widgets/Layout/SSplitter.cpp
--- a/Engine/Source/Runtime/Slate/Private/Widgets/Layout/SSplitter.cpp
+++ b/Engine/Source/Runtime/Slate/Private/Widgets/Layout/SSplitter.cpp
@@ -16,10 +16,9 @@ void SSplitter::FSlot::Construct(const FChildren& SlotOwner, FSlotArguments&& In
 	{
 		bIsResizable = InArgs._Resizable.GetValue();
 	}
-	if (InArgs._MinSize.IsSet())
-	{
-		MinSizeValue = InArgs._MinSize.GetValue();
-	}
+	
+	MinSizeValue = InArgs._MinSize;
+
 	if (InArgs._OnSlotResized.IsBound())
 	{
 		OnSlotResized_Handler = MoveTemp(InArgs._OnSlotResized);
@@ -137,7 +136,7 @@ TArray<FLayoutGeometry> SSplitter::ArrangeChildrenForLayout(const FGeometry& All
 		}
 		else // SizingRule == SSplitter::FractionOfParent
 		{
-			MinResizableSpace += FMath::Max(MinSplitterChildLength, CurSlot.MinSizeValue.Get(0));
+			MinResizableSpace += FMath::Max(MinSplitterChildLength, CurSlot.GetMinSize());
 			CoefficientTotal += CurSlot.SizeValue.Get();
 		}
 	}
@@ -187,7 +186,7 @@ TArray<FLayoutGeometry> SSplitter::ArrangeChildrenForLayout(const FGeometry& All
 			{
 				if (Children[PrevIndex].SizingRule.Get() == SSplitter::FractionOfParent && !IsSlotCollapsed(CurSlot))
 				{
-					const float MinChildSize = FMath::Max(MinSplitterChildLength, Children[PrevIndex].MinSizeValue.Get(0));
+					const float MinChildSize = FMath::Max(MinSplitterChildLength, Children[PrevIndex].GetMinSize());
 					const float AvailableSpace = SlotSizes[PrevIndex] - MinChildSize;
 					if (AvailableSpace > 0)
 					{
@@ -801,12 +800,7 @@ void SSplitter::HandleResizingByMousePosition(EOrientation SplitterOrientation,
 float SSplitter::ClampChild(const FSlot& ChildSlot, float ProposedSize) const
 {
 	float ClampedSize = FMath::Max(MinSplitterChildLength, ProposedSize);
-
-	if (ChildSlot.MinSizeValue.IsSet())
-	{
-		ClampedSize = FMath::Max(ChildSlot.MinSizeValue.GetValue(), ProposedSize);
-	}
-
+	ClampedSize = FMath::Max(ChildSlot.GetMinSize(), ProposedSize);
 	return ClampedSize;
 }
 
diff --git a/Engine/Source/Runtime/Slate/Public/Widgets/Layout/SSplitter.h b/Engine/Source/Runtime/Slate/Public/Widgets/Layout/SSplitter.h
--- a/Engine/Source/Runtime/Slate/Public/Widgets/Layout/SSplitter.h
+++ b/Engine/Source/Runtime/Slate/Public/Widgets/Layout/SSplitter.h
@@ -68,6 +68,7 @@ public:
 			: TSlotBase<FSlot>()
 			, SizingRule( FractionOfParent )
 			, SizeValue( 1 )
+			, MinSizeValue( 0 )
 		{
 		}
 
@@ -77,7 +78,7 @@ public:
 			/** When the RuleSize is set to FractionOfParent, the size of the slot is the Value percentage of its parent size. */
 			SLATE_ATTRIBUTE(float, Value)
 			/** Minimum slot size when resizing. */
-			SLATE_ARGUMENT(TOptional<float>, MinSize)
+			SLATE_ATTRIBUTE(float, MinSize)
 			/** Can the slot be resize by the user. */
 			SLATE_ARGUMENT(TOptional<bool>, Resizable)
 			/** Callback when the slot is resized. */
@@ -152,7 +153,7 @@ public:
 		UE_DEPRECATED(5.0, "Direct access to SizeValue is now deprecated. Use the getter.")
 		TAttribute<float> SizeValue;
 		UE_DEPRECATED(5.0, "Direct access to MinSizeValue is now deprecated. Use the getter.")
-		TOptional<float> MinSizeValue;
+		TAttribute<float> MinSizeValue;
 		UE_DEPRECATED(5.0, "Direct access to OnSlotResized_Handler is now deprecated. Use the getter.")
 		FOnSlotResized OnSlotResized_Handler;
 		UE_DEPRECATED(5.0, "Direct access to bIsResizable is now deprecated. Use the getter.")
