Details views that have Show Only Modified set don't hide the entire view - SDetailsViewBase::GetFilterBoxVisibility() uses IsEmptyFilter() to evaluate visibility, which was incorrectly causing it to hide when Show Only Modified was set.

diff --git a/Engine/Source/Editor/PropertyEditor/Private/DetailFilter.h b/Engine/Source/Editor/PropertyEditor/Private/DetailFilter.h
--- a/Engine/Source/Editor/PropertyEditor/Private/DetailFilter.h
+++ b/Engine/Source/Editor/PropertyEditor/Private/DetailFilter.h
@@ -21,6 +21,7 @@ struct FDetailFilter
 	{
 		return FilterStrings.Num() == 0
 			&& VisibleSections.Num() == 0
+			&& bShowOnlyModified == false
 			&& bShowAllAdvanced == false
 			&& bShowOnlyWhitelisted == false
 			&& bShowAllChildrenIfCategoryMatches == false
