Fix return IPropertyHandle in IDetailLayoutBuilder::AddCustomRowToCategory

diff --git a/Engine/Source/Editor/PropertyEditor/Private/DetailLayoutBuilderImpl.cpp b/Engine/Source/Editor/PropertyEditor/Private/DetailLayoutBuilderImpl.cpp
--- a/Engine/Source/Editor/PropertyEditor/Private/DetailLayoutBuilderImpl.cpp
+++ b/Engine/Source/Editor/PropertyEditor/Private/DetailLayoutBuilderImpl.cpp
@@ -176,6 +176,10 @@ TSharedPtr<IPropertyHandle> FDetailLayoutBuilderImpl::AddStructurePropertyData(c
 					FClassInstanceToPropertyMap& ClassInstanceToPropertyMap = PropertyMap.FindOrAdd(PropertyNode->GetProperty()->GetOwnerStruct()->GetFName());
 					FPropertyNodeMap& PropertyNodeMap = ClassInstanceToPropertyMap.FindOrAdd(NAME_None);
 					PropertyNodeMap.Add(PropertyName, PropertyNode);
+
+					RootPropertyNode->AddChildNode(PropertyNode);
+					PropertyNode->RebuildChildren();
+					Handle = GetPropertyHandle(PropertyNode);
 					break;
 				}
 			}
