The details Section API has been rewritten. Categories can now be part of multiple sections simultaneously. This solves the situation of wanting eg. Materials to be in both the General (because it's a very common property to manipulate for meshes) and the Rendering (because it affects rendering) sections.

Sections are now sorted correctly - General and All are special-cased to be the first and last categories respectively, the rest are alphabetically sorted.

diff --git a/Engine/Source/Editor/PropertyEditor/Private/DetailCategoryBuilderImpl.cpp b/Engine/Source/Editor/PropertyEditor/Private/DetailCategoryBuilderImpl.cpp
--- a/Engine/Source/Editor/PropertyEditor/Private/DetailCategoryBuilderImpl.cpp
+++ b/Engine/Source/Editor/PropertyEditor/Private/DetailCategoryBuilderImpl.cpp
@@ -999,8 +999,25 @@ void FDetailCategoryImpl::FilterNode(const FDetailFilter& InFilter)
 			static FName PropertyEditor("PropertyEditor");
 			const FPropertyEditorModule& PropertyModule = FModuleManager::GetModuleChecked<FPropertyEditorModule>(PropertyEditor);
 
-			const FName SectionName = PropertyModule.FindSectionForCategory(BaseStruct, CategoryName);
-			if (SectionName.IsNone() || !InFilter.VisibleSections.Contains(SectionName))
+			TArray<TSharedPtr<FPropertySection>> PropertySections = PropertyModule.FindSectionsForCategory(BaseStruct, CategoryName);
+			if (PropertySections.IsEmpty())
+			{
+				// property has no sections, must be filtered
+				return;
+			}
+
+			// if this property is not in any visible section, hide it
+			bool bFound = false;
+			for (const TSharedPtr<FPropertySection>& Section : PropertySections)
+			{
+				if (InFilter.VisibleSections.Contains(Section->GetName()))
+				{
+					bFound = true; 
+					break;
+				}
+			}
+
+			if (!bFound)
 			{
 				return;
 			}
diff --git a/Engine/Source/Editor/PropertyEditor/Private/PropertyEditorModule.cpp b/Engine/Source/Editor/PropertyEditor/Private/PropertyEditorModule.cpp
--- a/Engine/Source/Editor/PropertyEditor/Private/PropertyEditorModule.cpp
+++ b/Engine/Source/Editor/PropertyEditor/Private/PropertyEditorModule.cpp
@@ -513,42 +513,105 @@ void FPropertyEditorModule::UnregisterCustomPropertyTypeLayout( FName PropertyTy
 	}
 }
 
-void FPropertyEditorModule::RegisterSectionMapping(FName ClassName, FName SectionName, FName CategoryName)
+void FPropertySection::AddCategory(FName CategoryName)
 {
-	if (!ClassName.IsValid() || ClassName.IsNone())
+	checkf(!CategoryName.ToString().Contains(TEXT("|")), TEXT("Cannnot register a section mapping for a subcategory. Section: '%s', Category: '%s'"), *Name.ToString(), *CategoryName.ToString());
+
+	// Remove all spaces - customization authors will probably write "Static Mesh", but internally it's stored as "StaticMesh"
+	FString CategoryString = CategoryName.ToString();
+	CategoryString.RemoveSpacesInline();
+	CategoryName = FName(*CategoryString);
+
+	Categories.Add(CategoryName);
+}
+
+bool FPropertySection::ContainsCategory(FName CategoryName) const
+{
+	return Categories.Contains(CategoryName);
+}
+
+TSharedPtr<FPropertySection> FClassSectionMapping::FindSection(FName SectionName) const
+{
+	const TSharedPtr<FPropertySection>* Section = DefinedSections.Find(SectionName);
+	if (Section == nullptr)
 	{
-		return;
+		return TSharedPtr<FPropertySection>();
+	}
+
+	return *Section;
+}
+
+TSharedRef<FPropertySection> FClassSectionMapping::FindOrAddSection(FName SectionName, FText DisplayName)
+{
+	TSharedPtr<FPropertySection> Section = FindSection(SectionName);
+	if (!Section.IsValid())
+	{
+		Section = MakeShared<FPropertySection>(SectionName, DisplayName);
+		DefinedSections.Add(SectionName, Section);
+	}
+
+	return Section.ToSharedRef();
+}
+
+bool FClassSectionMapping::GetSectionsForCategory(FName CategoryName, TArray<TSharedPtr<FPropertySection>>& OutSections) const 
+{
+	bool bAnyAdded = false;
+
+	for (const TPair<FName, TSharedPtr<FPropertySection>>& Pair : DefinedSections)
+	{
+		if (Pair.Value->ContainsCategory(CategoryName))
+		{
+			OutSections.Add(Pair.Value);
+			bAnyAdded = true;
+		}
 	}
 
-	checkf(!CategoryName.ToString().Contains(TEXT("|")), TEXT("Cannnot register a section mapping for a subcategory. Class: '%s', Section: '%s', Category: '%s'"), *ClassName.ToString(), *SectionName.ToString(), *CategoryName.ToString());
+	return bAnyAdded;
+}
 
-	FClassSectionMapping& SectionMapping = ClassSectionMappings.FindOrAdd(ClassName);
+TSharedRef<FPropertySection> FPropertyEditorModule::FindOrCreateSection(FName ClassName, FName SectionName, FText DisplayName)
+{
+	checkf(!ClassName.IsNone(), TEXT("Invalid class name given."));
+	checkf(!SectionName.IsNone(), TEXT("Invalid section name given."));
 
-	const FName* ExistingSection = SectionMapping.CategoryToSectionMap.Find(CategoryName);
-	if (ExistingSection != nullptr)
+	TSharedPtr<FClassSectionMapping> ClassMapping;
+
+	TSharedPtr<FClassSectionMapping>* ExistingMapping = ClassSectionMappings.Find(ClassName);
+	if (ExistingMapping == nullptr)
+	{
+		ClassMapping = ClassSectionMappings.Add(ClassName, MakeShared<FClassSectionMapping>());
+	}
+	else
 	{
-		checkf(!ExistingSection, TEXT("This category has already been mapped to a section for this class. Class: '%s', Category: '%s', Existing Section: '%s', New Section: '%s'"), *ClassName.ToString(), *CategoryName.ToString(), *ExistingSection->ToString(), *SectionName.ToString());
+		ClassMapping = *ExistingMapping;
 	}
 
-	SectionMapping.CategoryToSectionMap.Add(CategoryName, SectionName);
+	return ClassMapping->FindOrAddSection(SectionName, DisplayName);
 }
 
-FName FPropertyEditorModule::FindSectionForCategory(const UStruct* Struct, FName CategoryName) const
+TArray<TSharedPtr<FPropertySection>> FPropertyEditorModule::FindSectionsForCategory(const UStruct* Struct, FName CategoryName) const
 {
-	if (Struct == nullptr)
+	TArray<TSharedPtr<FPropertySection>> Sections;
+	
+	// remove all spaces - customization authors will probably write "Static Mesh", but internally it's stored as "StaticMesh"
+	FString CategoryString = CategoryName.ToString();
+	CategoryString.RemoveSpacesInline();
+	CategoryName = FName(*CategoryString);
+
+	if (Struct != nullptr)
 	{
-		return NAME_None;
+		TSet<const UStruct*> SearchedStructs;
+		FindSectionsForCategoryHelper(Struct, CategoryName, Sections, SearchedStructs);
 	}
 
-	TSet<const UStruct*> SearchedStructs;
-	return FindSectionForCategoryHelper(Struct, CategoryName, SearchedStructs);
+	return MoveTemp(Sections);
 }
 
-FName FPropertyEditorModule::FindSectionForCategoryHelper(const UStruct* Struct, FName CategoryName, TSet<const UStruct*>& SearchedStructs) const
+void FPropertyEditorModule::FindSectionsForCategoryHelper(const UStruct* Struct, FName CategoryName, TArray<TSharedPtr<FPropertySection>>& OutSections, TSet<const UStruct*>& SearchedStructs) const
 {
 	if (Struct == nullptr)
 	{
-		return NAME_None;
+		return;
 	}
 
 	bool bAlreadySearched = false;
@@ -556,55 +619,38 @@ FName FPropertyEditorModule::FindSectionForCategoryHelper(const UStruct* Struct,
 
 	if (bAlreadySearched)
 	{
-		return NAME_None;
+		return;
 	}
-
-	const UStruct* CurrentStruct = Struct;
-	while (CurrentStruct != nullptr)
+	
+	// check this class' sections
+	const TSharedPtr<FClassSectionMapping>* SectionMapping = ClassSectionMappings.Find(Struct->GetFName());
+	if (SectionMapping != nullptr)
 	{
-		// check this class' sections
-		const FClassSectionMapping* SectionMapping = ClassSectionMappings.Find(Struct->GetFName());
-		if (SectionMapping != nullptr)
-		{
-			const FName* SectionName = SectionMapping->CategoryToSectionMap.Find(CategoryName);
-			if (SectionName != nullptr)
-			{
-				return *SectionName;
-			}
-		}
+		(*SectionMapping)->GetSectionsForCategory(CategoryName, OutSections);
+	}
 
-		// check all struct properties' sections
-		for (TFieldIterator<FStructProperty> It(CurrentStruct); It; ++It)
-		{
-			const FStructProperty* Property = *It;
-			const FName SectionName = FindSectionForCategoryHelper(Property->Struct, CategoryName, SearchedStructs);
-			if (!SectionName.IsNone())
-			{
-				return SectionName;
-			}
-		}
+	// check this struct's super
+	FindSectionsForCategoryHelper(Struct->GetSuperStruct(), CategoryName, OutSections, SearchedStructs);
+
+	// check all struct properties' sections
+	for (TFieldIterator<FStructProperty> It(Struct); It; ++It)
+	{
+		const FStructProperty* Property = *It;
+		FindSectionsForCategoryHelper(Property->Struct, CategoryName, OutSections, SearchedStructs);
+	}
 
-		// check all object properties' sections
-		for (TFieldIterator<FObjectPropertyBase> It(CurrentStruct); It; ++It)
+	// check all inline object properties' sections
+	for (TFieldIterator<FObjectPropertyBase> It(Struct); It; ++It)
+	{
+		const FObjectPropertyBase* Property = *It;
+		if (Property->HasAnyPropertyFlags(CPF_InstancedReference | CPF_ContainsInstancedReference))
 		{
-			const FObjectPropertyBase* Property = *It;
-			if (Property->HasAnyPropertyFlags(CPF_InstancedReference | CPF_ContainsInstancedReference))
-			{
-				const FName SectionName = FindSectionForCategoryHelper(Property->PropertyClass, CategoryName, SearchedStructs);
-				if (!SectionName.IsNone())
-				{
-					return SectionName;
-				}
-			}
+			FindSectionsForCategoryHelper(Property->PropertyClass, CategoryName, OutSections, SearchedStructs);
 		}
-
-		CurrentStruct = CurrentStruct->GetSuperStruct();
 	}
-
-	return NAME_None;
 }
 
-void FPropertyEditorModule::GetAllSections(const UStruct* Struct, TSet<FName>& OutSectionNames) const
+void FPropertyEditorModule::GetAllSections(const UStruct* Struct, TArray<TSharedPtr<FPropertySection>>& OutSections) const
 {
 	if (Struct == nullptr)
 	{
@@ -612,10 +658,10 @@ void FPropertyEditorModule::GetAllSections(const UStruct* Struct, TSet<FName>& O
 	}
 
 	TSet<const UStruct*> ProcessedStructs;
-	GetAllSectionsHelper(Struct, ProcessedStructs, OutSectionNames);
+	GetAllSectionsHelper(Struct, OutSections, ProcessedStructs);
 }
 
-void FPropertyEditorModule::GetAllSectionsHelper(const UStruct* Struct, TSet<const UStruct*>& ProcessedStructs, TSet<FName>& OutSectionNames) const
+void FPropertyEditorModule::GetAllSectionsHelper(const UStruct* Struct, TArray<TSharedPtr<FPropertySection>>& OutSections, TSet<const UStruct*>& ProcessedStructs) const
 {
 	if (Struct == nullptr)
 	{
@@ -630,37 +676,34 @@ void FPropertyEditorModule::GetAllSectionsHelper(const UStruct* Struct, TSet<con
 		return;
 	}
 
-	const UStruct* CurrentStruct = Struct;
-	while (CurrentStruct != nullptr)
+	// add all sections defined for this class
+	const TSharedPtr<FClassSectionMapping>* SectionMapping = ClassSectionMappings.Find(Struct->GetFName());
+	if (SectionMapping != nullptr)
 	{
-		// add all sections defined for this class
-		const FClassSectionMapping* SectionMapping = ClassSectionMappings.Find(Struct->GetFName());
-		if (SectionMapping != nullptr)
+		for (const TPair<FName, TSharedPtr<FPropertySection>>& Pair : (*SectionMapping)->DefinedSections)
 		{
-			for (const TPair<FName, FName>& Pair : SectionMapping->CategoryToSectionMap)
-			{
-				OutSectionNames.Add(Pair.Value);
-			}
+			OutSections.Add(Pair.Value);
 		}
+	}
+	
+	// add all sections for this class' super-struct
+	GetAllSectionsHelper(Struct->GetSuperStruct(), OutSections, ProcessedStructs);
 
-		// check all struct properties and add their sections
-		for (TFieldIterator<FStructProperty> It(CurrentStruct); It; ++It)
-		{
-			const FStructProperty* Property = *It;
-			GetAllSectionsHelper(Property->Struct, ProcessedStructs, OutSectionNames);
-		}
+	// add all sections from struct properties
+	for (TFieldIterator<FStructProperty> It(Struct); It; ++It)
+	{
+		const FStructProperty* Property = *It;
+		GetAllSectionsHelper(Property->Struct, OutSections, ProcessedStructs);
+	}
 
-		// check all object properties and add their sections
-		for (TFieldIterator<FObjectPropertyBase> It(CurrentStruct); It; ++It)
+	// add all sections from inline object properties
+	for (TFieldIterator<FObjectPropertyBase> It(Struct); It; ++It)
+	{
+		const FObjectPropertyBase* Property = *It;
+		if (Property->HasAnyPropertyFlags(CPF_InstancedReference | CPF_ContainsInstancedReference))
 		{
-			const FObjectPropertyBase* Property = *It;
-			if (Property->HasAnyPropertyFlags(CPF_InstancedReference | CPF_ContainsInstancedReference))
-			{
-				GetAllSectionsHelper(Property->PropertyClass, ProcessedStructs, OutSectionNames);
-			}
+			GetAllSectionsHelper(Property->PropertyClass, OutSections, ProcessedStructs);
 		}
-
-		CurrentStruct = CurrentStruct->GetSuperStruct();
 	}
 }
 
diff --git a/Engine/Source/Editor/PropertyEditor/Private/SDetailsView.cpp b/Engine/Source/Editor/PropertyEditor/Private/SDetailsView.cpp
--- a/Engine/Source/Editor/PropertyEditor/Private/SDetailsView.cpp
+++ b/Engine/Source/Editor/PropertyEditor/Private/SDetailsView.cpp
@@ -1241,11 +1241,30 @@ void SDetailsView::RebuildSectionSelector()
 			];
 	};
 
-	for (const TPair<FName, FText>& Section : AllSections)
+	TArray<FName> SortedKeys;
+	AllSections.GenerateKeyArray(SortedKeys);
+	SortedKeys.Sort([](FName A, FName B)
+		{
+			static const FName General("General");
+			static const FName All("All");
+
+			// General first, All last, rest alphabetical
+			if (A.IsEqual(General) || B.IsEqual(All))
+			{
+				return true;
+			}
+			if (A.IsEqual(All) || B.IsEqual(General))
+			{
+				return false;
+			}
+			return A.LexicalLess(B);
+		});
+
+	for (const FName& Key : SortedKeys)
 	{
 		SectionSelectorBox->AddSlot()
 		[
-			CreateSection(Section.Key, Section.Value)
+			CreateSection(Key, AllSections[Key])
 		];
 	}
 
@@ -1320,52 +1339,32 @@ TMap<FName, FText> SDetailsView::GetAllSections() const
 
 	// for every category, check every base struct and find the associated section
 	// if one exists, add it to the set of valid section names
+	TArray<FName> AllCategories;
+	AllCategories.Reserve(RootTreeNodes.Num());
 	for (const TSharedRef<FDetailTreeNode>& RootNode : RootTreeNodes)
 	{
 		if (RootNode->GetNodeType() == EDetailNodeType::Category)
 		{
 			const FName CategoryName = RootNode->GetNodeName();
-
-			for (const TSharedPtr<FComplexPropertyNode>& RootPropertyNode : RootPropertyNodes)
-			{
-				const UStruct* RootBaseStruct = RootPropertyNode->GetBaseStructure();
-
-				const FName SectionName = PropertyModule.FindSectionForCategory(RootBaseStruct, CategoryName);
-				if (!SectionName.IsNone())
-				{
-					AllSections.Add(SectionName, GetSectionDisplayName(SectionName));
-				}
-			}
+			AllCategories.Add(CategoryName);
 		}
 	}
 
-	return MoveTemp(AllSections);
-}
-
-FText SDetailsView::GetSectionDisplayName(FName SectionName) const
-{
-	static const FTextKey SectionLocalizationNamespace = TEXT("UObjectSection");
-	static const FName SectionMetaDataKey = TEXT("Section");
-
-	FText SectionDisplayName;
-
-	const FString SectionString = SectionName.ToString();
-	if (FText::FindText(SectionLocalizationNamespace, SectionString, /*OUT*/SectionDisplayName, &SectionString))
+	for (const TSharedPtr<FComplexPropertyNode>& RootPropertyNode : RootPropertyNodes)
 	{
-		// Category names in English are typically gathered in their non-pretty form (eg "UserInterface" rather than "User Interface"), so skip 
-		// applying the localized variant if the text matches the raw category name, as in this case the pretty printer will do a better job
-		if (SectionString.Equals(SectionDisplayName.ToString(), ESearchCase::CaseSensitive))
+		const UStruct* RootBaseStruct = RootPropertyNode->GetBaseStructure();
+
+		for (const FName& CategoryName : AllCategories)
 		{
-			SectionDisplayName = FText();
+			TArray<TSharedPtr<FPropertySection>> SectionsForCategory = PropertyModule.FindSectionsForCategory(RootBaseStruct, CategoryName);
+			for (const TSharedPtr<FPropertySection>& Section : SectionsForCategory)
+			{
+				AllSections.Add(Section->GetName(), Section->GetDisplayName());
+			}
 		}
 	}
 
-	if (SectionDisplayName.IsEmpty())
-	{
-		SectionDisplayName = FText::AsCultureInvariant(FName::NameToDisplayString(SectionString, false));
-	}
-
-	return SectionDisplayName;
+	return MoveTemp(AllSections);
 }
 
 #undef LOCTEXT_NAMESPACE
diff --git a/Engine/Source/Editor/PropertyEditor/Private/SDetailsView.h b/Engine/Source/Editor/PropertyEditor/Private/SDetailsView.h
--- a/Engine/Source/Editor/PropertyEditor/Private/SDetailsView.h
+++ b/Engine/Source/Editor/PropertyEditor/Private/SDetailsView.h
@@ -181,9 +181,6 @@ private:
 	/** Get all section names and display names for the objects currently selected in the view. */
 	TMap<FName, FText> GetAllSections() const;
 
-	/** Get the display name for the given section. */
-	FText GetSectionDisplayName(FName SectionName) const;
-
 	/** Rebuild the section selector widget after a selection has been changed. */
 	void RebuildSectionSelector();
 
diff --git a/Engine/Source/Editor/PropertyEditor/Public/PropertyEditorModule.h b/Engine/Source/Editor/PropertyEditor/Public/PropertyEditorModule.h
--- a/Engine/Source/Editor/PropertyEditor/Public/PropertyEditorModule.h
+++ b/Engine/Source/Editor/PropertyEditor/Public/PropertyEditorModule.h
@@ -112,6 +112,84 @@ struct FStructureDetailsViewArgs
 	bool bShowInterfaces : 1;
 };
 
+/** 
+ * A property section is a group of categories with a name, eg. "Rendering" might contain "Materials" and "Lighting".
+ * Categories may belong to zero or more sections. 
+ */
+class FPropertySection
+{
+public:
+
+	/**
+	 * @param InName		The internal name of this section.
+	 * @param InDisplayName	The localizable display name to show to the user.
+	 */
+	FPropertySection(FName InName, FText InDisplayName) : 
+		Name(InName),
+		DisplayName(InDisplayName)
+	{
+	}
+
+	FPropertySection(const FPropertySection&) = default;
+	virtual ~FPropertySection() = default;
+
+	/** Add a category to this section. */
+	virtual void AddCategory(FName CategoryName);
+
+	/** Does this section contain the given category. */
+	virtual bool ContainsCategory(FName CategoryName) const;
+
+	/** Get the internal name of this property section. */
+	virtual FName GetName() const { return Name; }
+
+	/** Get the display name of this section. */
+	virtual FText GetDisplayName() const { return DisplayName; }
+
+private:
+
+	/** The internal name to use for this section. */
+	FName Name;
+
+	/** The display name to use for this section. */
+	FText DisplayName;
+	
+	/** A mapping of category to the sections that it should be added to. */
+	TSet<FName> Categories;
+};
+
+/** A mapping of categories to section names for a given class. */
+class FClassSectionMapping
+{
+public:
+	FClassSectionMapping() = default;
+	FClassSectionMapping(const FClassSectionMapping&) = default;
+
+	/**
+	 * Find or add a section of the given name.
+	 */
+	TSharedPtr<FPropertySection> FindSection(FName SectionName) const;
+
+	/**
+	 * Find or add a section of the given name.
+	 */
+	TSharedRef<FPropertySection> FindOrAddSection(FName SectionName, FText DisplayName);
+
+	/** 
+	 * Get the sections that the given category belongs to and append them to OutSections. 
+	 * @param CategoryName	The category name to search for.
+	 * @param OutSections	The array to append any found sections. The array will not be cleared.
+	 * @return				true if any sections were found, false otherwise.
+	 */
+	bool GetSectionsForCategory(FName CategoryName, TArray<TSharedPtr<FPropertySection>>& OutSections) const;
+
+private:
+
+	friend class FPropertyEditorModule;
+
+	/** The sections defined for this class. */
+	TMap<FName, TSharedPtr<FPropertySection>> DefinedSections;
+};
+
 
 class FPropertyEditorModule : public IModuleInterface
 {
@@ -193,19 +271,28 @@ public:
 	virtual void UnregisterCustomPropertyTypeLayout( FName PropertyTypeName, TSharedPtr<IPropertyTypeIdentifier> InIdentifier = nullptr);
 
 	/**
-	 * Register a section mapping for a class.
+	 * Find an existing section or create a section for a class.
 	 * 
 	 * @param ClassName		The class to add a section mapping for.
-	 * @param SectionName	The section to add a category to.
-	 * @param CategoryName	The category to add.
+	 * @param SectionName	The section to find or create.
+	 * @param DisplayName	The display name to use for the section. If the section already exists for this class, the display name will not be replaced.
+	 * @return				A new section, or the existing one. 
 	 */
-	virtual void RegisterSectionMapping(FName ClassName, FName SectionName, FName CategoryName);
+	virtual TSharedRef<FPropertySection> FindOrCreateSection(FName ClassName, FName SectionName, FText DisplayName);
 
-	/** Find the section that the given category in the given struct should be a part of. */
-	virtual FName FindSectionForCategory(const UStruct* Struct, FName CategoryName) const;
+	/** 
+	 * Find the section that the given category in the given struct should be a part of. 
+	 * @param Struct		The struct to start searching from. Note: all super-structs of the given struct will also be searched.
+	 * @param CategoryName	The category to search for.
+	 */
+	virtual TArray<TSharedPtr<FPropertySection>> FindSectionsForCategory(const UStruct* Struct, FName CategoryName) const;
 
-	/** Get all registered sections for the given struct (including the default section). */
-	virtual void GetAllSections(const UStruct* Struct, TSet<FName>& OutSectionNames) const;
+	/** 
+	 * Get all registered sections for the given struct (including the default section). 
+	 * @param Struct		The struct to fetch sections for.
+	 * @param OutSections	Sections will be appended to this parameter. The array will not be cleared beforehand.
+	 */
+	virtual void GetAllSections(const UStruct* Struct, TArray<TSharedPtr<FPropertySection>>& OutSections) const;
 
 	/**
 	 * Customization modules should call this when that module has been unloaded, loaded, etc...
@@ -316,8 +403,8 @@ private:
 
 	TSharedPtr<FAssetThumbnailPool> GetThumbnailPool();
 
-	void GetAllSectionsHelper(const UStruct* Struct, TSet<const UStruct*>& ProcessedStructs, TSet<FName>& OutSectionNames) const;
-	FName FindSectionForCategoryHelper(const UStruct* Struct, FName CategoryName, TSet<const UStruct*>& SearchedStructs) const;
+	void GetAllSectionsHelper(const UStruct* Struct, TArray<TSharedPtr<FPropertySection>>& OutSections, TSet<const UStruct*>& ProcessedStructs) const;
+	void FindSectionsForCategoryHelper(const UStruct* Struct, FName CategoryName, TArray<TSharedPtr<FPropertySection>>& OutSections, TSet<const UStruct*>& SearchedStructs) const;
 
 private:
 	/** All created detail views */
@@ -328,16 +415,8 @@ private:
 	FCustomDetailLayoutNameMap ClassNameToDetailLayoutNameMap;
 	/** A mapping of property names to property type layout delegates, called when querying for custom property layouts */
 	FCustomPropertyTypeLayoutMap GlobalPropertyTypeToLayoutMap;
-
-	/** A mapping of categories to section names for a given class. */
-	struct FClassSectionMapping
-	{
-		/** A mapping of category to the section that it should be added to. */
-		TMap<FName, FName> CategoryToSectionMap;
-	};
-
 	/** A mapping of class names to section mappings. */
-	TMap<FName, FClassSectionMapping> ClassSectionMappings;
+	TMap<FName, TSharedPtr<FClassSectionMapping>> ClassSectionMappings;
 	/** Event to be called when a property editor is opened */
 	FPropertyEditorOpenedEvent PropertyEditorOpened;
 	/** Mapping of registered floating UStructs to their struct proxy so they show correctly in the details panel */
