Added ensure in GetDefaultValueAsString() to check that default value does not differ per-object (it really shouldn't) in FPropertyNode. If all values are the same, we return only the single default value, but if they differ we'll fallback to the previous behaviour.

This ensure is intended to check if we can simplify the logic here to just return the first object's default value in the near future.

diff --git a/Engine/Source/Editor/PropertyEditor/Private/PropertyNode.cpp b/Engine/Source/Editor/PropertyEditor/Private/PropertyNode.cpp
--- a/Engine/Source/Editor/PropertyEditor/Private/PropertyNode.cpp
+++ b/Engine/Source/Editor/PropertyEditor/Private/PropertyNode.cpp
@@ -2075,10 +2075,17 @@ FString FPropertyNode::GetDefaultValueAsStringForObject( FPropertyItemValueDataT
 
 FString FPropertyNode::GetDefaultValueAsString(bool bUseDisplayName)
 {
-	FObjectPropertyNode* ObjectNode = FindObjectItemParent();
 	FString DefaultValue;
+	FString DelimitedValue;
+
+	bool bAllSame = true;
+	
+	FObjectPropertyNode* ObjectNode = FindObjectItemParent();
 	if (ObjectNode && Property.IsValid())
 	{
+		TArray<FString> Values;
+		Values.Reserve(ObjectNode->GetNumObjects());
+
 		// Get an iterator for the enclosing objects.
 		for (int32 ObjIndex = 0; ObjIndex < ObjectNode->GetNumObjects(); ++ObjIndex)
 		{
@@ -2087,17 +2094,28 @@ FString FPropertyNode::GetDefaultValueAsString(bool bUseDisplayName)
 
 			if (Object && ValueTracker.IsValid())
 			{
-				FString NodeDefaultValue = GetDefaultValueAsStringForObject( *ValueTracker, Object, Property.Get(), bUseDisplayName );
-				if (DefaultValue.Len() > 0 && NodeDefaultValue.Len() > 0)
+				const FString NodeDefaultValue = GetDefaultValueAsStringForObject( *ValueTracker, Object, Property.Get(), bUseDisplayName );
+
+				if (DefaultValue.IsEmpty())
 				{
-					DefaultValue += TEXT(", ");
+					DefaultValue = NodeDefaultValue;
+				}
+
+				if (DelimitedValue.Len() > 0 && NodeDefaultValue.Len() > 0)
+				{
+					DelimitedValue += TEXT(", ");
+				}
+				DelimitedValue += NodeDefaultValue;
+
+				if (!ensureAlwaysMsgf(NodeDefaultValue == DefaultValue, TEXT("Default values differ for different objects of property '%s'. First: \"%s\", Other: \"%s\""), *Property->GetNameCPP(), *DefaultValue, *NodeDefaultValue))
+				{
+					bAllSame = false;
 				}
-				DefaultValue += NodeDefaultValue;
 			}
 		}
 	}
 
-	return DefaultValue;
+	return bAllSame ? DefaultValue : DelimitedValue; 
 }
 
 FText FPropertyNode::GetResetToDefaultLabel()
