Changed FPropertyNode::GetChildNode(FName) to breadth-first search instead of depth-first, which lead to these types of situations:

class UFoo
{
   struct Bar { int X; } Bar;
   int X;
}

Returning Bar.X instead of just X when searching for "X", because we were iterating into the struct's children before our own.

This entire behaviour is a bit questionable since it doesn't use property paths, and is presumably in place to deal with groups and not structs, but I'm not comfortable making that change without further testing.

diff --git a/Engine/Source/Editor/PropertyEditor/Private/PropertyNode.cpp b/Engine/Source/Editor/PropertyEditor/Private/PropertyNode.cpp
--- a/Engine/Source/Editor/PropertyEditor/Private/PropertyNode.cpp
+++ b/Engine/Source/Editor/PropertyEditor/Private/PropertyNode.cpp
@@ -858,19 +858,20 @@ bool FPropertyNode::GetChildNode(const int32 ChildArrayIndex, TSharedPtr<FProper
 TSharedPtr<FPropertyNode> FPropertyNode::FindChildPropertyNode( const FName InPropertyName, bool bRecurse )
 {
 	// Search Children
-	for(int32 ChildIndex=0; ChildIndex<ChildNodes.Num(); ChildIndex++)
+	for (const TSharedPtr<FPropertyNode>& ChildNode : ChildNodes)
 	{
-		TSharedPtr<FPropertyNode>& ChildNode = ChildNodes[ChildIndex];
-
 		if( ChildNode->GetProperty() && ChildNode->GetProperty()->GetFName() == InPropertyName )
 		{
 			return ChildNode;
 		}
-		else if( bRecurse )
-		{
-			TSharedPtr<FPropertyNode> PropertyNode = ChildNode->FindChildPropertyNode(InPropertyName, bRecurse );
+	}
 
-			if( PropertyNode.IsValid() )
+	if (bRecurse)
+	{
+		for (const TSharedPtr<FPropertyNode>& ChildNode : ChildNodes)
+		{
+			TSharedPtr<FPropertyNode> PropertyNode = ChildNode->FindChildPropertyNode(InPropertyName, bRecurse);
+			if (PropertyNode.IsValid())
 			{
 				return PropertyNode;
 			}
