Spacing changes to section selector - left and right reduced to 8, inter-section spacing now 4 both vertically and horizontally.

diff --git a/Engine/Source/Editor/PropertyEditor/Private/SDetailsView.cpp b/Engine/Source/Editor/PropertyEditor/Private/SDetailsView.cpp
--- a/Engine/Source/Editor/PropertyEditor/Private/SDetailsView.cpp
+++ b/Engine/Source/Editor/PropertyEditor/Private/SDetailsView.cpp
@@ -340,11 +340,12 @@ void SDetailsView::Construct(const FArguments& InArgs, const FDetailsViewArgs& I
 		];
 
 	FilterRowVBox->AddSlot()
-		.Padding(30, 2, 30, 7)
+		.Padding(8, 2, 8, 7)
 		.AutoHeight()
 		[
 			SAssignNew(SectionSelectorBox, SWrapBox)
 			.UseAllottedSize(true)
+			.InnerSlotPadding(FVector2D(4,4))
 		];
 
 	RebuildSectionSelector();
@@ -1227,7 +1228,7 @@ void SDetailsView::RebuildSectionSelector()
 		-> TSharedRef<SWidget>
 	{
 		return SNew(SBox)
-			.Padding(FMargin(0, 0, 2, 0))
+			.Padding(FMargin(0))
 			[
 				SNew(SCheckBox)
 				.Style(FAppStyle::Get(), "DetailsView.SectionButton")
