Fix return value of SDetailSingleItemRow::OnContextMenuOpening for newly-added menu items

diff --git a/Engine/Source/Editor/PropertyEditor/Private/SDetailSingleItemRow.cpp b/Engine/Source/Editor/PropertyEditor/Private/SDetailSingleItemRow.cpp
--- a/Engine/Source/Editor/PropertyEditor/Private/SDetailSingleItemRow.cpp
+++ b/Engine/Source/Editor/PropertyEditor/Private/SDetailSingleItemRow.cpp
@@ -754,6 +754,8 @@ bool SDetailSingleItemRow::OnContextMenuOpening(FMenuBuilder& MenuBuilder)
 					  FIsActionChecked::CreateSP(this, &SDetailSingleItemRow::IsBlacklistChecked)),
 			NAME_None,
 			EUserInterfaceActionType::Check);
+			
+		bAddedMenuEntry = true;
 	}
 
 	if (WidgetRow.CustomMenuItems.Num() > 0)
