Fix Horde Issue 71820

Error on conflict resolution.

diff --git a/Engine/Source/Editor/PropertyEditor/Private/PropertyRowGenerator.h b/Engine/Source/Editor/PropertyEditor/Private/PropertyRowGenerator.h
--- a/Engine/Source/Editor/PropertyEditor/Private/PropertyRowGenerator.h
+++ b/Engine/Source/Editor/PropertyEditor/Private/PropertyRowGenerator.h
@@ -93,10 +93,6 @@ public:
 	virtual void RegisterInstancedCustomPropertyTypeLayout(FName PropertyTypeName, FOnGetPropertyTypeCustomizationInstance PropertyTypeLayoutDelegate, TSharedPtr<IPropertyTypeIdentifier> Identifier = nullptr) override;
 	virtual void UnregisterInstancedCustomPropertyLayout(UStruct* Class) override;
 	virtual void UnregisterInstancedCustomPropertyTypeLayout(FName PropertyTypeName, TSharedPtr<IPropertyTypeIdentifier> Identifier = nullptr) override;
-	virtual TSharedPtr<FAssetThumbnailPool> GetGeneratedThumbnailPool() override
-	{ 
-		return GetThumbnailPool();
-	}
 	virtual void SetCustomValidatePropertyNodesFunction(FOnValidatePropertyRowGeneratorNodes InCustomValidatePropertyNodesFunction) override
 	{
 		CustomValidatePropertyNodesFunction = MoveTemp(InCustomValidatePropertyNodesFunction);
