Give names to additional properties in Details View
Allow property whitelist entries to specify their owner

diff --git a/Engine/Source/Editor/DetailCustomizations/Private/ComponentTransformDetails.h b/Engine/Source/Editor/DetailCustomizations/Private/ComponentTransformDetails.h
--- a/Engine/Source/Editor/DetailCustomizations/Private/ComponentTransformDetails.h
+++ b/Engine/Source/Editor/DetailCustomizations/Private/ComponentTransformDetails.h
@@ -45,7 +45,7 @@ public:
 	virtual void GenerateHeaderRowContent( FDetailWidgetRow& NodeRow ) override {}
 	virtual void GenerateChildContent( IDetailChildrenBuilder& ChildrenBuilder ) override;
 	virtual bool RequiresTick() const override { return true; }
-	virtual FName GetName() const override { return NAME_None; }
+	virtual FName GetName() const override { return "Transform"; }
 	virtual bool InitiallyCollapsed() const override { return false; }
 	virtual void Tick( float DeltaTime ) override;
 	virtual void SetOnRebuildChildren( FSimpleDelegate OnRebuildChildren ) override{}
diff --git a/Engine/Source/Editor/DetailCustomizations/Private/SkeletalMeshComponentDetails.cpp b/Engine/Source/Editor/DetailCustomizations/Private/SkeletalMeshComponentDetails.cpp
--- a/Engine/Source/Editor/DetailCustomizations/Private/SkeletalMeshComponentDetails.cpp
+++ b/Engine/Source/Editor/DetailCustomizations/Private/SkeletalMeshComponentDetails.cpp
@@ -100,7 +100,8 @@ void FSkeletalMeshComponentDetails::UpdateAnimationCategory(IDetailLayoutBuilder
 	AnimationModeHandle = DetailBuilder.GetProperty(GET_MEMBER_NAME_CHECKED(USkeletalMeshComponent, AnimationMode));
 	check (AnimationModeHandle->IsValidHandle());
 
-	AnimationBlueprintHandle = DetailBuilder.GetProperty(GET_MEMBER_NAME_CHECKED(USkeletalMeshComponent, AnimClass));
+	const FName AnimationBlueprintName = GET_MEMBER_NAME_CHECKED(USkeletalMeshComponent, AnimClass);
+	AnimationBlueprintHandle = DetailBuilder.GetProperty(AnimationBlueprintName);
 	check(AnimationBlueprintHandle->IsValidHandle());
 
 	AnimationCategory.AddProperty(AnimationModeHandle);
@@ -110,6 +111,7 @@ void FSkeletalMeshComponentDetails::UpdateAnimationCategory(IDetailLayoutBuilder
 
 	DetailBuilder.HideProperty(AnimationBlueprintHandle);
 	AnimationCategory.AddCustomRow(AnimationBlueprintHandle->GetPropertyDisplayName())
+		.RowTag(AnimationBlueprintName)
 		.Visibility(BlueprintVisibility)
 		.NameContent()
 		[
diff --git a/Engine/Source/Editor/PropertyEditor/Private/MaterialList.cpp b/Engine/Source/Editor/PropertyEditor/Private/MaterialList.cpp
--- a/Engine/Source/Editor/PropertyEditor/Private/MaterialList.cpp
+++ b/Engine/Source/Editor/PropertyEditor/Private/MaterialList.cpp
@@ -425,6 +425,7 @@ void FMaterialList::GenerateChildContent( IDetailChildrenBuilder& ChildrenBuilde
 	DisplayedMaterials.Empty();
 	if( MaterialListBuilder->GetNumMaterials() > 0 )
 	{
+		const FName RowTagName = "Material";
 		DisplayedMaterials = MaterialListBuilder->MaterialSlots;
 
 		MaterialListBuilder->Sort();
@@ -454,6 +455,7 @@ void FMaterialList::GenerateChildContent( IDetailChildrenBuilder& ChildrenBuilde
 					FFormatNamedArguments Arguments;
 					Arguments.Add(TEXT("ElementSlot"), CurrentSlot);
 					ChildRow
+					.RowTag(RowTagName)
 					.ValueContent()
 					.MaxDesiredWidth(0.0f)// No Max Width
 					[
@@ -474,6 +476,7 @@ void FMaterialList::GenerateChildContent( IDetailChildrenBuilder& ChildrenBuilde
 					bDisplayAllMaterialsInSlot = false;
 
 					FDetailWidgetRow& ChildRow = ChildrenBuilder.AddCustomRow( FText::GetEmpty() );
+					ChildRow.RowTag(RowTagName);
 
 					AddMaterialItem( ChildRow, CurrentSlot, FMaterialListItem( NULL, CurrentSlot, true ), !bDisplayAllMaterialsInSlot );
 				}
@@ -488,6 +491,7 @@ void FMaterialList::GenerateChildContent( IDetailChildrenBuilder& ChildrenBuilde
 			if( bDisplayAllMaterialsInSlot )
 			{
 				FDetailWidgetRow& ChildRow = ChildrenBuilder.AddCustomRow( Material.Material.IsValid()? FText::FromString(Material.Material->GetName()) : FText::GetEmpty() );
+				ChildRow.RowTag(RowTagName);
 
 				AddMaterialItem( ChildRow, CurrentSlot, Material, !bDisplayAllMaterialsInSlot );
 			}
diff --git a/Engine/Source/Editor/PropertyEditor/Private/PropertyEditorWhitelist.cpp b/Engine/Source/Editor/PropertyEditor/Private/PropertyEditorWhitelist.cpp
--- a/Engine/Source/Editor/PropertyEditor/Private/PropertyEditorWhitelist.cpp
+++ b/Engine/Source/Editor/PropertyEditor/Private/PropertyEditorWhitelist.cpp
@@ -28,36 +28,36 @@ void FPropertyEditorWhitelist::ClearWhitelist()
 	WhitelistUpdatedDelegate.Broadcast();
 }
 
-void FPropertyEditorWhitelist::AddToWhitelist(TSoftObjectPtr<UStruct> Struct, const FName PropertyName)
+void FPropertyEditorWhitelist::AddToWhitelist(TSoftObjectPtr<UStruct> Struct, const FName PropertyName, const FName Owner)
 {
 	FBlacklistNames& Blacklist = RawPropertyEditorWhitelist.FindOrAdd(Struct);
-	Blacklist.AddWhitelistItem(NAME_None, PropertyName);
+	Blacklist.AddWhitelistItem(Owner, PropertyName);
 	CachedPropertyEditorWhitelist.Reset();
 	WhitelistUpdatedDelegate.Broadcast();
 }
 
-void FPropertyEditorWhitelist::RemoveFromWhitelist(TSoftObjectPtr<UStruct> Struct, const FName PropertyName)
+void FPropertyEditorWhitelist::RemoveFromWhitelist(TSoftObjectPtr<UStruct> Struct, const FName PropertyName, const FName Owner)
 {
 	FBlacklistNames& Blacklist = RawPropertyEditorWhitelist.FindOrAdd(Struct);
-	if (Blacklist.RemoveWhitelistItem(NAME_None, PropertyName))
+	if (Blacklist.RemoveWhitelistItem(Owner, PropertyName))
 	{
 		CachedPropertyEditorWhitelist.Reset();
 		WhitelistUpdatedDelegate.Broadcast();
 	}
 }
 
-void FPropertyEditorWhitelist::AddToBlacklist(TSoftObjectPtr<UStruct> Struct, const FName PropertyName)
+void FPropertyEditorWhitelist::AddToBlacklist(TSoftObjectPtr<UStruct> Struct, const FName PropertyName, const FName Owner)
 {
 	FBlacklistNames& Blacklist = RawPropertyEditorWhitelist.FindOrAdd(Struct);
-	Blacklist.AddBlacklistItem(NAME_None, PropertyName);
+	Blacklist.AddBlacklistItem(Owner, PropertyName);
 	CachedPropertyEditorWhitelist.Reset();
 	WhitelistUpdatedDelegate.Broadcast();
 }
 
-void FPropertyEditorWhitelist::RemoveFromBlacklist(TSoftObjectPtr<UStruct> Struct, const FName PropertyName)
+void FPropertyEditorWhitelist::RemoveFromBlacklist(TSoftObjectPtr<UStruct> Struct, const FName PropertyName, const FName Owner)
 {
 	FBlacklistNames& Blacklist = RawPropertyEditorWhitelist.FindOrAdd(Struct);
-	if (Blacklist.RemoveBlacklistItem(NAME_None, PropertyName))
+	if (Blacklist.RemoveBlacklistItem(Owner, PropertyName))
 	{
 		CachedPropertyEditorWhitelist.Reset();
 		WhitelistUpdatedDelegate.Broadcast();
diff --git a/Engine/Source/Editor/PropertyEditor/Private/SDetailSingleItemRow.cpp b/Engine/Source/Editor/PropertyEditor/Private/SDetailSingleItemRow.cpp
--- a/Engine/Source/Editor/PropertyEditor/Private/SDetailSingleItemRow.cpp
+++ b/Engine/Source/Editor/PropertyEditor/Private/SDetailSingleItemRow.cpp
@@ -882,14 +882,15 @@ void SDetailSingleItemRow::OnToggleWhitelist() const
 	const TSharedPtr<FDetailTreeNode> Owner = OwnerTreeNode.Pin();
 	if (Owner)
 	{
+		const FName OwnerName = "DetailRowContextMenu";
 		const UStruct* ExactStruct = GetExactStructForProperty(Owner->GetParentBaseStructure(), Owner->GetNodeName());
 		if (IsWhitelistChecked())
 		{
-			FPropertyEditorWhitelist::Get().RemoveFromWhitelist(ExactStruct, Owner->GetNodeName());
+			FPropertyEditorWhitelist::Get().RemoveFromWhitelist(ExactStruct, Owner->GetNodeName(), OwnerName);
 		}
 		else
 		{
-			FPropertyEditorWhitelist::Get().AddToWhitelist(ExactStruct, Owner->GetNodeName());
+			FPropertyEditorWhitelist::Get().AddToWhitelist(ExactStruct, Owner->GetNodeName(), OwnerName);
 		}
 	}
 }
@@ -909,14 +910,15 @@ void SDetailSingleItemRow::OnToggleBlacklist() const
 	const TSharedPtr<FDetailTreeNode> Owner = OwnerTreeNode.Pin();
 	if (Owner)
 	{
+		const FName OwnerName = "DetailRowContextMenu";
 		const UStruct* ExactStruct = GetExactStructForProperty(Owner->GetParentBaseStructure(), Owner->GetNodeName());
 		if (IsBlacklistChecked())
 		{
-			FPropertyEditorWhitelist::Get().RemoveFromBlacklist(ExactStruct, Owner->GetNodeName());
+			FPropertyEditorWhitelist::Get().RemoveFromBlacklist(ExactStruct, Owner->GetNodeName(), OwnerName);
 		}
 		else
 		{
-			FPropertyEditorWhitelist::Get().AddToBlacklist(ExactStruct, Owner->GetNodeName());
+			FPropertyEditorWhitelist::Get().AddToBlacklist(ExactStruct, Owner->GetNodeName(), OwnerName);
 		}
 	}
 }
diff --git a/Engine/Source/Editor/PropertyEditor/Public/MaterialList.h b/Engine/Source/Editor/PropertyEditor/Public/MaterialList.h
--- a/Engine/Source/Editor/PropertyEditor/Public/MaterialList.h
+++ b/Engine/Source/Editor/PropertyEditor/Public/MaterialList.h
@@ -166,7 +166,7 @@ private:
 	virtual void Tick( float DeltaTime ) override;
 	virtual void GenerateHeaderRowContent( FDetailWidgetRow& NodeRow ) override;
 	virtual void GenerateChildContent( IDetailChildrenBuilder& ChildrenBuilder ) override;
-	virtual FName GetName() const override { return NAME_None; }
+	virtual FName GetName() const override { return "MaterialList"; }
 	virtual bool InitiallyCollapsed() const override { return bAllowCollpase; }
 
 	/**
diff --git a/Engine/Source/Editor/PropertyEditor/Public/PropertyEditorWhitelist.h b/Engine/Source/Editor/PropertyEditor/Public/PropertyEditorWhitelist.h
--- a/Engine/Source/Editor/PropertyEditor/Public/PropertyEditorWhitelist.h
+++ b/Engine/Source/Editor/PropertyEditor/Public/PropertyEditorWhitelist.h
@@ -26,13 +26,13 @@ public:
 	void ClearWhitelist();
 
 	/** Add a specific property to a UStruct's whitelist */
-	void AddToWhitelist(TSoftObjectPtr<UStruct> Struct, const FName PropertyName);
+	void AddToWhitelist(TSoftObjectPtr<UStruct> Struct, const FName PropertyName, const FName Owner = NAME_None);
 	/** Remove a specific property from a UStruct's whitelist */
-	void RemoveFromWhitelist(TSoftObjectPtr<UStruct> Struct, const FName PropertyName);
+	void RemoveFromWhitelist(TSoftObjectPtr<UStruct> Struct, const FName PropertyName, const FName Owner = NAME_None);
 	/** Add a specific property to a UStruct's blacklist */
-	void AddToBlacklist(TSoftObjectPtr<UStruct> Struct, const FName PropertyName);
+	void AddToBlacklist(TSoftObjectPtr<UStruct> Struct, const FName PropertyName, const FName Owner = NAME_None);
 	/** Remove a specific property from a UStruct's blacklist */
-    void RemoveFromBlacklist(TSoftObjectPtr<UStruct> Struct, const FName PropertyName);
+    void RemoveFromBlacklist(TSoftObjectPtr<UStruct> Struct, const FName PropertyName, const FName Owner = NAME_None);
 
 	/** When the whitelist or blacklist for any struct was added to or removed from. */
     FWhitelistUpdated WhitelistUpdatedDelegate;
