Add ValidEnumValues to PropertyMetadata enums (this does NOT introduce a new property metadata specifier, but rather just increase its visibility)

#rnx none

diff --git a/Engine/Source/Editor/PropertyEditor/Private/PropertyEditorHelpers.h b/Engine/Source/Editor/PropertyEditor/Private/PropertyEditorHelpers.h
--- a/Engine/Source/Editor/PropertyEditor/Private/PropertyEditorHelpers.h
+++ b/Engine/Source/Editor/PropertyEditor/Private/PropertyEditorHelpers.h
@@ -240,9 +240,9 @@ namespace PropertyEditorHelpers
 	void CollectObjectNodes( TSharedPtr<FPropertyNode> StartNode, TArray<FObjectPropertyNode*>& OutObjectNodes );
 
 	/**
-	 * Returns any enums that are explicitly allowed by the "AllowedEnumValues" metadata on FProperty using the specified enum.
+	 * Returns any enums that are explicitly allowed by the "ValidEnumValues" metadata on FProperty using the specified enum.
 	 *
-	 * @param Property	The property which may contain the "AllowedEnumValues" metadata
+	 * @param Property	The property which may contain the "ValidEnumValues" metadata
 	 * @param InEnum	The enum to search
 	 * @return The array of allowed enums.  NOTE: If an empty array is returned all enum values are allowed.  It is an error for a property to hide all enum values so that state is undefined here.
 	 */
diff --git a/Engine/Source/Runtime/CoreUObject/Public/UObject/ObjectMacros.h b/Engine/Source/Runtime/CoreUObject/Public/UObject/ObjectMacros.h
--- a/Engine/Source/Runtime/CoreUObject/Public/UObject/ObjectMacros.h
+++ b/Engine/Source/Runtime/CoreUObject/Public/UObject/ObjectMacros.h
@@ -1350,6 +1350,9 @@ namespace UM
 
 		/// [PropertyMetadata] The property can be exposed as a data pin, but is hidden by default.
 		PinHiddenByDefault,
+
+		/// [PropertyMetadata] Used for enum properties to define the subset of valid values as a comma-separated string; values outside of this subset are not made available.
+		ValidEnumValues,
 	};
 
 	// Metadata usable in UPROPERTY for customizing the behavior of Persona and UMG
