PropertyEditor: Add ValidEnumValues metadata support to BitmaskEnum properties

#rnx

diff --git a/Engine/Source/Editor/PropertyEditor/Private/UserInterface/PropertyEditor/SPropertyEditorNumeric.h b/Engine/Source/Editor/PropertyEditor/Private/UserInterface/PropertyEditor/SPropertyEditorNumeric.h
--- a/Engine/Source/Editor/PropertyEditor/Private/UserInterface/PropertyEditor/SPropertyEditorNumeric.h
+++ b/Engine/Source/Editor/PropertyEditor/Private/UserInterface/PropertyEditor/SPropertyEditorNumeric.h
@@ -20,6 +20,7 @@
 #include "Textures/SlateIcon.h"
 #include "PropertyHandle.h"
 #include "Presentation/PropertyEditor/PropertyEditor.h"
+#include "PropertyEditorHelpers.h"
 #include "UserInterface/PropertyEditor/PropertyEditorConstants.h"
 #include "Framework/Commands/UIAction.h"
 #include "Framework/MultiBox/MultiBoxBuilder.h"
@@ -88,12 +89,17 @@ public:
 						}
 					};
 
+					const TArray<FName> AllowedPropertyEnums = PropertyEditorHelpers::GetValidEnumsFromPropertyOverride(Prop, BitmaskEnum);
 					// Note: This loop doesn't include (BitflagsEnum->NumEnums() - 1) in order to skip the implicit "MAX" value that gets added to the enum type at compile time.
 					for (int32 BitmaskEnumIndex = 0; BitmaskEnumIndex < BitmaskEnum->NumEnums() - 1; ++BitmaskEnumIndex)
 					{
 						const int64 EnumValue = BitmaskEnum->GetValueByIndex(BitmaskEnumIndex);
-						const bool bIsHidden = BitmaskEnum->HasMetaData(TEXT("Hidden"), BitmaskEnumIndex);
-						if (EnumValue >= 0 && !bIsHidden)
+						bool bShouldBeHidden = BitmaskEnum->HasMetaData(TEXT("Hidden"), BitmaskEnumIndex);
+						if (!bShouldBeHidden && AllowedPropertyEnums.Num() > 0)
+						{
+							bShouldBeHidden = AllowedPropertyEnums.Find(BitmaskEnum->GetNameByIndex(BitmaskEnumIndex)) == INDEX_NONE;
+						}
+						if (EnumValue >= 0 && !bShouldBeHidden)
 						{
 							if (bUseEnumValuessAsMaskValues)
 							{
