Added support for parentheses in EditConditions. Added as a result of a UDN feature request.

diff --git a/Engine/Source/Editor/PropertyEditor/Private/EditConditionParser.cpp b/Engine/Source/Editor/PropertyEditor/Private/EditConditionParser.cpp
--- a/Engine/Source/Editor/PropertyEditor/Private/EditConditionParser.cpp
+++ b/Engine/Source/Editor/PropertyEditor/Private/EditConditionParser.cpp
@@ -25,9 +25,11 @@ namespace EditConditionParserTokens
 	const TCHAR* const FMultiply::Moniker = TEXT("*");
 	const TCHAR* const FDivide::Moniker = TEXT("/");
 	const TCHAR* const FBitwiseAnd::Moniker = TEXT("&");
+	const TCHAR* const FSubExpressionStart::Moniker = TEXT("(");
+	const TCHAR* const FSubExpressionEnd::Moniker = TEXT(")");
 }
 
-static const TCHAR PropertyBreakingChars[] = { '|', '=', '&', '>', '<', '!', '+', '-', '*', '/', ' ', '\t' };
+static const TCHAR PropertyBreakingChars[] = { '|', '=', '&', '>', '<', '!', '+', '-', '*', '/', ' ', '\t', '(', ')' };
 
 static TOptional<FExpressionError> ConsumeBool(FExpressionTokenConsumer& Consumer)
 {
@@ -634,6 +636,8 @@ FEditConditionParser::FEditConditionParser()
 	TokenDefinitions.DefineToken(&ExpressionParser::ConsumeSymbol<FMultiply>);
 	TokenDefinitions.DefineToken(&ExpressionParser::ConsumeSymbol<FDivide>);
 	TokenDefinitions.DefineToken(&ExpressionParser::ConsumeSymbol<FBitwiseAnd>);
+	TokenDefinitions.DefineToken(&ExpressionParser::ConsumeSymbol<FSubExpressionStart>);
+	TokenDefinitions.DefineToken(&ExpressionParser::ConsumeSymbol<FSubExpressionEnd>);
 	TokenDefinitions.DefineToken(&ExpressionParser::ConsumeNumber);
 	TokenDefinitions.DefineToken(&ConsumeNullPtr);
 	TokenDefinitions.DefineToken(&ConsumeBool);
@@ -653,6 +657,7 @@ FEditConditionParser::FEditConditionParser()
 	ExpressionGrammar.DefineBinaryOperator<FMultiply>(1);
 	ExpressionGrammar.DefineBinaryOperator<FDivide>(1);
 	ExpressionGrammar.DefinePreUnaryOperator<FNot>();
+	ExpressionGrammar.DefineGrouping<FSubExpressionStart, FSubExpressionEnd>();
 
 	// POINTER EQUALITY
 	OperatorJumpTable.MapBinary<FEqual>([](const FPropertyToken& A, const FPropertyToken& B, const IEditConditionContext* Context) -> FExpressionResult
@@ -687,7 +692,7 @@ FEditConditionParser::FEditConditionParser()
 	CreateEnumOperators(OperatorJumpTable);
 }
 
-TOptional<bool> FEditConditionParser::Evaluate(const FEditConditionExpression& Expression, const IEditConditionContext& Context) const
+TValueOrError<bool, FText> FEditConditionParser::Evaluate(const FEditConditionExpression& Expression, const IEditConditionContext& Context) const
 {
 	using namespace EditConditionParserTokens;
 	
@@ -697,7 +702,7 @@ TOptional<bool> FEditConditionParser::Evaluate(const FEditConditionExpression& E
 		const bool* BoolResult = Result.GetValue().Cast<bool>();
 		if (BoolResult != nullptr)
 		{
-			return *BoolResult;
+			return MakeValue(*BoolResult);
 		}
 
 		const FPropertyToken* PropertyResult = Result.GetValue().Cast<FPropertyToken>();
@@ -706,12 +711,12 @@ TOptional<bool> FEditConditionParser::Evaluate(const FEditConditionExpression& E
 			TOptional<bool> PropertyValue = Context.GetBoolValue(PropertyResult->PropertyName);
 			if (PropertyValue.IsSet())
 			{
-				return PropertyValue.GetValue();
+				return MakeValue(PropertyValue.GetValue());
 			}
 		}
 	}
 
-	return TOptional<bool>();
+	return MakeError(Result.GetError().Text);
 }
 
 TSharedPtr<FEditConditionExpression> FEditConditionParser::Parse(const FString& ExpressionString) const
diff --git a/Engine/Source/Editor/PropertyEditor/Private/EditConditionParser.h b/Engine/Source/Editor/PropertyEditor/Private/EditConditionParser.h
--- a/Engine/Source/Editor/PropertyEditor/Private/EditConditionParser.h
+++ b/Engine/Source/Editor/PropertyEditor/Private/EditConditionParser.h
@@ -6,6 +6,7 @@
 #include "Misc/ExpressionParserTypes.h"
 #include "UObject/WeakObjectPtr.h"
 #include "UObject/WeakObjectPtrTemplates.h"
+#include "Templates/ValueOrError.h"
 
 namespace EditConditionParserTokens
 {
@@ -112,7 +113,7 @@ public:
 	 * Evaluate the given expression within the given context.
 	 * @returns The result of the evaluated expression if valid, invalid TOptional if the evaluation failed or produced a non-bool result.
 	 */
-	TOptional<bool> Evaluate(const FEditConditionExpression& Expression, const IEditConditionContext& Context) const;
+	TValueOrError<bool, FText> Evaluate(const FEditConditionExpression& Expression, const IEditConditionContext& Context) const;
 
 private:
 	FTokenDefinitions TokenDefinitions;
diff --git a/Engine/Source/Editor/PropertyEditor/Private/EditConditionParserTests.cpp b/Engine/Source/Editor/PropertyEditor/Private/EditConditionParserTests.cpp
--- a/Engine/Source/Editor/PropertyEditor/Private/EditConditionParserTests.cpp
+++ b/Engine/Source/Editor/PropertyEditor/Private/EditConditionParserTests.cpp
@@ -205,16 +205,17 @@ static bool CanEvaluate(const FEditConditionParser& Parser, const IEditCondition
 		return false;
 	}
 
-	TOptional<bool> Result = Parser.Evaluate(*Parsed.Get(), Context);
-	if (!Result.IsSet())
+	TValueOrError<bool, FText> Result = Parser.Evaluate(*Parsed.Get(), Context);
+	if (!Result.IsValid())
 	{
-		ensureMsgf(false, TEXT("Expression failed to evaluate: %s"), *Expression);
+		ensureMsgf(false, TEXT("Expression failed to evaluate: %s, Error: %s"), *Expression, *Result.GetError().ToString());
 		return false;
 	}
 
 	if (Result.GetValue() != Expected)
 	{
-		ensureMsgf(false, TEXT("Expression evaluated to unexpected value."), *Expression);
+		auto BoolToString = [](bool Value) { return Value ? TEXT("true") : TEXT("false"); };
+		ensureMsgf(false, TEXT("Expression evaluated to unexpected value: %s, Expected: %s, Actual: %s"), *Expression, BoolToString(Expected), BoolToString(Result.GetValue()));
 		return false;
 	}
 
@@ -592,6 +593,67 @@ bool FEditConditionParser_EvaluatePointers::RunTest(const FString& Parameters)
 	return bAllResults;
 }
 
+IMPLEMENT_SIMPLE_AUTOMATION_TEST(FEditConditionParser_Grouping, "EditConditionParser.Grouping", EAutomationTestFlags::EditorContext | EAutomationTestFlags::EngineFilter)
+bool FEditConditionParser_Grouping::RunTest(const FString& Parameters)
+{
+	UEditConditionTestObject* TestObject = NewObject<UEditConditionTestObject>();
+	TestObject->AddToRoot();
+
+	TSharedPtr<FObjectPropertyNode> ObjectNode(new FObjectPropertyNode);
+	ObjectNode->AddObject(TestObject);
+
+	FPropertyNodeInitParams InitParams;
+	ObjectNode->InitNode(InitParams);
+
+	static const FName BoolPropertyName(TEXT("BoolProperty"));
+	TSharedPtr<FPropertyNode> PropertyNode = ObjectNode->FindChildPropertyNode(BoolPropertyName, true);
+	FEditConditionContext Context(*PropertyNode.Get());
+
+	bool bAllResults = true;
+
+	FEditConditionParser Parser;
+
+	bAllResults &= CanEvaluate(Parser, Context, TEXT("(true)"), true);
+	bAllResults &= CanEvaluate(Parser, Context, TEXT("(true) == true"), true);
+	bAllResults &= CanEvaluate(Parser, Context, TEXT("(true) == false"), false);
+	bAllResults &= CanEvaluate(Parser, Context, TEXT("(true) == (true)"), true);
+	bAllResults &= CanEvaluate(Parser, Context, TEXT("((true == false) == false) == true"), true);
+
+	TestObject->DoubleProperty = 5.0;
+	bAllResults &= CanEvaluate(Parser, Context, TEXT("(DoubleProperty == 5.0)"), true);
+	bAllResults &= CanEvaluate(Parser, Context, TEXT("(DoubleProperty == 5.0) == false"), false);
+	bAllResults &= CanEvaluate(Parser, Context, TEXT("(DoubleProperty == 5.0) == (true == true)"), true);
+
+	TestObject->EnumProperty = EditConditionTestEnum::First;
+	bAllResults &= CanEvaluate(Parser, Context, TEXT("(EnumProperty == EditConditionTestEnum::First) == true"), true);
+	bAllResults &= CanEvaluate(Parser, Context, TEXT("(EnumProperty == EditConditionTestEnum::First) && (DoubleProperty == 5.0)"), true);
+	bAllResults &= CanEvaluate(Parser, Context, TEXT("(EnumProperty == EditConditionTestEnum::Second) && (DoubleProperty == 5.0)"), false);
+	bAllResults &= CanEvaluate(Parser, Context, TEXT("(EnumProperty == EditConditionTestEnum::Second) || (DoubleProperty == 5.0)"), true);
+
+	{
+		FTestEditConditionContext TestContext;
+
+		const FString EnumType = TEXT("TestEnum");
+		TestContext.SetupEnumType(EnumType);
+		TestContext.SetupEnumTypeValue(TEXT("Nil"), 0);
+		TestContext.SetupEnumTypeValue(TEXT("One"), 1 << 0);
+
+		TestContext.SetupInteger("FlagsProperty", 1);
+		TestContext.SetupDouble("DoubleProperty", 5.0);
+
+		bAllResults &= CanEvaluate(Parser, TestContext, TEXT("(FlagsProperty & TestEnum::Nil)"), false); 
+		bAllResults &= CanEvaluate(Parser, TestContext, TEXT("(FlagsProperty & TestEnum::One)"), true);
+		bAllResults &= CanEvaluate(Parser, TestContext, TEXT("(FlagsProperty & TestEnum::One) && (DoubleProperty == 5.0)"), true);
+		bAllResults &= CanEvaluate(Parser, TestContext, TEXT("(FlagsProperty & TestEnum::One) == false"), false);
+		bAllResults &= CanEvaluate(Parser, TestContext, TEXT("!(FlagsProperty & TestEnum::One)"), false);
+		bAllResults &= CanEvaluate(Parser, TestContext, TEXT("!(FlagsProperty & TestEnum::Nil)"), true);
+	}
+
+	TestObject->RemoveFromRoot();
+
+	return bAllResults;
+}
+
 IMPLEMENT_SIMPLE_AUTOMATION_TEST(FEditConditionParser_SingleBool, "EditConditionParser.SingleBool", EAutomationTestFlags::EditorContext | EAutomationTestFlags::EngineFilter)
 bool FEditConditionParser_SingleBool::RunTest(const FString& Parameters)
 {
diff --git a/Engine/Source/Editor/PropertyEditor/Private/PropertyNode.cpp b/Engine/Source/Editor/PropertyEditor/Private/PropertyNode.cpp
--- a/Engine/Source/Editor/PropertyEditor/Private/PropertyNode.cpp
+++ b/Engine/Source/Editor/PropertyEditor/Private/PropertyNode.cpp
@@ -961,8 +961,8 @@ bool FPropertyNode::IsEditConditionMet() const
 { 
 	if (HasEditCondition())
 	{
-		TOptional<bool> Result = EditConditionParser.Evaluate(*EditConditionExpression.Get(), *EditConditionContext.Get());
-		if (Result.IsSet())
+		TValueOrError<bool, FText> Result = EditConditionParser.Evaluate(*EditConditionExpression.Get(), *EditConditionContext.Get());
+		if (Result.IsValid())
 		{
 			return Result.GetValue();
 		}
