Moved FDetailColumnSizeData to its own file.

diff --git a/Engine/Plugins/Enterprise/DataprepEditor/Source/DataprepEditor/Private/Widgets/DataprepAssetView.cpp b/Engine/Plugins/Enterprise/DataprepEditor/Source/DataprepEditor/Private/Widgets/DataprepAssetView.cpp
--- a/Engine/Plugins/Enterprise/DataprepEditor/Source/DataprepEditor/Private/Widgets/DataprepAssetView.cpp
+++ b/Engine/Plugins/Enterprise/DataprepEditor/Source/DataprepEditor/Private/Widgets/DataprepAssetView.cpp
@@ -18,6 +18,7 @@
 #include "K2Node_AddComponent.h"
 #include "PropertyEditorModule.h"
 
+#include "DetailColumnSizeData.h"
 #include "DetailLayoutBuilder.h"
 #include "EditorFontGlyphs.h"
 #include "Modules/ModuleManager.h"
diff --git a/Engine/Plugins/Enterprise/DataprepEditor/Source/DataprepEditor/Private/Widgets/DataprepWidgets.cpp b/Engine/Plugins/Enterprise/DataprepEditor/Source/DataprepEditor/Private/Widgets/DataprepWidgets.cpp
--- a/Engine/Plugins/Enterprise/DataprepEditor/Source/DataprepEditor/Private/Widgets/DataprepWidgets.cpp
+++ b/Engine/Plugins/Enterprise/DataprepEditor/Source/DataprepEditor/Private/Widgets/DataprepWidgets.cpp
@@ -14,6 +14,7 @@
 #include "Widgets/Parameterization/SDataprepParameterizationLinkIcon.h"
 
 #include "ContentBrowserModule.h"
+#include "DetailColumnSizeData.h"
 #include "DetailLayoutBuilder.h"
 #include "Dialogs/DlgPickPath.h"
 #include "Editor.h"
@@ -77,15 +78,15 @@ namespace DataprepWidgetUtils
 			ColumnSizeData = InArgs._ColumnSizeData;
 
 			AddSlot()
-			.Value(ColumnSizeData->NameColumnWidth)
-			.OnSlotResized(ColumnSizeData->OnNameColumnResized)
+			.Value(ColumnSizeData->GetNameColumnWidth())
+			.OnSlotResized(ColumnSizeData->GetOnNameColumnResized())
 			[
 				InArgs._NameWidget.ToSharedRef()
 			];
 
 			AddSlot()
-			.Value(InArgs._ColumnSizeData->ValueColumnWidth)
-			.OnSlotResized(ColumnSizeData->OnValueColumnResized)
+			.Value(InArgs._ColumnSizeData->GetValueColumnWidth())
+			.OnSlotResized(ColumnSizeData->GetOnValueColumnResized())
 			[
 				SNew(SHorizontalBox)
 				+ SHorizontalBox::Slot()
@@ -108,14 +109,14 @@ namespace DataprepWidgetUtils
 				const FVector2D AllottedSize = AllottedGeometry.GetLocalSize();
 
 				FVector2D LocalPosition = FVector2D::ZeroVector;
-				FVector2D LocalSize(  AllottedSize.X * ColumnSizeData->NameColumnWidth.Get(0.f), AllottedSize.Y );
+				FVector2D LocalSize(AllottedSize.X * ColumnSizeData->GetNameColumnWidth().Get(0), AllottedSize.Y);
 
 				ArrangedChildren.AddWidget( EVisibility::Visible, AllottedGeometry.MakeChild( LeftChild.GetWidget(), LocalPosition, LocalSize ));
 
 				const SSplitter::FSlot& RightChild = Children[1];
 
 				LocalPosition = FVector2D(LocalSize.X, 0.f);
-				LocalSize.X = AllottedSize.X * ColumnSizeData->ValueColumnWidth.Get(0.f);
+				LocalSize.X = AllottedSize.X * ColumnSizeData->GetValueColumnWidth().Get(0);
 
 				ArrangedChildren.AddWidget( EVisibility::Visible, AllottedGeometry.MakeChild( RightChild.GetWidget(), LocalPosition, LocalSize ));
 			}
diff --git a/Engine/Plugins/Experimental/VirtualProductionUtilities/Source/VPMaterialsEditor/Private/SMaterialDynamicParametersPanelWidget.cpp b/Engine/Plugins/Experimental/VirtualProductionUtilities/Source/VPMaterialsEditor/Private/SMaterialDynamicParametersPanelWidget.cpp
--- a/Engine/Plugins/Experimental/VirtualProductionUtilities/Source/VPMaterialsEditor/Private/SMaterialDynamicParametersPanelWidget.cpp
+++ b/Engine/Plugins/Experimental/VirtualProductionUtilities/Source/VPMaterialsEditor/Private/SMaterialDynamicParametersPanelWidget.cpp
@@ -2,8 +2,7 @@
 
 #include "SMaterialDynamicParametersPanelWidget.h"
 
-#include "PropertyHandle.h"
-
+#include "DetailColumnSizeData.h"
 #include "DetailWidgetRow.h"
 #include "EditorStyleSet.h"
 #include "IDetailPropertyRow.h"
@@ -11,6 +10,7 @@
 #include "IPropertyRowGenerator.h"
 #include "PropertyCustomizationHelpers.h"
 #include "PropertyEditorModule.h"
+#include "PropertyHandle.h"
 #include "Framework/MultiBox/MultiBoxBuilder.h"
 #include "Framework/MultiBox/MultiBoxDefs.h"
 #include "Materials/MaterialInstance.h"
@@ -376,8 +376,8 @@ void SMaterialDynamicParametersOverviewTreeItem::Construct(const FArguments& InA
 						.PhysicalSplitterHandleSize(1.0f)
 						.HitDetectionSplitterHandleSize(5.0f)
 						+ SSplitter::Slot()
-						.Value(ColumnSizeData.NameColumnWidth)
-						.OnSlotResized(ColumnSizeData.OnNameColumnResized)
+						.Value(ColumnSizeData.GetNameColumnWidth())
+						.OnSlotResized(ColumnSizeData.GetOnNameColumnResized())
 						.Value(0.25f)
 						[
 							SNew(SHorizontalBox)
@@ -396,8 +396,8 @@ void SMaterialDynamicParametersOverviewTreeItem::Construct(const FArguments& InA
 							]
 						]
 						+ SSplitter::Slot()
-						.Value(ColumnSizeData.ValueColumnWidth)
-						.OnSlotResized(ColumnSizeData.OnValueColumnResized) 
+						.Value(ColumnSizeData.GetValueColumnWidth())
+						.OnSlotResized(ColumnSizeData.GetOnValueColumnResized()) 
 						[
 							SNew(SHorizontalBox)
 							+ SHorizontalBox::Slot()
diff --git a/Engine/Source/Editor/Kismet/Private/SPinValueInspector.cpp b/Engine/Source/Editor/Kismet/Private/SPinValueInspector.cpp
--- a/Engine/Source/Editor/Kismet/Private/SPinValueInspector.cpp
+++ b/Engine/Source/Editor/Kismet/Private/SPinValueInspector.cpp
@@ -8,7 +8,7 @@
 #include "Widgets/Input/SSearchBox.h"
 #include "Kismet2/BlueprintEditorUtils.h"
 #include "Kismet2/KismetDebugUtilities.h"
-#include "IDetailPropertyRow.h"
+#include "DetailColumnSizeData.h"
 #include "Widgets/SWindow.h"
 #include "Widgets/SToolTip.h"
 
diff --git a/Engine/Source/Editor/LevelEditor/Private/SInViewportDetails.cpp b/Engine/Source/Editor/LevelEditor/Private/SInViewportDetails.cpp
--- a/Engine/Source/Editor/LevelEditor/Private/SInViewportDetails.cpp
+++ b/Engine/Source/Editor/LevelEditor/Private/SInViewportDetails.cpp
@@ -135,8 +135,8 @@ public:
 				.PhysicalSplitterHandleSize(1.0f)
  				.HitDetectionSplitterHandleSize(5.0f)
  				+ SSplitter::Slot()
- 				.Value(ColumnSizeData.NameColumnWidth)
-				.OnSlotResized(ColumnSizeData.OnNameColumnResized)
+ 				.Value(ColumnSizeData.GetNameColumnWidth())
+				.OnSlotResized(ColumnSizeData.GetOnNameColumnResized())
  				[
 					SNew(SBox)
 					.HAlign(HAlign_Left)
@@ -147,8 +147,8 @@ public:
 					]
  				]
  				+ SSplitter::Slot()
- 				.Value(ColumnSizeData.ValueColumnWidth)
- 				.OnSlotResized(ColumnSizeData.OnValueColumnResized)
+ 				.Value(ColumnSizeData.GetValueColumnWidth())
+ 				.OnSlotResized(ColumnSizeData.GetOnValueColumnResized())
  				[
  					SNew(SBox)
  					.Padding(2.0f)
@@ -157,8 +157,8 @@ public:
  					]
 				]
 				+ SSplitter::Slot()
-				.Value(ColumnSizeData.RightColumnWidth)
-				.OnSlotResized(ColumnSizeData.OnRightColumnResized)
+				.Value(ColumnSizeData.GetRightColumnWidth())
+				.OnSlotResized(ColumnSizeData.GetOnRightColumnResized())
 				[
 					SNew(SHorizontalBox)
 					+ SHorizontalBox::Slot()
diff --git a/Engine/Source/Editor/LevelEditor/Private/SInViewportDetails.h b/Engine/Source/Editor/LevelEditor/Private/SInViewportDetails.h
--- a/Engine/Source/Editor/LevelEditor/Private/SInViewportDetails.h
+++ b/Engine/Source/Editor/LevelEditor/Private/SInViewportDetails.h
@@ -3,15 +3,14 @@
 #pragma once
 
 #include "CoreMinimal.h"
+#include "DetailColumnSizeData.h"
+#include "EditorUndoClient.h"
+#include "Framework/Text/SlateHyperlinkRun.h"
 #include "Layout/Visibility.h"
+#include "ToolMenuContext.h"
 #include "Widgets/DeclarativeSyntaxSupport.h"
 #include "Widgets/SCompoundWidget.h"
-#include "Framework/Text/SlateHyperlinkRun.h"
-#include "EditorUndoClient.h"
 #include "Widgets/Views/SListView.h"
-#include "IDetailPropertyRow.h"
-#include "ToolMenuContext.h"
-
 
 class AActor;
 class FSCSEditorTreeNode;
diff --git a/Engine/Source/Editor/MaterialEditor/Private/SMaterialLayersFunctionsTree.cpp b/Engine/Source/Editor/MaterialEditor/Private/SMaterialLayersFunctionsTree.cpp
--- a/Engine/Source/Editor/MaterialEditor/Private/SMaterialLayersFunctionsTree.cpp
+++ b/Engine/Source/Editor/MaterialEditor/Private/SMaterialLayersFunctionsTree.cpp
@@ -1050,10 +1050,10 @@ void SMaterialLayersFunctionsInstanceTreeItem::Construct(const FArguments& InArg
 						.Style(FEditorStyle::Get(), "DetailsView.Splitter")
 						.PhysicalSplitterHandleSize(1.0f)
 						.HitDetectionSplitterHandleSize(5.0f)
-						.HighlightedHandleIndex(Tree->ColumnSizeData.HoveredSplitterIndex)
+						.HighlightedHandleIndex(Tree->ColumnSizeData.GetHoveredSplitterIndex())
 						+ SSplitter::Slot()
-						.Value(Tree->ColumnSizeData.NameColumnWidth)
-						.OnSlotResized(Tree->ColumnSizeData.OnNameColumnResized)
+						.Value(Tree->ColumnSizeData.GetNameColumnWidth())
+						.OnSlotResized(Tree->ColumnSizeData.GetOnNameColumnResized())
 						.Value(0.25f)
 						[
 							SNew(SHorizontalBox)
@@ -1072,8 +1072,8 @@ void SMaterialLayersFunctionsInstanceTreeItem::Construct(const FArguments& InArg
 							]
 						]
 						+ SSplitter::Slot()
-						.Value(Tree->ColumnSizeData.ValueColumnWidth)
-						.OnSlotResized(Tree->ColumnSizeData.OnValueColumnResized)
+						.Value(Tree->ColumnSizeData.GetValueColumnWidth())
+						.OnSlotResized(Tree->ColumnSizeData.GetOnValueColumnResized())
 						[
 							SNew(SHorizontalBox)
 							.Clipping(EWidgetClipping::OnDemand)
@@ -1138,7 +1138,7 @@ FString SMaterialLayersFunctionsInstanceTreeItem::GetInstancePath(SMaterialLayer
 
 void SMaterialLayersFunctionsInstanceTree::Construct(const FArguments& InArgs)
 {
-	ColumnSizeData.OnValueColumnResized.Execute(0.5f);
+	ColumnSizeData.SetValueColumnWidth(0.5f);
 
 	MaterialEditorInstance = InArgs._InMaterialEditorInstance;
 	Wrapper = InArgs._InWrapper;
@@ -2540,8 +2540,8 @@ void SMaterialLayersFunctionsMaterialTreeItem::Construct(const FArguments& InArg
 						.PhysicalSplitterHandleSize(1.0f)
 						.HitDetectionSplitterHandleSize(5.0f)
 						+ SSplitter::Slot()
-						.Value(Tree->ColumnSizeData.NameColumnWidth)
-						.OnSlotResized(Tree->ColumnSizeData.OnNameColumnResized)
+						.Value(Tree->ColumnSizeData.GetNameColumnWidth())
+						.OnSlotResized(Tree->ColumnSizeData.GetOnNameColumnResized())
 						.Value(0.25f)
 						[
 							SNew(SHorizontalBox)
@@ -2560,8 +2560,8 @@ void SMaterialLayersFunctionsMaterialTreeItem::Construct(const FArguments& InArg
 							]
 						]
 						+ SSplitter::Slot()
-						.Value(Tree->ColumnSizeData.ValueColumnWidth)
-						.OnSlotResized(Tree->ColumnSizeData.OnValueColumnResized)
+						.Value(Tree->ColumnSizeData.GetValueColumnWidth())
+						.OnSlotResized(Tree->ColumnSizeData.GetOnValueColumnResized())
 						[
 							SNew(SHorizontalBox)
 							.Clipping(EWidgetClipping::OnDemand)
@@ -2608,7 +2608,7 @@ FString SMaterialLayersFunctionsMaterialTreeItem::GetInstancePath(SMaterialLayer
 
 void SMaterialLayersFunctionsMaterialTree::Construct(const FArguments& InArgs)
 {
-	ColumnSizeData.OnValueColumnResized.Execute(0.5f);
+	ColumnSizeData.SetValueColumnWidth(0.5f);
 
 	MaterialEditorInstance = InArgs._InMaterialEditorInstance;
 	Wrapper = InArgs._InWrapper;
diff --git a/Engine/Source/Editor/MaterialEditor/Private/SMaterialLayersFunctionsTree.h b/Engine/Source/Editor/MaterialEditor/Private/SMaterialLayersFunctionsTree.h
--- a/Engine/Source/Editor/MaterialEditor/Private/SMaterialLayersFunctionsTree.h
+++ b/Engine/Source/Editor/MaterialEditor/Private/SMaterialLayersFunctionsTree.h
@@ -3,18 +3,19 @@
 #pragma once
 
 #include "CoreMinimal.h"
-#include "UObject/ObjectMacros.h"
-#include "UObject/Object.h"
+#include "DetailColumnSizeData.h"
 #include "Engine/EngineTypes.h"
-#include "Widgets/DeclarativeSyntaxSupport.h"
-#include "Widgets/SCompoundWidget.h"
-#include "Widgets/Views/STreeView.h"
-#include "Materials/Material.h"
-#include "MaterialEditor/MaterialEditorInstanceConstant.h"
-#include "IDetailTreeNode.h"
 #include "IDetailPropertyRow.h"
+#include "IDetailTreeNode.h"
+#include "MaterialEditor/MaterialEditorInstanceConstant.h"
 #include "MaterialPropertyHelpers.h"
+#include "Materials/Material.h"
 #include "PropertyCustomizationHelpers.h"
+#include "UObject/Object.h"
+#include "UObject/ObjectMacros.h"
+#include "Widgets/DeclarativeSyntaxSupport.h"
+#include "Widgets/SCompoundWidget.h"
+#include "Widgets/Views/STreeView.h"
 
 class IPropertyHandle;
 class SMaterialLayersFunctionsInstanceTree;
diff --git a/Engine/Source/Editor/MaterialEditor/Private/SMaterialParametersOverviewWidget.cpp b/Engine/Source/Editor/MaterialEditor/Private/SMaterialParametersOverviewWidget.cpp
--- a/Engine/Source/Editor/MaterialEditor/Private/SMaterialParametersOverviewWidget.cpp
+++ b/Engine/Source/Editor/MaterialEditor/Private/SMaterialParametersOverviewWidget.cpp
@@ -487,8 +487,8 @@ void SMaterialParametersOverviewTreeItem::Construct(const FArguments& InArgs, co
 						.PhysicalSplitterHandleSize(1.0f)
 						.HitDetectionSplitterHandleSize(5.0f)
 						+ SSplitter::Slot()
-						.Value(ColumnSizeData.NameColumnWidth)
-						.OnSlotResized(ColumnSizeData.OnNameColumnResized)
+						.Value(ColumnSizeData.GetNameColumnWidth())
+						.OnSlotResized(ColumnSizeData.GetOnNameColumnResized())
 						.Value(0.25f)
 						[
 							SNew(SHorizontalBox)
@@ -507,8 +507,8 @@ void SMaterialParametersOverviewTreeItem::Construct(const FArguments& InArgs, co
 							]
 						]
 						+ SSplitter::Slot()
-						.Value(ColumnSizeData.ValueColumnWidth)
-						.OnSlotResized(ColumnSizeData.OnValueColumnResized) 
+						.Value(ColumnSizeData.GetValueColumnWidth())
+						.OnSlotResized(ColumnSizeData.GetOnValueColumnResized())
 						[
 							SNew(SHorizontalBox)
 							+ SHorizontalBox::Slot()
diff --git a/Engine/Source/Editor/MaterialEditor/Private/SMaterialParametersOverviewWidget.h b/Engine/Source/Editor/MaterialEditor/Private/SMaterialParametersOverviewWidget.h
--- a/Engine/Source/Editor/MaterialEditor/Private/SMaterialParametersOverviewWidget.h
+++ b/Engine/Source/Editor/MaterialEditor/Private/SMaterialParametersOverviewWidget.h
@@ -3,17 +3,18 @@
 #pragma once
 
 #include "CoreMinimal.h"
-#include "UObject/ObjectMacros.h"
-#include "UObject/Object.h"
+#include "DetailColumnSizeData.h"
 #include "Engine/EngineTypes.h"
-#include "Widgets/DeclarativeSyntaxSupport.h"
-#include "Widgets/SCompoundWidget.h"
-#include "Widgets/Views/STreeView.h"
-#include "Materials/Material.h"
 #include "IDetailTreeNode.h"
-#include "IDetailPropertyRow.h"
 #include "MaterialPropertyHelpers.h"
+#include "Materials/Material.h"
 #include "PropertyCustomizationHelpers.h"
+#include "UObject/Object.h"
+#include "UObject/ObjectMacros.h"
+#include "Widgets/DeclarativeSyntaxSupport.h"
+#include "Widgets/SCompoundWidget.h"
+#include "Widgets/Views/STreeView.h"
+
 class IPropertyHandle;
 class SMaterialParametersOverviewTree;
 class UMaterialEditorPreviewParameters;
diff --git a/Engine/Source/Editor/PropertyEditor/Private/DetailFilter.h b/Engine/Source/Editor/PropertyEditor/Private/DetailFilter.h
--- a/Engine/Source/Editor/PropertyEditor/Private/DetailFilter.h
+++ b/Engine/Source/Editor/PropertyEditor/Private/DetailFilter.h
@@ -3,6 +3,7 @@
 #pragma once
 
 #include "CoreMinimal.h"
+#include "PropertyPath.h"
 
 /** Represents a filter which controls the visibility of items in the details view */
 struct FDetailFilter
diff --git a/Engine/Source/Editor/PropertyEditor/Private/IDetailsViewPrivate.h b/Engine/Source/Editor/PropertyEditor/Private/IDetailsViewPrivate.h
--- a/Engine/Source/Editor/PropertyEditor/Private/IDetailsViewPrivate.h
+++ b/Engine/Source/Editor/PropertyEditor/Private/IDetailsViewPrivate.h
@@ -145,7 +145,7 @@ public:
 	virtual void RestoreExpandedItems(TSharedRef<FPropertyNode> StartNode) = 0;
 
 	/** Column width accessibility */
-	virtual FDetailColumnSizeData& GetColumnSizeData() = 0;
+	virtual class FDetailColumnSizeData& GetColumnSizeData() = 0;
 
 	/** Does this details view allow favoriting? */
 	virtual bool IsFavoritingEnabled() const = 0;
diff --git a/Engine/Source/Editor/PropertyEditor/Private/SDetailSingleItemRow.cpp b/Engine/Source/Editor/PropertyEditor/Private/SDetailSingleItemRow.cpp
--- a/Engine/Source/Editor/PropertyEditor/Private/SDetailSingleItemRow.cpp
+++ b/Engine/Source/Editor/PropertyEditor/Private/SDetailSingleItemRow.cpp
@@ -375,8 +375,8 @@ void SDetailSingleItemRow::Construct( const FArguments& InArgs, FDetailLayoutCus
 					.Style(FEditorStyle::Get(), "DetailsView.Splitter")
 					.PhysicalSplitterHandleSize(1.0f)
 					.HitDetectionSplitterHandleSize(5.0f)
-					.HighlightedHandleIndex(ColumnSizeData.HoveredSplitterIndex)
-					.OnHandleHovered(ColumnSizeData.OnSplitterHandleHovered);
+					.HighlightedHandleIndex(ColumnSizeData.GetHoveredSplitterIndex())
+					.OnHandleHovered(ColumnSizeData.GetOnSplitterHandleHovered());
 
 			Widget = SNew(SBorder)
 				.BorderImage(FAppStyle::Get().GetBrush("DetailsView.CategoryMiddle"))
@@ -481,8 +481,8 @@ void SDetailSingleItemRow::Construct( const FArguments& InArgs, FDetailLayoutCus
 				// create Name column:
 				// | Name | Value | Right |
 				Splitter->AddSlot()
-					.Value(ColumnSizeData.NameColumnWidth)
-					.OnSlotResized(ColumnSizeData.OnNameColumnResized)
+					.Value(ColumnSizeData.GetNameColumnWidth())
+					.OnSlotResized(ColumnSizeData.GetOnNameColumnResized())
 					[
 						NameColumnBox
 					];
@@ -490,8 +490,8 @@ void SDetailSingleItemRow::Construct( const FArguments& InArgs, FDetailLayoutCus
 				// create Value column:
 				// | Name | Value | Right |
 				Splitter->AddSlot()
-					.Value(ColumnSizeData.ValueColumnWidth)
-					.OnSlotResized(ColumnSizeData.OnValueColumnResized) 
+					.Value(ColumnSizeData.GetValueColumnWidth())
+					.OnSlotResized(ColumnSizeData.GetOnValueColumnResized())
 					[
 						SNew(SHorizontalBox)
 						.Clipping(EWidgetClipping::OnDemand)
@@ -527,8 +527,8 @@ void SDetailSingleItemRow::Construct( const FArguments& InArgs, FDetailLayoutCus
 					];
 
 				Splitter->AddSlot()
-					.Value(ColumnSizeData.PropertyColumnWidth)
-					.OnSlotResized(ColumnSizeData.OnPropertyColumnResized)
+					.Value(ColumnSizeData.GetWholeRowColumnWidth())
+					.OnSlotResized(ColumnSizeData.GetOnWholeRowColumnResized())
 					[
 						NameColumnBox
 					];
@@ -574,9 +574,9 @@ void SDetailSingleItemRow::Construct( const FArguments& InArgs, FDetailLayoutCus
 			}
 
 			Splitter->AddSlot()
-				.Value(ColumnSizeData.RightColumnWidth)
-				.OnSlotResized(ColumnSizeData.OnRightColumnResized)
-				.MinSize(ColumnSizeData.RightColumnMinWidth)
+				.Value(ColumnSizeData.GetRightColumnWidth())
+				.OnSlotResized(ColumnSizeData.GetOnRightColumnResized())
+				.MinSize(ColumnSizeData.GetRightColumnMinWidth())
 			[
 				SNew(SBorder)
 				.BorderImage(FAppStyle::Get().GetBrush("DetailsView.CategoryMiddle"))
diff --git a/Engine/Source/Editor/PropertyEditor/Private/SDetailsViewBase.cpp b/Engine/Source/Editor/PropertyEditor/Private/SDetailsViewBase.cpp
--- a/Engine/Source/Editor/PropertyEditor/Private/SDetailsViewBase.cpp
+++ b/Engine/Source/Editor/PropertyEditor/Private/SDetailsViewBase.cpp
@@ -529,7 +529,7 @@ void SDetailsViewBase::SetKeyframeHandler( TSharedPtr<class IDetailKeyframeHandl
 		ExtraWidth = -22;
 	}
 	
-	const float NewWidth = ColumnSizeData.GetRightColumnMinWidth() + ExtraWidth;
+	const float NewWidth = ColumnSizeData.GetRightColumnMinWidth().Get(0) + ExtraWidth;
 	ColumnSizeData.SetRightColumnMinWidth(NewWidth);
 
 	KeyframeHandler = InKeyframeHandler;
@@ -906,10 +906,10 @@ void SDetailsViewBase::Tick( const FGeometry& AllottedGeometry, const double InC
 	HandlePendingCleanup();
 
 	FDetailsViewConfig* ViewConfig = GetMutableViewConfig();
-	if (ViewConfig != nullptr &&
-		ViewConfig->ValueColumnWidth != ColumnSizeData.ValueColumnWidth.Get(0))
+	if (ViewConfig != nullptr && 
+		ViewConfig->ValueColumnWidth != ColumnSizeData.GetValueColumnWidth().Get(0))
 	{
-		ViewConfig->ValueColumnWidth = ColumnSizeData.ValueColumnWidth.Get(0);
+		ViewConfig->ValueColumnWidth = ColumnSizeData.GetValueColumnWidth().Get(0);
 		SaveViewConfig();
 	}
 
diff --git a/Engine/Source/Editor/PropertyEditor/Private/SDetailsViewBase.h b/Engine/Source/Editor/PropertyEditor/Private/SDetailsViewBase.h
--- a/Engine/Source/Editor/PropertyEditor/Private/SDetailsViewBase.h
+++ b/Engine/Source/Editor/PropertyEditor/Private/SDetailsViewBase.h
@@ -4,25 +4,26 @@
 
 #include "CoreMinimal.h"
 #include "AssetThumbnail.h"
-#include "IDetailsView.h"
-#include "IDetailsViewPrivate.h"
-#include "IPropertyUtilities.h"
+#include "DetailColumnSizeData.h"
 #include "DetailFilter.h"
 #include "DetailTreeNode.h"
 #include "DetailsViewConfig.h"
+#include "Framework/Commands/UICommandList.h"
+#include "IDetailsView.h"
+#include "IDetailsViewPrivate.h"
+#include "IPropertyUtilities.h"
+#include "Input/Reply.h"
+#include "Layout/Visibility.h"
+#include "Misc/Attribute.h"
 #include "PropertyCustomizationHelpers.h"
 #include "PropertyEditorModule.h"
 #include "PropertyNode.h"
 #include "PropertyPath.h"
 #include "PropertyRowGenerator.h"
-#include "Framework/Commands/UICommandList.h"
-#include "Input/Reply.h"
-#include "Layout/Visibility.h"
-#include "Misc/Attribute.h"
-#include "Widgets/SWindow.h"
 #include "Widgets/Layout/SSplitter.h"
-#include "Widgets/Views/STableViewBase.h"
+#include "Widgets/SWindow.h"
 #include "Widgets/Views/STableRow.h"
+#include "Widgets/Views/STableViewBase.h"
 #include "Widgets/Views/STreeView.h"
 
 class FDetailCategoryImpl;
@@ -344,7 +345,7 @@ protected:
 	/** A mapping of type names to detail layout delegates, called when querying for custom detail layouts in this instance of the details view only */
 	FCustomPropertyTypeLayoutMap InstancedTypeToLayoutMap;
 	/** The current detail layout based on objects in this details panel.  There is one layout for each top level object node.*/
-	FDetailLayoutList DetailLayouts;
+	TArray<FDetailLayoutData> DetailLayouts;
 	/** Row for searching and view options */
 	TSharedPtr<SWidget> FilterRow;
 	/** Search box */
diff --git a/Engine/Source/Editor/PropertyEditor/Public/DetailColumnSizeData.h b/Engine/Source/Editor/PropertyEditor/Public/DetailColumnSizeData.h
new file mode 100644
--- /dev/null
+++ b/Engine/Source/Editor/PropertyEditor/Public/DetailColumnSizeData.h
@@ -0,0 +1,88 @@
+#pragma once
+
+#include "Delegates/Delegate.h"
+#include "Misc/Attribute.h"
+#include "Widgets/Layout/SSplitter.h"
+
+class FDetailColumnSizeData
+{
+public:
+
+	FDetailColumnSizeData()
+	{
+		ValueColumnWidthValue = 0.7f;
+		RightColumnWidthValue = 0;
+		HoveredSplitterIndexValue = INDEX_NONE;
+
+		NameColumnWidth.BindRaw(this, &FDetailColumnSizeData::GetNameColumnWidthValue);
+		ValueColumnWidth.BindRaw(this, &FDetailColumnSizeData::GetValueColumnWidthValue);
+		WholeRowColumnWidth.BindRaw(this, &FDetailColumnSizeData::GetWholeRowColumnWidthValue);
+		RightColumnWidth.BindRaw(this, &FDetailColumnSizeData::GetRightColumnWidthValue);
+		HoveredSplitterIndex.BindRaw(this, &FDetailColumnSizeData::GetHoveredSplitterIndexValue);
+		OnValueColumnResized.BindRaw(this, &FDetailColumnSizeData::SetValueColumnWidth);
+		OnRightColumnResized.BindRaw(this, &FDetailColumnSizeData::OnSetRightColumnWidth);
+		OnSplitterHandleHovered.BindRaw(this, &FDetailColumnSizeData::OnSetHoveredSplitterIndex);
+
+		// these are intentionally left no-op, since the widths are derived from the other column width
+		OnWholeRowColumnResized.BindLambda([](float) {});
+		OnNameColumnResized.BindLambda([](float) {});
+	}
+
+	TAttribute<float> GetNameColumnWidth() const		{ return NameColumnWidth; }
+	TAttribute<float> GetValueColumnWidth() const		{ return ValueColumnWidth; }
+	TAttribute<float> GetWholeRowColumnWidth() const	{ return WholeRowColumnWidth; }
+	TAttribute<float> GetRightColumnWidth() const		{ return RightColumnWidth; }
+	TAttribute<float> GetRightColumnMinWidth() const	{ return RightColumnMinWidth; }
+	TAttribute<int> GetHoveredSplitterIndex() const		{ return HoveredSplitterIndex; }
+
+	const SSplitter::FOnSlotResized& GetOnNameColumnResized() const		{ return OnNameColumnResized; }
+	const SSplitter::FOnSlotResized& GetOnValueColumnResized() const	{ return OnValueColumnResized; }
+	const SSplitter::FOnSlotResized& GetOnWholeRowColumnResized() const { return OnWholeRowColumnResized; }
+	const SSplitter::FOnSlotResized& GetOnRightColumnResized() const	{ return OnRightColumnResized; }
+	const SSplitter::FOnHandleHovered& GetOnSplitterHandleHovered()	const { return OnSplitterHandleHovered; }
+
+	void SetValueColumnWidth(float NewWidth)
+	{
+		ensure(NewWidth <= 1.0f);
+		ValueColumnWidthValue = NewWidth;
+	}
+
+	void SetRightColumnMinWidth(float NewMinWidth)
+	{
+		RightColumnMinWidth = NewMinWidth;
+	}
+
+private:
+	float ValueColumnWidthValue;
+	float RightColumnWidthValue;
+	int HoveredSplitterIndexValue;
+
+	TAttribute<float> NameColumnWidth;
+	TAttribute<float> ValueColumnWidth;
+	TAttribute<float> WholeRowColumnWidth;
+	TAttribute<float> RightColumnWidth;
+	TAttribute<float> RightColumnMinWidth;
+	TAttribute<int32> HoveredSplitterIndex;
+	SSplitter::FOnSlotResized OnNameColumnResized;
+	SSplitter::FOnSlotResized OnValueColumnResized;
+	SSplitter::FOnSlotResized OnWholeRowColumnResized;
+	SSplitter::FOnSlotResized OnRightColumnResized;
+	SSplitter::FOnHandleHovered OnSplitterHandleHovered;
+
+	float GetNameColumnWidthValue() const { return 1.0f - (ValueColumnWidthValue + RightColumnWidthValue); }
+	float GetValueColumnWidthValue() const { return ValueColumnWidthValue; }
+	float GetWholeRowColumnWidthValue() const { return 1.0f - RightColumnWidthValue; }
+	int32 GetHoveredSplitterIndexValue() const { return HoveredSplitterIndexValue; }
+	float GetRightColumnWidthValue() const { return RightColumnWidthValue; }
+
+	void OnSetRightColumnWidth(float NewWidth)
+	{
+		ensure(NewWidth <= 1.0f);
+		RightColumnWidthValue = NewWidth;
+	}
+
+	void OnSetHoveredSplitterIndex(int32 HoveredIndex)
+	{
+		HoveredSplitterIndexValue = HoveredIndex;
+	}
+};
diff --git a/Engine/Source/Editor/PropertyEditor/Public/IDetailPropertyRow.h b/Engine/Source/Editor/PropertyEditor/Public/IDetailPropertyRow.h
--- a/Engine/Source/Editor/PropertyEditor/Public/IDetailPropertyRow.h
+++ b/Engine/Source/Editor/PropertyEditor/Public/IDetailPropertyRow.h
@@ -3,91 +3,15 @@
 #pragma once
 
 #include "CoreMinimal.h"
-#include "Misc/Attribute.h"
-
-#include "Layout/Visibility.h"
 #include "Framework/SlateDelegates.h"
+#include "Layout/Visibility.h"
+#include "Misc/Attribute.h"
 #include "PropertyHandle.h"
-#include "Widgets/Layout/SSplitter.h"
-
 
 class FDetailWidgetRow;
 class FDetailWidgetDecl;
 class IDetailDragDropHandler;
 
-class FDetailColumnSizeData
-{
-public:
-	
-	FDetailColumnSizeData()
-	{
-		ValueColumnWidthValue = 0.7f;
-		RightColumnWidthValue = 0;
-		HoveredSplitterIndexValue = INDEX_NONE;
-
-		NameColumnWidth.BindRaw(this, &FDetailColumnSizeData::GetNameColumnWidth);
-		ValueColumnWidth.BindRaw(this, &FDetailColumnSizeData::GetValueColumnWidth);
-		PropertyColumnWidth.BindRaw(this, &FDetailColumnSizeData::GetPropertyColumnWidth);
-		RightColumnWidth.BindRaw(this, &FDetailColumnSizeData::GetRightColumnWidth);
-		RightColumnMinWidth.BindRaw(this, &FDetailColumnSizeData::GetRightColumnMinWidth);
-		HoveredSplitterIndex.BindRaw(this, &FDetailColumnSizeData::GetHoveredSplitterIndex);
-		OnValueColumnResized.BindRaw(this, &FDetailColumnSizeData::SetValueColumnWidth);
-		OnRightColumnResized.BindRaw(this, &FDetailColumnSizeData::OnSetRightColumnWidth);
-		OnSplitterHandleHovered.BindRaw(this, &FDetailColumnSizeData::OnSetHoveredSplitterIndex);
-
-		// these are intentionally left no-op, since the widths are derived from the other column width
-		OnPropertyColumnResized.BindLambda([](float) {});
-		OnNameColumnResized.BindLambda([](float) {}); 
-	}
-
-	// don't change these attributes directly - they're here only to be forwarded to the details view
-	TAttribute<float> NameColumnWidth;
-	TAttribute<float> ValueColumnWidth;
-	TAttribute<float> PropertyColumnWidth;
-	TAttribute<float> RightColumnWidth;
-	TAttribute<float> RightColumnMinWidth;
-	TAttribute<int32> HoveredSplitterIndex;
-	SSplitter::FOnSlotResized OnNameColumnResized; 
-	SSplitter::FOnSlotResized OnValueColumnResized;
-	SSplitter::FOnSlotResized OnPropertyColumnResized; 
-	SSplitter::FOnSlotResized OnRightColumnResized;
-	SSplitter::FOnHandleHovered OnSplitterHandleHovered;
-
-	void SetValueColumnWidth(float NewWidth)
-	{ 
-		ensureAlways(NewWidth <= 1.0f);
-		ValueColumnWidthValue = NewWidth;
-	}
-
-	float GetRightColumnMinWidth() const { return RightColumnMinWidthValue; }
-	void SetRightColumnMinWidth(float NewWidth)
-	{ 
-		RightColumnMinWidthValue = NewWidth;
-	}
-
-private:
-	float ValueColumnWidthValue;
-	float RightColumnWidthValue;
-	float RightColumnMinWidthValue;
-	int HoveredSplitterIndexValue;
-
-	float GetNameColumnWidth() const { return 1.0f - (ValueColumnWidthValue + RightColumnWidthValue); }
-	float GetValueColumnWidth() const { return ValueColumnWidthValue; }
-	float GetRightColumnWidth() const { return RightColumnWidthValue; }
-	float GetPropertyColumnWidth() const { return 1.0f - RightColumnWidthValue; }
-	int32 GetHoveredSplitterIndex() const { return HoveredSplitterIndexValue; }
-	
-	void OnSetRightColumnWidth(float NewWidth)
-	{
-		RightColumnWidthValue = NewWidth;
-	}
-
-	void OnSetHoveredSplitterIndex(int32 HoveredIndex)
-	{
-		HoveredSplitterIndexValue = HoveredIndex;
-	}
-};
-
 DECLARE_DELEGATE_RetVal_OneParam(bool, FIsResetToDefaultVisible, TSharedPtr<IPropertyHandle> /* PropertyHandle */);
 DECLARE_DELEGATE_OneParam(FResetToDefaultHandler, TSharedPtr<IPropertyHandle> /* PropertyHandle*/);
 
