Added missing copyright boilerplate.

diff --git a/Engine/Source/Editor/PropertyEditor/Public/DetailColumnSizeData.h b/Engine/Source/Editor/PropertyEditor/Public/DetailColumnSizeData.h
--- a/Engine/Source/Editor/PropertyEditor/Public/DetailColumnSizeData.h
+++ b/Engine/Source/Editor/PropertyEditor/Public/DetailColumnSizeData.h
@@ -1,3 +1,5 @@
+// Copyright Epic Games, Inc. All Rights Reserved.
+
 #pragma once
 
 #include "Delegates/Delegate.h"
