Fixed crash when neither a result nor an error is set in the expression parser result when evaluating edit conditions.

diff --git a/Engine/Source/Editor/PropertyEditor/Private/EditConditionParser.cpp b/Engine/Source/Editor/PropertyEditor/Private/EditConditionParser.cpp
--- a/Engine/Source/Editor/PropertyEditor/Private/EditConditionParser.cpp
+++ b/Engine/Source/Editor/PropertyEditor/Private/EditConditionParser.cpp
@@ -697,7 +697,7 @@ TValueOrError<bool, FText> FEditConditionParser::Evaluate(const FEditConditionEx
 	using namespace EditConditionParserTokens;
 	
 	FExpressionResult Result = ExpressionParser::Evaluate(Expression.Tokens, OperatorJumpTable, &Context);
-	if (Result.IsValid())
+	if (Result.HasValue())
 	{
 		const bool* BoolResult = Result.GetValue().Cast<bool>();
 		if (BoolResult != nullptr)
@@ -716,7 +716,8 @@ TValueOrError<bool, FText> FEditConditionParser::Evaluate(const FEditConditionEx
 		}
 	}
 
-	return MakeError(Result.GetError().Text);
+	const FText ErrorText = Result.HasError() ? Result.GetError().Text : FText::GetEmpty();
+	return MakeError(ErrorText);
 }
 
 TSharedPtr<FEditConditionExpression> FEditConditionParser::Parse(const FString& ExpressionString) const
