Fix crash caused by dereferencing nullptr when displaying a detail property row with an external root node using a property row generator

diff --git a/Engine/Source/Editor/PropertyEditor/Private/DetailPropertyRow.cpp b/Engine/Source/Editor/PropertyEditor/Private/DetailPropertyRow.cpp
--- a/Engine/Source/Editor/PropertyEditor/Private/DetailPropertyRow.cpp
+++ b/Engine/Source/Editor/PropertyEditor/Private/DetailPropertyRow.cpp
@@ -43,8 +43,11 @@ FDetailPropertyRow::FDetailPropertyRow(TSharedPtr<FPropertyNode> InPropertyNode,
 		if (PropertyNode->AsObjectNode() && ExternalRootNode.IsValid())
 		{
 			// We are showing an entirely different object inline.  Generate a layout for it now.
-			ExternalObjectLayout = MakeShared<FDetailLayoutData>();
-			InParentCategory->GetDetailsView()->UpdateSinglePropertyMap(InExternalRootNode, *ExternalObjectLayout, true);
+			if (IDetailsViewPrivate* DetailsView = InParentCategory->GetDetailsView())
+			{
+				ExternalObjectLayout = MakeShared<FDetailLayoutData>();
+				DetailsView->UpdateSinglePropertyMap(InExternalRootNode, *ExternalObjectLayout, true);
+			}
 		}
 
 		if (PropertyNode->GetPropertyKeyNode().IsValid())
