Add shift click copy / paste to details panel

diff --git a/Engine/Source/Editor/PropertyEditor/Private/IDetailsViewPrivate.h b/Engine/Source/Editor/PropertyEditor/Private/IDetailsViewPrivate.h
--- a/Engine/Source/Editor/PropertyEditor/Private/IDetailsViewPrivate.h
+++ b/Engine/Source/Editor/PropertyEditor/Private/IDetailsViewPrivate.h
@@ -144,6 +144,12 @@ public:
 	*/
 	virtual void RestoreExpandedItems(TSharedRef<FPropertyNode> StartNode) = 0;
 
+	/** Mark node as animating, useful if animating during behaviors that trigger widget reconstruction */
+	virtual void MarkNodeAnimating(TSharedPtr<FPropertyNode> InNode, float InAnimationDuration) = 0;
+
+	/** Returns true if node is animating, currently we only keep track of the last animated node */
+	virtual bool IsNodeAnimating(TSharedPtr<FPropertyNode> InNode) = 0;
+
 	/** Column width accessibility */
 	virtual class FDetailColumnSizeData& GetColumnSizeData() = 0;
 
diff --git a/Engine/Source/Editor/PropertyEditor/Private/SDetailSingleItemRow.cpp b/Engine/Source/Editor/PropertyEditor/Private/SDetailSingleItemRow.cpp
--- a/Engine/Source/Editor/PropertyEditor/Private/SDetailSingleItemRow.cpp
+++ b/Engine/Source/Editor/PropertyEditor/Private/SDetailSingleItemRow.cpp
@@ -32,11 +32,12 @@ namespace DetailWidgetConstants
 {
 	const FMargin LeftRowPadding( 20.0f, 0.0f, 10.0f, 0.0f );
 	const FMargin RightRowPadding( 12.0f, 0.0f, 2.0f, 0.0f );
+	const float PulseAnimationLength = 0.5f;
 }
 
 namespace SDetailSingleItemRow_Helper
 {
-	//Get the node item number in case it is expand we have to recursively count all expanded children
+	// Get the node item number in case it is expand we have to recursively count all expanded children
 	void RecursivelyGetItemShow(TSharedRef<FDetailTreeNode> ParentItem, int32& ItemShowNum)
 	{
 		if (ParentItem->GetVisibility() == ENodeVisibility::Visible)
@@ -325,7 +326,16 @@ void SDetailSingleItemRow::Construct( const FArguments& InArgs, FDetailLayoutCus
 	FOnAcceptDrop AcceptDropDelegate;
 	FOnCanAcceptDrop CanAcceptDropDelegate;
 
-	FDetailColumnSizeData& ColumnSizeData = InOwnerTreeNode->GetDetailsView()->GetColumnSizeData();
+	IDetailsViewPrivate* DetailsView = InOwnerTreeNode->GetDetailsView();
+	FDetailColumnSizeData& ColumnSizeData = DetailsView->GetColumnSizeData();
+
+	PulseAnimation.AddCurve(0.0f, DetailWidgetConstants::PulseAnimationLength, ECurveEaseFunction::CubicInOut);
+
+	// Play on construction if animation was started from a behavior the re-constructs this widget
+	if (DetailsView->IsNodeAnimating(GetPropertyNode()))
+	{
+		PulseAnimation.Play(SharedThis(this));
+	}
 
 	const bool bIsValidTreeNode = InOwnerTreeNode->GetParentCategory().IsValid() && InOwnerTreeNode->GetParentCategory()->IsParentLayoutValid();
 	if (bIsValidTreeNode)
@@ -334,6 +344,33 @@ void SDetailSingleItemRow::Construct( const FArguments& InArgs, FDetailLayoutCus
 		{
 			WidgetRow = Customization->GetWidgetRow();
 
+			// Setup copy / paste actions
+			{
+				if (WidgetRow.IsCopyPasteBound())
+				{
+					CopyAction = WidgetRow.CopyMenuAction;
+					PasteAction = WidgetRow.PasteMenuAction;
+				}
+				else
+				{
+					TSharedPtr<FPropertyNode> PropertyNode = GetPropertyNode();
+					static const FName DisableCopyPasteMetaDataName("DisableCopyPaste");
+					if (PropertyNode.IsValid() && !PropertyNode->ParentOrSelfHasMetaData(DisableCopyPasteMetaDataName))
+					{
+						CopyAction.ExecuteAction = FExecuteAction::CreateSP(this, &SDetailSingleItemRow::OnCopyProperty);
+						PasteAction.ExecuteAction = FExecuteAction::CreateSP(this, &SDetailSingleItemRow::OnPasteProperty);
+						PasteAction.CanExecuteAction = FCanExecuteAction::CreateSP(this, &SDetailSingleItemRow::CanPasteProperty);
+					}
+					else
+					{
+						CopyAction.ExecuteAction = FExecuteAction::CreateLambda([]() {});
+						CopyAction.CanExecuteAction = FCanExecuteAction::CreateLambda([]() { return false; });
+						PasteAction.ExecuteAction = FExecuteAction::CreateLambda([]() {});
+						PasteAction.CanExecuteAction = FCanExecuteAction::CreateLambda([]() { return false; });
+					}
+				}
+			}
+
 			// Populate the extension content in the WidgetRow if there's an extension handler.
 			PopulateExtensionWidget();
 
@@ -670,6 +707,31 @@ void SDetailSingleItemRow::Construct( const FArguments& InArgs, FDetailLayoutCus
 	);
 }
 
+FReply SDetailSingleItemRow::OnMouseButtonUp(const FGeometry& MyGeometry, const FPointerEvent& MouseEvent)
+{
+	if (MouseEvent.GetModifierKeys().IsShiftDown())
+	{
+		bool bIsHandled = false;
+		if (CopyAction.CanExecute() && MouseEvent.GetEffectingButton() == EKeys::RightMouseButton)
+		{
+			CopyAction.Execute();
+			bIsHandled = true;
+		}
+		else if (PasteAction.CanExecute() && MouseEvent.GetEffectingButton() == EKeys::LeftMouseButton)
+		{
+			PasteAction.Execute();
+			bIsHandled = true;
+		}
+
+		if (bIsHandled)
+		{
+			return FReply::Handled();
+		}
+	}
+
+	return SDetailTableRowBase::OnMouseButtonUp(MyGeometry, MouseEvent);
+}
+
 bool SDetailSingleItemRow::IsResetToDefaultEnabled() const
 {
 	return bCachedResetToDefaultEnabled;
@@ -702,47 +764,31 @@ FSlateColor SDetailSingleItemRow::GetOuterBackgroundColor() const
 /** Get the background color of the inner part of the row, which contains the name and value widgets. */
 FSlateColor SDetailSingleItemRow::GetInnerBackgroundColor() const
 {
+	FSlateColor Color;
+
 	if (IsHighlighted())
 	{
-		return FAppStyle::Get().GetSlateColor("Colors.Hover");
+		Color = FAppStyle::Get().GetSlateColor("Colors.Hover");
+	}
+	else
+	{
+		const int32 IndentLevel = GetIndentLevelForBackgroundColor();
+		Color = PropertyEditorConstants::GetRowBackgroundColor(IndentLevel, this->IsHovered());
+	}
+
+	if (PulseAnimation.IsPlaying())
+	{
+		float Lerp = PulseAnimation.GetLerp();
+		return FMath::Lerp(FAppStyle::Get().GetSlateColor("Colors.Hover2").GetSpecifiedColor(), Color.GetSpecifiedColor(), Lerp);
 	}
 
-	const int32 IndentLevel = GetIndentLevelForBackgroundColor();
-	return PropertyEditorConstants::GetRowBackgroundColor(IndentLevel, this->IsHovered());
+	return Color;
 }
 
 bool SDetailSingleItemRow::OnContextMenuOpening(FMenuBuilder& MenuBuilder)
 {
-	const bool bIsCopyPasteBound = WidgetRow.IsCopyPasteBound();
-
-	FUIAction CopyAction;
-	FUIAction PasteAction;
 	FUIAction CopyDisplayNameAction = FExecuteAction::CreateSP(this, &SDetailSingleItemRow::OnCopyPropertyDisplayName);
 
-	if (bIsCopyPasteBound)
-	{
-		CopyAction = WidgetRow.CopyMenuAction;
-		PasteAction = WidgetRow.PasteMenuAction;
-	}
-	else
-	{
-		TSharedPtr<FPropertyNode> PropertyNode = GetPropertyNode();
-		static const FName DisableCopyPasteMetaDataName("DisableCopyPaste");
-		if (PropertyNode.IsValid() && !PropertyNode->ParentOrSelfHasMetaData(DisableCopyPasteMetaDataName))
-		{
-			CopyAction.ExecuteAction = FExecuteAction::CreateSP(this, &SDetailSingleItemRow::OnCopyProperty);
-			PasteAction.ExecuteAction = FExecuteAction::CreateSP(this, &SDetailSingleItemRow::OnPasteProperty);
-			PasteAction.CanExecuteAction = FCanExecuteAction::CreateSP(this, &SDetailSingleItemRow::CanPasteProperty);
-		}
-		else
-		{
-			CopyAction.ExecuteAction = FExecuteAction::CreateLambda([](){});
-			CopyAction.CanExecuteAction = FCanExecuteAction::CreateLambda([]() { return false; });
-			PasteAction.ExecuteAction = FExecuteAction::CreateLambda([](){});
-			PasteAction.CanExecuteAction = FCanExecuteAction::CreateLambda([]() { return false; });
-		}
-	}
-
 	bool bAddedMenuEntry = false;
 	if (CopyAction.IsBound() && PasteAction.IsBound())
 	{
@@ -752,17 +798,23 @@ bool SDetailSingleItemRow::OnContextMenuOpening(FMenuBuilder& MenuBuilder)
 			MenuBuilder.AddMenuSeparator();
 		}
 
-		MenuBuilder.AddMenuEntry(
-			NSLOCTEXT("PropertyView", "CopyProperty", "Copy"),
-			NSLOCTEXT("PropertyView", "CopyProperty_ToolTip", "Copy this property value"),
-			FSlateIcon(FCoreStyle::Get().GetStyleSetName(), "GenericCommands.Copy"),
-			CopyAction);
+		bool bLongDisplayName = false;
 
-		MenuBuilder.AddMenuEntry(
-			NSLOCTEXT("PropertyView", "PasteProperty", "Paste"),
-			NSLOCTEXT("PropertyView", "PasteProperty_ToolTip", "Paste the copied value here"),
-			FSlateIcon(FCoreStyle::Get().GetStyleSetName(), "GenericCommands.Paste"),
-			PasteAction);
+		FMenuEntryParams CopyContentParams;
+		CopyContentParams.LabelOverride = NSLOCTEXT("PropertyView", "CopyProperty", "Copy");
+		CopyContentParams.ToolTipOverride = NSLOCTEXT("PropertyView", "CopyProperty_ToolTip", "Copy this property value");
+		CopyContentParams.InputBindingOverride = FInputChord(EModifierKey::Shift, EKeys::RightMouseButton).GetInputText(bLongDisplayName);
+		CopyContentParams.IconOverride = FSlateIcon(FCoreStyle::Get().GetStyleSetName(), "GenericCommands.Copy");
+		CopyContentParams.DirectActions = CopyAction;
+		MenuBuilder.AddMenuEntry(CopyContentParams);
+
+		FMenuEntryParams PasteContentParams;
+		PasteContentParams.LabelOverride = NSLOCTEXT("PropertyView", "PasteProperty", "Paste");
+		PasteContentParams.ToolTipOverride = NSLOCTEXT("PropertyView", "PasteProperty_ToolTip", "Paste the copied value here");
+		PasteContentParams.InputBindingOverride = FInputChord(EModifierKey::Shift, EKeys::LeftMouseButton).GetInputText(bLongDisplayName);
+		PasteContentParams.IconOverride = FSlateIcon(FCoreStyle::Get().GetStyleSetName(), "GenericCommands.Paste");
+		PasteContentParams.DirectActions = PasteAction;
+		MenuBuilder.AddMenuEntry(PasteContentParams);
 
 		MenuBuilder.AddMenuEntry(
 			NSLOCTEXT("PropertyView", "CopyPropertyDisplayName", "Copy Display Name"),
@@ -866,6 +918,7 @@ void SDetailSingleItemRow::OnCopyProperty()
 			if (Handle->GetValueAsFormattedString(Value, PPF_Copy) == FPropertyAccess::Success)
 			{
 				FPlatformApplicationMisc::ClipboardCopy(*Value);
+				PulseAnimation.Play(SharedThis(this));
 			}
 		}
 	}
@@ -906,6 +959,9 @@ void SDetailSingleItemRow::OnPasteProperty()
 
 			// Need to refresh the details panel in case a property was pasted over another.
 			OwnerTreeNode.Pin()->GetDetailsView()->ForceRefresh();
+
+			// Mark property node as animating so we will animate after re-construction
+			DetailsView->MarkNodeAnimating(PropertyNode, DetailWidgetConstants::PulseAnimationLength);
 		}
 	}
 }
diff --git a/Engine/Source/Editor/PropertyEditor/Private/SDetailSingleItemRow.h b/Engine/Source/Editor/PropertyEditor/Private/SDetailSingleItemRow.h
--- a/Engine/Source/Editor/PropertyEditor/Private/SDetailSingleItemRow.h
+++ b/Engine/Source/Editor/PropertyEditor/Private/SDetailSingleItemRow.h
@@ -60,6 +60,10 @@ public:
 	 */
 	void Construct( const FArguments& InArgs, FDetailLayoutCustomization* InCustomization, bool bHasMultipleColumns, TSharedRef<FDetailTreeNode> InOwnerTreeNode, const TSharedRef<STableViewBase>& InOwnerTableView );
 
+	// ~Begin SWidget Interface
+	virtual FReply OnMouseButtonUp(const FGeometry& MyGeometry, const FPointerEvent& MouseEvent) override;
+	// ~Begin End Interface
+
 	TSharedPtr<FDragDropOperation> CreateDragDropOperation();
 
 	virtual void Tick(const FGeometry& AllottedGeometry, const double InCurrentTime, const float InDeltaTime);
@@ -117,6 +121,11 @@ private:
 	TSharedPtr<FPropertyNode> SwappablePropertyNode;
 	TSharedPtr<SButton> ExpanderArrow;
 	TWeakPtr<FDragDropOperation> DragOperation; // last drag initiated by this widget
+	FUIAction CopyAction;
+	FUIAction PasteAction;
+
+	/** Animation curve for displaying pulse */
+	FCurveSequence PulseAnimation;
 };
 
 class FArrayRowDragDropOp : public FDecoratedDragDropOp
diff --git a/Engine/Source/Editor/PropertyEditor/Private/SDetailsViewBase.cpp b/Engine/Source/Editor/PropertyEditor/Private/SDetailsViewBase.cpp
--- a/Engine/Source/Editor/PropertyEditor/Private/SDetailsViewBase.cpp
+++ b/Engine/Source/Editor/PropertyEditor/Private/SDetailsViewBase.cpp
@@ -1315,6 +1315,32 @@ void SDetailsViewBase::RestoreExpandedItems(TSharedRef<FPropertyNode> StartNode)
 	}
 }
 
+void SDetailsViewBase::MarkNodeAnimating(TSharedPtr<FPropertyNode> InNode, float InAnimationDuration)
+{
+	if (InNode.IsValid() && InAnimationDuration > 0.0f)
+	{
+		CurrentlyAnimatingNodePath = FPropertyNode::CreatePropertyPath(InNode.ToSharedRef());
+		GEditor->GetTimerManager()->ClearTimer(AnimateNodeTimer);
+		GEditor->GetTimerManager()->SetTimer(AnimateNodeTimer, FTimerDelegate::CreateSP(this, &SDetailsViewBase::HandleNodeAnimationComplete), InAnimationDuration, false);
+	}
+}
+
+bool SDetailsViewBase::IsNodeAnimating(TSharedPtr<FPropertyNode> InNode)
+{
+	if (CurrentlyAnimatingNodePath.IsValid() && InNode.IsValid())
+	{
+		TSharedRef<FPropertyPath> InNodePath = FPropertyNode::CreatePropertyPath(InNode.ToSharedRef());
+		return FPropertyPath::AreEqual(CurrentlyAnimatingNodePath.ToSharedRef(), InNodePath);
+	}
+
+	return false;
+}
+
+void SDetailsViewBase::HandleNodeAnimationComplete()
+{
+	CurrentlyAnimatingNodePath = nullptr;
+}
+
 void SDetailsViewBase::UpdateFilteredDetails()
 {
 	RootTreeNodes.Reset();
diff --git a/Engine/Source/Editor/PropertyEditor/Private/SDetailsViewBase.h b/Engine/Source/Editor/PropertyEditor/Private/SDetailsViewBase.h
--- a/Engine/Source/Editor/PropertyEditor/Private/SDetailsViewBase.h
+++ b/Engine/Source/Editor/PropertyEditor/Private/SDetailsViewBase.h
@@ -131,6 +131,8 @@ public:
 	virtual const FCustomPropertyTypeLayoutMap& GetCustomPropertyTypeLayoutMap() const { return InstancedTypeToLayoutMap; }
 	virtual void SaveExpandedItems( TSharedRef<FPropertyNode> StartNode ) override;
 	virtual void RestoreExpandedItems(TSharedRef<FPropertyNode> StartNode) override;
+	virtual void MarkNodeAnimating(TSharedPtr<FPropertyNode> InNode, float InAnimationDuration) override;
+	virtual bool IsNodeAnimating(TSharedPtr<FPropertyNode> InNode) override;
 	virtual FDetailColumnSizeData& GetColumnSizeData() override { return ColumnSizeData; }
 	virtual bool IsFavoritingEnabled() const override { return DetailsViewArgs.bAllowFavoriteSystem; }
 	virtual bool IsConnected() const = 0;
@@ -177,6 +179,9 @@ protected:
 	 */
 	void SetColorPropertyFromColorPicker(FLinearColor NewColor);
 
+	/** Called when node animation completes */
+	void HandleNodeAnimationComplete();
+
 	/** Updates the property map for access when customizing the details view.  Generates default layout for properties */
 	void UpdatePropertyMaps();
 
@@ -403,8 +408,12 @@ protected:
 	TSharedPtr<IDetailKeyframeHandler> KeyframeHandler;
 	/** Property extension handler returns additional UI to apply after the customization is applied to the property. */
 	TSharedPtr<IDetailPropertyExtensionHandler> ExtensionHandler;
-	/** The tree node that is currently highlighted, may be none: */
+	/** The tree node that is currently highlighted, may be none. */
 	TWeakPtr<FDetailTreeNode> CurrentlyHighlightedNode;
+	/** The property node whose widget-equivalent should be animating, may be none. */
+	TSharedPtr<FPropertyPath> CurrentlyAnimatingNodePath;
+	/** Timer for that current node's widget animation duration. */
+	FTimerHandle AnimateNodeTimer;
 
 	TSet<FString> PreSearchExpandedItems;
 	TSet<FString> PreSearchExpandedCategories;
