Add FPropertyHandleColor to support copy / paste between different color types.

diff --git a/Engine/Source/Editor/PropertyEditor/Private/PropertyEditorHelpers.cpp b/Engine/Source/Editor/PropertyEditor/Private/PropertyEditorHelpers.cpp
--- a/Engine/Source/Editor/PropertyEditor/Private/PropertyEditorHelpers.cpp
+++ b/Engine/Source/Editor/PropertyEditor/Private/PropertyEditorHelpers.cpp
@@ -604,6 +604,10 @@ namespace PropertyEditorHelpers
 		{
 			PropertyHandle = MakeShareable( new FPropertyHandleRotator( PropertyNode, NotifyHook, PropertyUtilities ) );
 		}
+		else if (FPropertyHandleColor::Supports(PropertyNode))
+		{
+			PropertyHandle = MakeShareable(new FPropertyHandleColor(PropertyNode, NotifyHook, PropertyUtilities));
+		}
 		else if (FPropertyHandleSet::Supports(PropertyNode))
 		{
 			PropertyHandle = MakeShareable( new FPropertyHandleSet( PropertyNode, NotifyHook, PropertyUtilities ) );
diff --git a/Engine/Source/Editor/PropertyEditor/Private/PropertyHandleImpl.cpp b/Engine/Source/Editor/PropertyEditor/Private/PropertyHandleImpl.cpp
--- a/Engine/Source/Editor/PropertyEditor/Private/PropertyHandleImpl.cpp
+++ b/Engine/Source/Editor/PropertyEditor/Private/PropertyHandleImpl.cpp
@@ -4496,6 +4496,60 @@ FPropertyAccess::Result FPropertyHandleRotator::SetYaw( float InYaw, EPropertyVa
 	return Res;
 }
 
+FPropertyHandleColor::FPropertyHandleColor(TSharedRef<FPropertyNode> PropertyNode, FNotifyHook* NotifyHook, TSharedPtr<IPropertyUtilities> PropertyUtilities)\
+	: FPropertyHandleBase(PropertyNode, NotifyHook, PropertyUtilities)
+{
+}
+
+bool FPropertyHandleColor::Supports(TSharedRef<FPropertyNode> PropertyNode)
+{
+	FProperty* Property = PropertyNode->GetProperty();
+
+	if ( FStructProperty* StructProperty = CastField<FStructProperty>(Property) )
+	{
+		return StructProperty->Struct == FSlateColor::StaticStruct() 
+			|| StructProperty->Struct == TBaseStructure<FLinearColor>::Get();
+	}
+
+	return false;
+}
+
+FPropertyAccess::Result FPropertyHandleColor::SetValueFromFormattedString(const FString& InValue, EPropertyValueSetFlags::Type Flags)
+{
+	FProperty* Property = Implementation->GetPropertyNode()->GetProperty();
+
+	if (FStructProperty* StructProperty = CastField<FStructProperty>(Property))
+	{
+		FString OutValue;
+		FLinearColor LinearColor;
+
+		if (LinearColor.InitFromString(InValue))
+		{
+			if (StructProperty->Struct == FSlateColor::StaticStruct())
+			{
+				if (InValue.Contains("ColorUseRule"))
+				{
+					return Implementation->SetValueAsString(InValue, Flags);
+				}
+				else
+				{
+					LinearColor.InitFromString(InValue);
+					FSlateColor SlateColor(LinearColor);
+					Implementation->GetPropertyNode()->GetProperty()->ExportText_Direct(OutValue, &SlateColor, &SlateColor, nullptr, 0);
+					return Implementation->SetValueAsString(OutValue, Flags);
+				}
+			}
+			else if (StructProperty->Struct == TBaseStructure<FLinearColor>::Get())
+			{
+				Implementation->GetPropertyNode()->GetProperty()->ExportText_Direct(OutValue, &LinearColor, &LinearColor, nullptr, 0);
+				return Implementation->SetValueAsString(OutValue, Flags);
+			}
+		}
+	}
+
+	return FPropertyAccess::Result::Fail;
+}
+
 
 bool FPropertyHandleArray::Supports( TSharedRef<FPropertyNode> PropertyNode )
 {
diff --git a/Engine/Source/Editor/PropertyEditor/Private/PropertyHandleImpl.h b/Engine/Source/Editor/PropertyEditor/Private/PropertyHandleImpl.h
--- a/Engine/Source/Editor/PropertyEditor/Private/PropertyHandleImpl.h
+++ b/Engine/Source/Editor/PropertyEditor/Private/PropertyHandleImpl.h
@@ -654,6 +654,14 @@ private:
 	TSharedPtr<FPropertyHandleFloat> YawValue;
 };
 
+class FPropertyHandleColor : public FPropertyHandleBase
+{
+public:
+	FPropertyHandleColor( TSharedRef<FPropertyNode> PropertyNode, FNotifyHook* NotifyHook, TSharedPtr<IPropertyUtilities> PropertyUtilities );
+	static bool Supports( TSharedRef<FPropertyNode> PropertyNode );
+
+	virtual FPropertyAccess::Result SetValueFromFormattedString(const FString& InValue, EPropertyValueSetFlags::Type Flags = EPropertyValueSetFlags::DefaultFlags) override;
+};
 
 class FPropertyHandleArray : public FPropertyHandleBase, public IPropertyHandleArray
 {
