Revert  part of 18198783. Do not set all property editor bools to be on an interactive change

diff --git a/Engine/Source/Editor/PropertyEditor/Private/UserInterface/PropertyEditor/SPropertyEditorBool.cpp b/Engine/Source/Editor/PropertyEditor/Private/UserInterface/PropertyEditor/SPropertyEditorBool.cpp
--- a/Engine/Source/Editor/PropertyEditor/Private/UserInterface/PropertyEditor/SPropertyEditorBool.cpp
+++ b/Engine/Source/Editor/PropertyEditor/Private/UserInterface/PropertyEditor/SPropertyEditorBool.cpp
@@ -73,11 +73,11 @@ void SPropertyEditorBool::OnCheckStateChanged( ECheckBoxState InNewState )
 	const TSharedRef< IPropertyHandle > PropertyHandle = PropertyEditor->GetPropertyHandle();
 	if( InNewState == ECheckBoxState::Checked || InNewState == ECheckBoxState::Undetermined )
 	{
-		PropertyHandle->SetValue( true, EPropertyValueSetFlags::InteractiveChange );
+		PropertyHandle->SetValue( true );
 	}
 	else
 	{
-		PropertyHandle->SetValue( false, EPropertyValueSetFlags::InteractiveChange );
+		PropertyHandle->SetValue( false );
 	}
 }
 
