Change the Property Matrix's pin icons to not blink, and also update them to the new icons

diff --git a/Engine/Source/Editor/EditorStyle/Private/StarshipStyle.cpp b/Engine/Source/Editor/EditorStyle/Private/StarshipStyle.cpp
--- a/Engine/Source/Editor/EditorStyle/Private/StarshipStyle.cpp
+++ b/Engine/Source/Editor/EditorStyle/Private/StarshipStyle.cpp
@@ -460,6 +460,8 @@ void FStarshipEditorStyle::FStyle::SetupGeneralStyles()
 		Set("Icons.OpenSourceLocation", new IMAGE_BRUSH_SVG("Starship/Common/OpenSourceLocation", Icon16x16));
 		Set("Icons.Find", new IMAGE_BRUSH_SVG("Starship/Common/Find", Icon16x16));
 		Set("Icons.Validate", new IMAGE_BRUSH_SVG("Starship/Common/validate", Icon16x16));
+		Set("Icons.Pinned", new IMAGE_BRUSH_SVG("Starship/Common/Pinned", Icon16x16));
+		Set("Icons.Unpinned", new IMAGE_BRUSH_SVG("Starship/Common/Unpinned", Icon16x16));
 
 
 		Set("Icons.Toolbar.Play", new IMAGE_BRUSH_SVG("Starship/Common/play", Icon20x20));
@@ -1220,9 +1222,6 @@ void FStarshipEditorStyle::FStyle::SetupGeneralStyles()
 		Set( "SceneOutliner.World", 		new CORE_IMAGE_BRUSH_SVG("Starship/Common/world", Icon16x16 ) );
 		Set( "SceneOutliner.ChangedItemHighlight", new FSlateRoundedBoxBrush( FStyleColors::White, 1.0) );
 
-		Set("SceneOutliner.PinnedIcon", new IMAGE_BRUSH_SVG("Icons/SceneOutliner_Pinned", Icon16x16));
-		Set("SceneOutliner.UnpinnedIcon", new IMAGE_BRUSH_SVG("Icons/SceneOutliner_Unpinned", Icon16x16));
-
 		Set( "SceneOutliner.TableViewRow", FTableRowStyle(NormalTableRowStyle));
 	}
 
diff --git a/Engine/Source/Editor/PropertyEditor/Private/PropertyEditorToolkit.cpp b/Engine/Source/Editor/PropertyEditor/Private/PropertyEditorToolkit.cpp
--- a/Engine/Source/Editor/PropertyEditor/Private/PropertyEditorToolkit.cpp
+++ b/Engine/Source/Editor/PropertyEditor/Private/PropertyEditorToolkit.cpp
@@ -14,6 +14,7 @@
 
 #include "Widgets/Docking/SDockTab.h"
 #include "Subsystems/AssetEditorSubsystem.h"
+#include "Styling/StyleColors.h"
 
 #define LOCTEXT_NAMESPACE "PropertyEditorToolkit"
 
@@ -57,7 +58,6 @@ FPropertyEditorToolkit::FPropertyEditorToolkit()
 	, PathToRoot()
 	, DetailsView()
 {
-	PinSequence.AddCurve( 0, 1.0f, ECurveEaseFunction::QuadIn );
 }
 
 
@@ -183,8 +183,6 @@ void FPropertyEditorToolkit::Initialize( const EToolkitMode::Type Mode, const TS
 
 		DetailsView->SetObjects(AdjustedObjectsToEditWeak, true);
 
-		PinColor = FSlateColor(FLinearColor(1, 1, 1, 0));
-		GEditor->GetTimerManager()->SetTimer(TimerHandle_TickPinColor, FTimerDelegate::CreateSP(this, &FPropertyEditorToolkit::TickPinColorAndOpacity), 0.1f, true);
 	}
 }
 
@@ -248,7 +246,7 @@ TSharedRef<SDockTab> FPropertyEditorToolkit::SpawnTab_PropertyTable( const FSpaw
 				.VAlign( VAlign_Center )
 				[
 					SNew( SImage )
-					.Image( FEditorStyle::GetBrush( "PropertyEditor.RemoveColumn" ) )
+					.Image( FEditorStyle::GetBrush( "Icons.Unpinned" ) )
 					.Visibility( this, &FPropertyEditorToolkit::GetAddColumnInstructionsOverlayVisibility )
 				]
 
@@ -329,14 +327,14 @@ void FPropertyEditorToolkit::ConstructTreeColumns( const TSharedRef< class SHead
 		.ToolTipText( LOCTEXT("AddColumnLabel", "Push Pins to Add Columns") )
 		[
 			SNew( SImage )
-			.Image( FEditorStyle::GetBrush(TEXT("PropertyEditor.RemoveColumn")) )
+			.Image( FEditorStyle::GetBrush(TEXT("Icons.Unpinned")) )
+			.ColorAndOpacity(FStyleColors::Foreground)
 		]
 	];
 
 	HeaderRow->InsertColumn( ColumnArgs, 0 );
 }
 
-
 TSharedRef< SWidget > FPropertyEditorToolkit::ConstructTreeCell( const FName& ColumnName, const TSharedRef< IPropertyTreeRow >& Row )
 {
 	if ( ColumnName == TreePinAsColumnHeaderId )
@@ -349,9 +347,8 @@ TSharedRef< SWidget > FPropertyEditorToolkit::ConstructTreeCell( const FName& Co
 			.BorderImage( &FEditorStyle::GetWidgetStyle<FHeaderRowStyle>("PropertyTable.HeaderRow").ColumnStyle.NormalBrush )
 			[
 				SNew(SButton)
-				.IsFocusable(false)
 				.ToolTipText(NSLOCTEXT("PropertyEditor", "ToggleColumnButtonToolTip", "Toggle Column"))
-				.ButtonStyle( FEditorStyle::Get(), "NoBorder" )
+				.ButtonStyle( FEditorStyle::Get(), "HoverOnlyButton" )
 				.ContentPadding(0) 
 				.OnClicked( this, &FPropertyEditorToolkit::OnToggleColumnClicked, RowPtr )
 				.HAlign( HAlign_Center )
@@ -359,7 +356,7 @@ TSharedRef< SWidget > FPropertyEditorToolkit::ConstructTreeCell( const FName& Co
 				[
 					SNew(SImage)
 					.Image( this, &FPropertyEditorToolkit::GetToggleColumnButtonImageBrush, RowPtr )
-					.ColorAndOpacity( this, &FPropertyEditorToolkit::GetPinColorAndOpacity, RowPtr )
+					.ColorAndOpacity(FSlateColor::UseForeground())
 				]
 			];
 	}
@@ -432,12 +429,6 @@ bool FPropertyEditorToolkit::TableHasCustomColumns() const
 	return PropertyPathsAddedAsColumns.Num() > 0;
 }
 
-bool FPropertyEditorToolkit::CloseWindow()
-{
-	GEditor->GetTimerManager()->ClearTimer( TimerHandle_TickPinColor );
-	return FAssetEditorToolkit::CloseWindow();
-}
-
 
 bool FPropertyEditorToolkit::IsExposedAsColumn( const TWeakPtr< IPropertyTreeRow >& Row ) const
 {
@@ -662,70 +653,12 @@ const FSlateBrush* FPropertyEditorToolkit::GetToggleColumnButtonImageBrush( cons
 {
 	if ( IsExposedAsColumn( Row ) )
 	{
-		return FEditorStyle::GetBrush("PropertyEditor.RemoveColumn");
-	}
-
-	return FEditorStyle::GetBrush("PropertyEditor.AddColumn");
-}
-
-void FPropertyEditorToolkit::TickPinColorAndOpacity()
-{
-	bool IsRowBeingHoveredOver = false;
-	for (int Index = PinRows.Num() - 1; Index >= 0 ; Index--)
-	{
-		TSharedPtr< IPropertyTreeRow > Row = PinRows[ Index ].Pin();
-		if ( Row.IsValid() )
-		{
-			IsRowBeingHoveredOver |= Row->IsCursorHovering();
-
-			if ( IsRowBeingHoveredOver )
-			{
-				break;
-			}
-		}
-		else
-		{
-			PinRows.RemoveAt( Index );
-		}
-	}
-
-	if ( IsRowBeingHoveredOver )
-	{
-		PinSequence.JumpToStart();
-	}
-
-	float Opacity = 0.0f;
-	if ( !TableHasCustomColumns() )
-	{
-		Opacity = PinSequence.GetLerp();
-	}
-
-	if ( !PinSequence.IsPlaying() )
-	{
-		if ( PinSequence.IsAtStart() )
-		{
-			PinSequence.Play( PropertyTree.ToSharedRef() );
-		}
-		else
-		{
-			PinSequence.PlayReverse( PropertyTree.ToSharedRef() );
-		}
-	}
-
-	PinColor = FSlateColor( FColor( 255, 255, 255, FMath::Lerp( 0, 200, Opacity ) ).ReinterpretAsLinear() );
-}
-
-FSlateColor FPropertyEditorToolkit::GetPinColorAndOpacity( const TWeakPtr< IPropertyTreeRow > Row ) const
-{
-	if ( Row.IsValid() && ( Row.Pin()->IsCursorHovering() || IsExposedAsColumn( Row ) ) )
-	{
-		return FSlateColor( FLinearColor::White );
+		return FEditorStyle::GetBrush("Icons.Pinned");
 	}
 
-	return PinColor;
+	return FEditorStyle::GetBrush("Icons.Unpinned");
 }
 
-
 FString FPropertyEditorToolkit::GetWorldCentricTabPrefix() const
 {
 	check(0);
diff --git a/Engine/Source/Editor/PropertyEditor/Private/PropertyEditorToolkit.h b/Engine/Source/Editor/PropertyEditor/Private/PropertyEditorToolkit.h
--- a/Engine/Source/Editor/PropertyEditor/Private/PropertyEditorToolkit.h
+++ b/Engine/Source/Editor/PropertyEditor/Private/PropertyEditorToolkit.h
@@ -44,8 +44,6 @@ public:
 
 	bool TableHasCustomColumns() const;
 
-	virtual bool CloseWindow() override;
-
 	virtual bool IsPrimaryEditor() const override{ return false; };
 
 public:
@@ -86,10 +84,6 @@ private:
 
 	EVisibility GetToggleColumnButtonVisibility( const TSharedRef< class IPropertyTreeRow > Row ) const;
 
-	void TickPinColorAndOpacity();
-
-	FSlateColor GetPinColorAndOpacity( const TWeakPtr< IPropertyTreeRow > Row ) const;
-
 	void TableColumnsChanged();
 
 	EVisibility GetAddColumnInstructionsOverlayVisibility() const;
@@ -106,10 +100,6 @@ private:
 
 	TArray< TSharedRef< FPropertyPath > > PropertyPathsAddedAsColumns;
 
-	/** Animation sequence to pulse the pin image */
-	FCurveSequence PinSequence;
-	FTimerHandle TimerHandle_TickPinColor;
-	FSlateColor PinColor;
 	TArray< TWeakPtr<IPropertyTreeRow> > PinRows;
 
 	static const FName ToolkitFName;
diff --git a/Engine/Source/Editor/SceneOutliner/Private/SceneOutlinerPinnedColumn.cpp b/Engine/Source/Editor/SceneOutliner/Private/SceneOutlinerPinnedColumn.cpp
--- a/Engine/Source/Editor/SceneOutliner/Private/SceneOutlinerPinnedColumn.cpp
+++ b/Engine/Source/Editor/SceneOutliner/Private/SceneOutlinerPinnedColumn.cpp
@@ -70,8 +70,8 @@ public:
             .Image(this, &SPinnedWidget::GetBrush)
         );
 
-		static const FName NAME_PinnedBrush = TEXT("SceneOutliner.PinnedIcon");
-		static const FName NAME_UnpinnedBrush = TEXT("SceneOutliner.UnpinnedIcon");
+		static const FName NAME_PinnedBrush = TEXT("Icons.Pinned");
+		static const FName NAME_UnpinnedBrush = TEXT("Icons.Unpinned");
 
 		PinnedBrush = FAppStyle::Get().GetBrush(NAME_PinnedBrush);
 		UnpinnedBrush = FAppStyle::Get().GetBrush(NAME_UnpinnedBrush);
@@ -221,7 +221,7 @@ bool FSceneOutlinerPinnedColumn::IsItemPinned(const ISceneOutlinerTreeItem& Item
 
 const FSlateBrush* FSceneOutlinerPinnedColumn::GetHeaderIcon() const
 {
-	return FAppStyle::Get().GetBrush("SceneOutliner.UnpinnedIcon");
+	return FAppStyle::Get().GetBrush("Icons.Unpinned");
 }
 
 #undef LOCTEXT_NAMESPACE
\ No newline at end of file
