Fixed crash when using a TitleProperty that references an invalid property name.

diff --git a/Engine/Source/Editor/PropertyEditor/Private/UserInterface/PropertyEditor/SPropertyEditorArrayItem.cpp b/Engine/Source/Editor/PropertyEditor/Private/UserInterface/PropertyEditor/SPropertyEditorArrayItem.cpp
--- a/Engine/Source/Editor/PropertyEditor/Private/UserInterface/PropertyEditor/SPropertyEditorArrayItem.cpp
+++ b/Engine/Source/Editor/PropertyEditor/Private/UserInterface/PropertyEditor/SPropertyEditorArrayItem.cpp
@@ -22,13 +22,33 @@
 			FText::GetFormatPatternParameters(TitleFormatter->Format, OutParameterNames);
 			for (const FString& ParameterName : OutParameterNames)
 			{
-				TitleFormatter->PropertyHandles.Add(RootProperty->GetChildHandle(FName(*ParameterName), false));
+				TSharedPtr<IPropertyHandle> Handle = RootProperty->GetChildHandle(FName(*ParameterName), false /*bRecurse*/);
+				if (Handle.IsValid())
+				{
+					TitleFormatter->PropertyHandles.Add(Handle);
+				}
+				else
+				{
+					// title property doesn't exist, display an error message
+					TitleFormatter->PropertyHandles.Empty();
+					TitleFormatter->Format = FText::FromString(TEXT("Invalid Title Property!"));
+					break;
+				}
 			}
 		}
 		else // Support the old style where it was just a name of a property with no formatting.
 		{
-			TitleFormatter->PropertyHandles.Add(RootProperty->GetChildHandle(FName(*TitlePropertyRaw), false));
-			TitleFormatter->Format = FText::FromString(TEXT("{") + TitlePropertyRaw + TEXT("}"));
+			TSharedPtr<IPropertyHandle> Handle = RootProperty->GetChildHandle(FName(*TitlePropertyRaw), false /*bRecurse*/);
+			if (Handle.IsValid())
+			{
+				TitleFormatter->PropertyHandles.Add(Handle);
+				TitleFormatter->Format = FText::FromString(TEXT("{") + TitlePropertyRaw + TEXT("}"));
+			}
+			else
+			{
+				// title property doesn't exist, display an error message
+				TitleFormatter->Format = FText::FromString(TEXT("Invalid Title Property!"));
+			}
 		}
 
 		return TitleFormatter;
diff --git a/Engine/Source/Editor/UnrealEd/Classes/Editor/PropertyEditorTestObject.h b/Engine/Source/Editor/UnrealEd/Classes/Editor/PropertyEditorTestObject.h
--- a/Engine/Source/Editor/UnrealEd/Classes/Editor/PropertyEditorTestObject.h
+++ b/Engine/Source/Editor/UnrealEd/Classes/Editor/PropertyEditorTestObject.h
@@ -359,6 +359,15 @@ class UPropertyEditorTestObject : public UObject
 	UPROPERTY(EditAnywhere, Category=ArraysOfProperties, meta=(TitleProperty=IntPropertyInsideAStruct))
 	TArray<FPropertyEditorTestBasicStruct> StructPropertyArrayWithTitle;
 
+	UPROPERTY(EditAnywhere, Category=ArraysOfProperties, meta=(TitleProperty="{IntPropertyInsideAStruct} + {FloatPropertyInsideAStruct}"))
+	TArray<FPropertyEditorTestBasicStruct> StructPropertyArrayWithFormattedTitle;
+
+	UPROPERTY(EditAnywhere, Category=ArraysOfProperties, meta=(TitleProperty=ErrorProperty))
+	TArray<FPropertyEditorTestBasicStruct> StructPropertyArrayWithTitleError;
+
+	UPROPERTY(EditAnywhere, Category=ArraysOfProperties, meta=(TitleProperty="{ErrorProperty}"))
+	TArray<FPropertyEditorTestBasicStruct> StructPropertyArrayWithFormattedTitleError;
+
 	UPROPERTY(EditAnywhere, Category=ArraysOfProperties)
 	TArray<FPropertyEditorTestInstancedStruct> InstancedStructArray;
 
