Fixed a crash when pasting into an array on class sparse data

This would try and get the value of the outer array property (for the inner array value) using FObjectBaseAddress::ObjectOrStruct as the start address, however for class sparse data FObjectBaseAddress::ObjectOrStruct was set to the CDO of the class that owns the sparse data, rather than the start of the sparse data itself. This resulted in a crash as it would try and read garbage from the CDO.

This change removes FObjectBaseAddress::ObjectOrStruct and FObjectBaseAddress::bIsStruct, in favor of storing FObjectBaseAddress::Object (if not editing a struct) and FObjectBaseAddress::StructAddress. FObjectBaseAddress::StructAddress always points to the thing that actually contains FObjectBaseAddress::BaseAddress, which allows the ImportText code to get the correct offset for the array property within the sparse data instance.

#rnx

diff --git a/Engine/Source/Editor/PropertyEditor/Private/PropertyHandleImpl.cpp b/Engine/Source/Editor/PropertyEditor/Private/PropertyHandleImpl.cpp
--- a/Engine/Source/Editor/PropertyEditor/Private/PropertyHandleImpl.cpp
+++ b/Engine/Source/Editor/PropertyEditor/Private/PropertyHandleImpl.cpp
@@ -59,21 +59,20 @@ void FPropertyValueImpl::EnumerateObjectsToModify( FPropertyNode* InPropertyNode
 		const int32 NumInstances = ComplexNode->GetInstancesNum();
 		for (int32 Index = 0; Index < NumInstances; ++Index)
 		{
-			uint8* ObjectOrStruct = nullptr;
-			uint8* BaseAddress = nullptr;
+			UObject* Object = nullptr;
+			uint8* StructAddress = nullptr;
 			if (bIsStruct)
 			{
-				ObjectOrStruct = ComplexNode->GetMemoryOfInstance(Index);
-				BaseAddress = InPropertyNode->GetValueBaseAddress(ObjectOrStruct, false);
+				StructAddress = ComplexNode->GetMemoryOfInstance(Index);
 			}
 			else
 			{
-				const UObject* Obj = ComplexNode->GetInstanceAsUObject(Index).Get();
-				ObjectOrStruct = (uint8*)Obj;
-				BaseAddress = InPropertyNode->GetValueBaseAddressFromObject(Obj);
+				Object = ComplexNode->GetInstanceAsUObject(Index).Get();
+				StructAddress = InPropertyNode->GetStartAddressFromObject(Object);
 			}
+			uint8* BaseAddress = InPropertyNode->GetValueBaseAddress(StructAddress, InPropertyNode->HasNodeFlags(EPropertyNodeFlags::IsSparseClassData));
 
-			if (!InObjectsToModifyCallback(FObjectBaseAddress(ObjectOrStruct, BaseAddress, bIsStruct), Index, NumInstances))
+			if (!InObjectsToModifyCallback(FObjectBaseAddress(Object, StructAddress, BaseAddress), Index, NumInstances))
 			{
 				break;
 			}
@@ -214,7 +213,7 @@ FPropertyAccess::Result FPropertyValueImpl::ImportText( const FString& InValue,
 	TArray<FString> Values;
 	for (const FObjectBaseAddress& BaseAddress : ObjectsToModify)
 	{
-		if (BaseAddress.ObjectOrStruct)
+		if (BaseAddress.BaseAddress)
 		{
 			Values.Add(InValue);
 		}
@@ -261,7 +260,7 @@ FPropertyAccess::Result FPropertyValueImpl::ImportText( const TArray<FObjectBase
 		bool bIsGameWorld = false;
 		// If the object we are modifying is in the PIE world, than make the PIE world the active
 		// GWorld.  Assumes all objects managed by this property window belong to the same world.
-		UObject* FirstObject = InObjects[0].GetUObject();
+		UObject* FirstObject = InObjects[0].Object;
 		if (UPackage* ObjectPackage = (FirstObject ? FirstObject->GetOutermost() : nullptr))
 		{
 			const bool bIsPIEPackage = ObjectPackage->HasAnyPackageFlags(PKG_PlayInEditor);
@@ -295,7 +294,7 @@ FPropertyAccess::Result FPropertyValueImpl::ImportText( const TArray<FObjectBase
 				break;
 			}
 
-			UObject* CurObject = Cur.GetUObject();
+			UObject* CurObject = Cur.Object;
 
 			const FString& NewValue = InValues[ObjectIndex];
 
@@ -324,7 +323,7 @@ FPropertyAccess::Result FPropertyValueImpl::ImportText( const TArray<FObjectBase
 
 				if (bIsInContainer)
 				{
-					uint8* ValueBaseAddress = ParentNode->GetValueBaseAddress(Cur.ObjectOrStruct, bIsSparseClassData);
+					uint8* ValueBaseAddress = ParentNode->GetValueBaseAddress(Cur.StructAddress, bIsSparseClassData);
 					
 					FScriptSetHelper SetHelper(SetProperty, ValueBaseAddress);
 					if (SetHelper.HasElement(Cur.BaseAddress, NewValue) &&
@@ -343,7 +342,7 @@ FPropertyAccess::Result FPropertyValueImpl::ImportText( const TArray<FObjectBase
 
 				if (bIsInContainer)
 				{
-					uint8* ValueBaseAddress = ParentNode->GetValueBaseAddress(Cur.ObjectOrStruct, bIsSparseClassData);
+					uint8* ValueBaseAddress = ParentNode->GetValueBaseAddress(Cur.StructAddress, bIsSparseClassData);
 					
 					FScriptMapHelper MapHelper(MapProperty, ValueBaseAddress);
 					if (MapHelper.HasKey(Cur.BaseAddress, NewValue) && 
@@ -363,7 +362,7 @@ FPropertyAccess::Result FPropertyValueImpl::ImportText( const TArray<FObjectBase
 
 			if (bIsInContainer)
 			{
-				uint8* Addr = ParentNode->GetValueBaseAddress(Cur.ObjectOrStruct, bIsSparseClassData);
+				uint8* Addr = ParentNode->GetValueBaseAddress(Cur.StructAddress, bIsSparseClassData);
 				Property->ExportText_Direct(PreviousContainerValue, Addr, Addr, nullptr, 0);
 			}
 
@@ -429,13 +428,13 @@ FPropertyAccess::Result FPropertyValueImpl::ImportText( const TArray<FObjectBase
 				// For TMap and TSet, we need to rehash it in case a key was modified
 				if (NodeProperty->GetOwner<FMapProperty>())
 				{
-					uint8* Addr = InPropertyNode->GetParentNode()->GetValueBaseAddress(Cur.ObjectOrStruct, bIsSparseClassData);
+					uint8* Addr = InPropertyNode->GetParentNode()->GetValueBaseAddress(Cur.StructAddress, bIsSparseClassData);
 					FScriptMapHelper MapHelper(NodeProperty->GetOwner<FMapProperty>(), Addr);
 					MapHelper.Rehash();
 				}
 				else if (NodeProperty->GetOwner<FSetProperty>())
 				{
-					uint8* Addr = InPropertyNode->GetParentNode()->GetValueBaseAddress(Cur.ObjectOrStruct, bIsSparseClassData);
+					uint8* Addr = InPropertyNode->GetParentNode()->GetValueBaseAddress(Cur.StructAddress, bIsSparseClassData);
 					FScriptSetHelper SetHelper(NodeProperty->GetOwner<FSetProperty>(), Addr);
 					SetHelper.Rehash();
 				}
@@ -4731,7 +4730,7 @@ bool FPropertyHandleSet::HasDefaultElement()
 		{
 			const bool IsSparseClassData = PropNode->HasNodeFlags(EPropertyNodeFlags::IsSparseClassData) != 0;
 			FSetProperty* SetProperty = CastFieldChecked<FSetProperty>(PropNode->GetProperty());
-			FScriptSetHelper SetHelper(SetProperty, PropNode->GetValueBaseAddress(Addresses[0].ObjectOrStruct, IsSparseClassData));
+			FScriptSetHelper SetHelper(SetProperty, PropNode->GetValueBaseAddress(Addresses[0].StructAddress, IsSparseClassData));
 
 			FDefaultConstructedPropertyElement DefaultElement(SetHelper.ElementProp);
 			return SetHelper.FindElementIndex(DefaultElement.GetObjAddress()) != INDEX_NONE;
@@ -4829,7 +4828,7 @@ bool FPropertyHandleMap::HasDefaultKey()
 		{
 			const bool IsSparseClassData = PropNode->HasNodeFlags(EPropertyNodeFlags::IsSparseClassData) != 0;
 			FMapProperty* MapProperty = CastFieldChecked<FMapProperty>(PropNode->GetProperty());
-			FScriptMapHelper MapHelper(MapProperty, PropNode->GetValueBaseAddress(Addresses[0].ObjectOrStruct, IsSparseClassData));
+			FScriptMapHelper MapHelper(MapProperty, PropNode->GetValueBaseAddress(Addresses[0].StructAddress, IsSparseClassData));
 
 			FDefaultConstructedPropertyElement DefaultKey(MapHelper.KeyProp);
 			return MapHelper.FindMapIndexWithKey(DefaultKey.GetObjAddress()) != INDEX_NONE;
diff --git a/Engine/Source/Editor/PropertyEditor/Private/PropertyHandleImpl.h b/Engine/Source/Editor/PropertyEditor/Private/PropertyHandleImpl.h
--- a/Engine/Source/Editor/PropertyEditor/Private/PropertyHandleImpl.h
+++ b/Engine/Source/Editor/PropertyEditor/Private/PropertyHandleImpl.h
@@ -15,25 +15,19 @@ class FPropertyRestriction;
 class FObjectBaseAddress
 {
 public:
-	FObjectBaseAddress()
-		:	ObjectOrStruct(nullptr)
-		,	BaseAddress(nullptr)
-		,	bIsStruct(false)
+	FObjectBaseAddress() = default;
+	FObjectBaseAddress(UObject* InObject, uint8* InStructAddress, uint8* InBaseAddress)
+		: Object(InObject)
+		, StructAddress(InStructAddress)
+		, BaseAddress(InBaseAddress)
 	{}
-	FObjectBaseAddress(uint8* InObjectOrStruct, uint8* InBaseAddress, bool InIsStruct)
-		:	ObjectOrStruct(InObjectOrStruct)
-		,	BaseAddress(InBaseAddress)
-		,	bIsStruct(InIsStruct)
-	{}
-
-	FORCEINLINE UObject* GetUObject() const
-	{
-		return bIsStruct ? nullptr : (UObject*) ObjectOrStruct;
-	}
 
-	uint8*		ObjectOrStruct;
-	uint8*		BaseAddress;
-	bool		bIsStruct;
+	/** Object associated with the value, or null when editing a struct. */
+	UObject* Object = nullptr;
+	/** Pointer to the object/struct instance that actually owns the value (may be pointing to sidecar data). */
+	uint8* StructAddress = nullptr;
+	/** Pointer to the base address of the value within StructAddress. */
+	uint8* BaseAddress = nullptr;
 };
 
 /**
diff --git a/Engine/Source/Editor/PropertyEditor/Private/PropertyNode.cpp b/Engine/Source/Editor/PropertyEditor/Private/PropertyNode.cpp
--- a/Engine/Source/Editor/PropertyEditor/Private/PropertyNode.cpp
+++ b/Engine/Source/Editor/PropertyEditor/Private/PropertyNode.cpp
@@ -1210,7 +1210,7 @@ FPropertyAccess::Result FPropertyNode::GetSingleReadAddress(uint8*& OutValueAddr
 	return ReadAddresses.Num() > 1 ? FPropertyAccess::MultipleValues : FPropertyAccess::Fail;
 }
 
-uint8* FPropertyNode::GetStartAddress(const UObject* Obj) const
+uint8* FPropertyNode::GetStartAddressFromObject(const UObject* Obj) const
 {
 	if (!Obj)
 	{
@@ -1227,12 +1227,12 @@ uint8* FPropertyNode::GetStartAddress(const UObject* Obj) const
 
 uint8* FPropertyNode::GetValueBaseAddressFromObject(const UObject* Obj) const
 {
-	return GetValueBaseAddress(GetStartAddress(Obj), HasNodeFlags(EPropertyNodeFlags::IsSparseClassData));
+	return GetValueBaseAddress(GetStartAddressFromObject(Obj), HasNodeFlags(EPropertyNodeFlags::IsSparseClassData));
 }
 
 uint8* FPropertyNode::GetValueAddressFromObject(const UObject* Obj) const
 {
-	return GetValueAddress(GetStartAddress(Obj), HasNodeFlags(EPropertyNodeFlags::IsSparseClassData));
+	return GetValueAddress(GetStartAddressFromObject(Obj), HasNodeFlags(EPropertyNodeFlags::IsSparseClassData));
 }
 
 
diff --git a/Engine/Source/Editor/PropertyEditor/Private/PropertyNode.h b/Engine/Source/Editor/PropertyEditor/Private/PropertyNode.h
--- a/Engine/Source/Editor/PropertyEditor/Private/PropertyNode.h
+++ b/Engine/Source/Editor/PropertyEditor/Private/PropertyNode.h
@@ -504,6 +504,12 @@ public:
 	 */
 	virtual uint8* GetValueAddress(uint8* StartAddress, bool bIsSparseData) const;
 
+	/**
+	 * Caclulates the memory address for the starting point of the structure that contains the property this node uses.
+	 * This will often be Obj but may also point to a sidecar data structure.
+	 */
+	uint8* GetStartAddressFromObject(const UObject* Obj) const;
+
 	/**
 	 * Calculates the memory address for the data associated with this item's property.  This is typically the value of a FProperty or a UObject address.
 	 *
@@ -938,10 +944,6 @@ public:
 
 protected:
 
-	// Returns a pointer to the starting point of the structure that contains the property this node uses.
-	// This will often be Obj but may also point to a sidecar data structure
-	uint8* GetStartAddress(const UObject* Obj) const;
-
 	TSharedRef<FEditPropertyChain> BuildPropertyChain( FProperty* PropertyAboutToChange );
 	TSharedRef<FEditPropertyChain> BuildPropertyChain( FProperty* PropertyAboutToChange, const TSet<UObject*>& InAffectedArchetypeInstances );
 	TSharedRef<FEditPropertyChain> BuildPropertyChain( FProperty* PropertyAboutToChange, TSet<UObject*>&& InAffectedArchetypeInstances );
diff --git a/Engine/Source/Editor/PropertyEditor/Private/PropertyTextUtilities.cpp b/Engine/Source/Editor/PropertyEditor/Private/PropertyTextUtilities.cpp
--- a/Engine/Source/Editor/PropertyEditor/Private/PropertyTextUtilities.cpp
+++ b/Engine/Source/Editor/PropertyEditor/Private/PropertyTextUtilities.cpp
@@ -18,7 +18,7 @@ void FPropertyTextUtilities::PropertyToTextHelper(FString& OutString, const FPro
 
 void FPropertyTextUtilities::PropertyToTextHelper(FString& OutString, const FPropertyNode* InPropertyNode, const FProperty* Property, const FObjectBaseAddress& ObjectAddress, EPropertyPortFlags PortFlags)
 {
-	PropertyToTextHelper(OutString, InPropertyNode, Property, ObjectAddress.BaseAddress, ObjectAddress.GetUObject(), PortFlags);
+	PropertyToTextHelper(OutString, InPropertyNode, Property, ObjectAddress.BaseAddress, ObjectAddress.Object, PortFlags);
 	}
 
 void FPropertyTextUtilities::TextToPropertyHelper(const TCHAR* Buffer, const FPropertyNode* InPropertyNode, const FProperty* Property, uint8* ValueAddress, UObject* Object, EPropertyPortFlags PortFlags)
@@ -35,5 +35,5 @@ void FPropertyTextUtilities::TextToPropertyHelper(const TCHAR* Buffer, const FPr
 
 void FPropertyTextUtilities::TextToPropertyHelper(const TCHAR* Buffer, const FPropertyNode* InPropertyNode, const FProperty* Property, const FObjectBaseAddress& ObjectAddress, EPropertyPortFlags PortFlags)
 {
-	TextToPropertyHelper(Buffer, InPropertyNode, Property, ObjectAddress.BaseAddress, ObjectAddress.GetUObject(), PortFlags);
+	TextToPropertyHelper(Buffer, InPropertyNode, Property, ObjectAddress.BaseAddress, ObjectAddress.Object, PortFlags);
 }
\ No newline at end of file
