Fixed InvalidateCachedState not being called when editing structs

It was specifically looking for an object parent rather than a complex parent, which meant it skipped invalidating the cache when editing a struct (including a nested external struct) rather than an object.

#rnx

diff --git a/Engine/Source/Editor/PropertyEditor/Private/PropertyNode.cpp b/Engine/Source/Editor/PropertyEditor/Private/PropertyNode.cpp
--- a/Engine/Source/Editor/PropertyEditor/Private/PropertyNode.cpp
+++ b/Engine/Source/Editor/PropertyEditor/Private/PropertyNode.cpp
@@ -2465,15 +2465,17 @@ void FPropertyNode::NotifyPostChange( FPropertyChangedEvent& InPropertyChangedEv
 	// remember the property that was the chain's original active property; this will correspond to the outermost property of struct/array that was modified
 	FProperty* const OriginalActiveProperty = PropertyChain->GetActiveMemberNode()->GetValue();
 
-	// invalidate the entire chain of objects in the hierarchy 
-	FObjectPropertyNode* CurrentObjectNode = FindObjectItemParent();
-	while (CurrentObjectNode != nullptr)
+	// invalidate the entire chain of objects in the hierarchy
 	{
-		CurrentObjectNode->InvalidateCachedState();
+		FComplexPropertyNode* ComplexNode = FindComplexParent();
+		while (ComplexNode)
+		{
+			ComplexNode->InvalidateCachedState();
 
-		// FindObjectItemParent returns itself if the node is an object, so step up the hierarchy to get to its actual parent object
-		FPropertyNode* CurrentParent = CurrentObjectNode->GetParentNode();
-		CurrentObjectNode = CurrentParent != nullptr ? CurrentParent->FindObjectItemParent() : nullptr;
+			// FindComplexParent returns itself if the node is an object, so step up the hierarchy to get to its actual parent object
+			FPropertyNode* CurrentParent = ComplexNode->GetParentNode();
+			ComplexNode = CurrentParent != nullptr ? CurrentParent->FindComplexParent() : nullptr;
+		}
 	}
 
 	FObjectPropertyNode* ObjectNode = FindObjectItemParent();
