Report EditCondition errors to the Output Log

#rnx

diff --git a/Engine/Source/Editor/PropertyEditor/Private/EditConditionContext.cpp b/Engine/Source/Editor/PropertyEditor/Private/EditConditionContext.cpp
--- a/Engine/Source/Editor/PropertyEditor/Private/EditConditionContext.cpp
+++ b/Engine/Source/Editor/PropertyEditor/Private/EditConditionContext.cpp
@@ -6,7 +6,7 @@
 #include "PropertyNode.h"
 #include "PropertyEditorHelpers.h"
 
-DEFINE_LOG_CATEGORY_STATIC(LogEditCondition, Log, All);
+DEFINE_LOG_CATEGORY(LogEditCondition);
 
 FEditConditionContext::FEditConditionContext(FPropertyNode& InPropertyNode)
 {
diff --git a/Engine/Source/Editor/PropertyEditor/Private/EditConditionContext.h b/Engine/Source/Editor/PropertyEditor/Private/EditConditionContext.h
--- a/Engine/Source/Editor/PropertyEditor/Private/EditConditionContext.h
+++ b/Engine/Source/Editor/PropertyEditor/Private/EditConditionContext.h
@@ -4,6 +4,8 @@
 
 #include "CoreMinimal.h"
 
+DECLARE_LOG_CATEGORY_EXTERN(LogEditCondition, Log, All);
+
 class IEditConditionContext
 {
 public:
diff --git a/Engine/Source/Editor/PropertyEditor/Private/EditConditionParser.cpp b/Engine/Source/Editor/PropertyEditor/Private/EditConditionParser.cpp
--- a/Engine/Source/Editor/PropertyEditor/Private/EditConditionParser.cpp
+++ b/Engine/Source/Editor/PropertyEditor/Private/EditConditionParser.cpp
@@ -118,6 +118,24 @@ static TOptional<FExpressionError> ConsumePropertyName(FExpressionTokenConsumer&
 	return TOptional<FExpressionError>();
 }
 
+template <typename ValueType>
+static void LogEditConditionError(const TValueOrError<ValueType, FExpressionError>& Error)
+{
+	if (!Error.HasError())
+	{
+		return;
+	}
+
+	static TSet<FString> ErrorsAlreadyLogged;
+	
+	const FString& Message = Error.GetError().Text.ToString();
+	if (!ErrorsAlreadyLogged.Find(Message))
+	{
+		ErrorsAlreadyLogged.Add(Message);
+		UE_LOG(LogEditCondition, Error, TEXT("%s"), *Message);
+	}
+}
+
 template <typename T>
 TOptional<T> GetValueInternal(const IEditConditionContext& Context, const FString& PropertyName)
 {
@@ -715,6 +733,10 @@ TValueOrError<bool, FText> FEditConditionParser::Evaluate(const FEditConditionEx
 			}
 		}
 	}
+	else
+	{
+		LogEditConditionError(Result);
+	}
 
 	const FText ErrorText = Result.HasError() ? Result.GetError().Text : FText::GetEmpty();
 	return MakeError(ErrorText);
@@ -732,6 +754,14 @@ TSharedPtr<FEditConditionExpression> FEditConditionParser::Parse(const FString&
 		{
 			return TSharedPtr<FEditConditionExpression>(new FEditConditionExpression(CompileResult.StealValue()));
 		}
+		else
+		{
+			LogEditConditionError(CompileResult);
+		}
+	}
+	else
+	{
+		LogEditConditionError(LexResult);
 	}
 
 	return TSharedPtr<FEditConditionExpression>();
