Property Editor - Update the IDetailTreeNode to expose initial expansion information which can be set through customizations so that uses of the property row generator can access that data consistently.

diff --git a/Engine/Source/Editor/PropertyEditor/Private/DetailCategoryBuilderImpl.cpp b/Engine/Source/Editor/PropertyEditor/Private/DetailCategoryBuilderImpl.cpp
--- a/Engine/Source/Editor/PropertyEditor/Private/DetailCategoryBuilderImpl.cpp
+++ b/Engine/Source/Editor/PropertyEditor/Private/DetailCategoryBuilderImpl.cpp
@@ -750,6 +750,11 @@ void FDetailCategoryImpl::GetFilterStrings(TArray<FString>& OutFilterStrings) co
 	OutFilterStrings.Add(GetDisplayName().ToString());
 }
 
+bool FDetailCategoryImpl::GetInitiallyCollapsed() const
+{
+	return bShouldBeInitiallyCollapsed;
+}
+
 void FDetailCategoryImpl::OnItemExpansionChanged(bool bIsExpanded, bool bShouldSaveState)
 {
 	if (bRestoreExpansionState && bShouldSaveState)
diff --git a/Engine/Source/Editor/PropertyEditor/Private/DetailCategoryBuilderImpl.h b/Engine/Source/Editor/PropertyEditor/Private/DetailCategoryBuilderImpl.h
--- a/Engine/Source/Editor/PropertyEditor/Private/DetailCategoryBuilderImpl.h
+++ b/Engine/Source/Editor/PropertyEditor/Private/DetailCategoryBuilderImpl.h
@@ -190,6 +190,7 @@ public:
 	virtual EDetailNodeType GetNodeType() const override { return EDetailNodeType::Category; }
 	virtual TSharedPtr<IPropertyHandle> CreatePropertyHandle() const override { return nullptr; }
 	virtual void GetFilterStrings(TArray<FString>& OutFilterStrings) const override;
+	virtual bool GetInitiallyCollapsed() const override;
 
 	virtual void GetChildren(FDetailNodeList& OutChildren) override;
 	virtual bool ShouldBeExpanded() const override;
@@ -310,6 +311,9 @@ public:
 	/** Is this the Favorites category? */
 	bool IsFavoriteCategory() const { return bFavoriteCategory; }
 
+	/** Is this category initially collapsed? */
+	bool GetShouldBeInitiallyCollapsed() const { return bShouldBeInitiallyCollapsed; }
+
 	FDetailLayoutCustomization* GetDefaultCustomization(TSharedRef<FPropertyNode> PropertyNode);
 
 private:
diff --git a/Engine/Source/Editor/PropertyEditor/Private/DetailItemNode.cpp b/Engine/Source/Editor/PropertyEditor/Private/DetailItemNode.cpp
--- a/Engine/Source/Editor/PropertyEditor/Private/DetailItemNode.cpp
+++ b/Engine/Source/Editor/PropertyEditor/Private/DetailItemNode.cpp
@@ -126,6 +126,15 @@ void FDetailItemNode::GetFilterStrings(TArray<FString>& OutFilterStrings) const
 	}
 }
 
+bool FDetailItemNode::GetInitiallyCollapsed() const
+{
+	if (Customization.IsValidCustomization() && Customization.PropertyRow.IsValid())
+	{
+		return Customization.PropertyRow->GetForceAutoExpansion() == false;
+	}
+	return true;
+}
+
 void FDetailItemNode::InitPropertyEditor()
 {
 	FProperty* NodeProperty = Customization.GetPropertyNode()->GetProperty();
diff --git a/Engine/Source/Editor/PropertyEditor/Private/DetailItemNode.h b/Engine/Source/Editor/PropertyEditor/Private/DetailItemNode.h
--- a/Engine/Source/Editor/PropertyEditor/Private/DetailItemNode.h
+++ b/Engine/Source/Editor/PropertyEditor/Private/DetailItemNode.h
@@ -25,6 +25,7 @@ public:
 	virtual EDetailNodeType GetNodeType() const override;
 	virtual TSharedPtr<IPropertyHandle> CreatePropertyHandle() const override;
 	virtual void GetFilterStrings(TArray<FString>& OutFilterStrings) const override;
+	virtual bool GetInitiallyCollapsed() const override;
 
 	/**
 	 * Initializes this node                                                              
diff --git a/Engine/Source/Editor/PropertyEditor/Private/DetailTreeNode.h b/Engine/Source/Editor/PropertyEditor/Private/DetailTreeNode.h
--- a/Engine/Source/Editor/PropertyEditor/Private/DetailTreeNode.h
+++ b/Engine/Source/Editor/PropertyEditor/Private/DetailTreeNode.h
@@ -38,6 +38,7 @@ public:
 	virtual void GetChildren(TArray<TSharedRef<IDetailTreeNode>>& OutChildren);
 	virtual TSharedPtr<class IDetailPropertyRow> GetRow() const override { return nullptr; }
 	virtual void GetFilterStrings(TArray<FString>& OutFilterStrings) const override { };
+	virtual bool GetInitiallyCollapsed() const override { return false; }
 
 	/** @return The details view that this node is in */
 	virtual IDetailsViewPrivate* GetDetailsView() const = 0;
diff --git a/Engine/Source/Editor/PropertyEditor/Public/IDetailTreeNode.h b/Engine/Source/Editor/PropertyEditor/Public/IDetailTreeNode.h
--- a/Engine/Source/Editor/PropertyEditor/Public/IDetailTreeNode.h
+++ b/Engine/Source/Editor/PropertyEditor/Public/IDetailTreeNode.h
@@ -167,4 +167,9 @@ public:
 	 * Gets the filter strings for this node in the tree.
 	 */
 	virtual void GetFilterStrings(TArray<FString>& OutFilterStrings) const = 0;
+
+	/**
+	 * Gets if this node should be initially collapsed by default.
+	 */
+	virtual bool GetInitiallyCollapsed() const = 0;
 };
\ No newline at end of file
