Fix crash when pasting an empty string, or string with multiple empty entries using newlines separators, into bulk property editor. Crash was an index-out-of-bound because code assumed at least one entry would exist after calling ParseIntoArray

Explicitly handle pasting a single empty string by creating Row and Cell entries in that situation. Updated paste handler so when pasting a string with multiple empty strings with newline separators, each empty string is treated like an empty row string. This makes pasting a string with empty row strings more consistent, as rows will no longer get trimmed and potentially mutate the expected outcome

diff --git a/Engine/Source/Editor/PropertyEditor/Private/Presentation/PropertyTable/PropertyTable.cpp b/Engine/Source/Editor/PropertyEditor/Private/Presentation/PropertyTable/PropertyTable.cpp
--- a/Engine/Source/Editor/PropertyEditor/Private/Presentation/PropertyTable/PropertyTable.cpp
+++ b/Engine/Source/Editor/PropertyEditor/Private/Presentation/PropertyTable/PropertyTable.cpp
@@ -480,11 +480,21 @@ void FPropertyTable::PasteTextAtCell( const FString& Text, const TSharedRef< IPr
 
 	// Parse into row strings
 	TArray<FString> RowStrings;
-	Text.ParseIntoArray(RowStrings,LINE_TERMINATOR, true);
+	Text.ParseIntoArray(RowStrings,LINE_TERMINATOR, false);
+	if (RowStrings.Num() == 0)
+	{
+		// Pasted an empty string, ParseIntoArray didn't create empty entry, so create that now
+		RowStrings.Emplace(TEXT(""));
+	}
 
 	// Parse row strings into individual cell strings
 	TArray<FString> CellStrings;
 	RowStrings[CurrentRowIdx++].ParseIntoArray(CellStrings, TEXT("\t"), false);
+	if (CellStrings.Num() == 0)
+	{
+		// Pasted an empty string, ParseIntoArray didn't create empty entry, so create that now
+		CellStrings.Emplace(TEXT(""));
+	}
 
 	// Get the maximum paste operations before displaying the slow task
 	int32 NumPasteOperationsBeforeWarning = GetDefault<UEditorPerProjectUserSettings>()->PropertyMatrix_NumberOfPasteOperationsBeforeWarning;
@@ -528,6 +538,11 @@ void FPropertyTable::PasteTextAtCell( const FString& Text, const TSharedRef< IPr
 					CurrentColumnIdx = 0;
 					FirstCellInRow = TargetCell;
 					RowStrings[CurrentRowIdx++].ParseIntoArray(CellStrings, TEXT("\t"), false);
+					if (CellStrings.Num() == 0)
+					{
+						// Tried to parse an empty row string, create an empty cell string
+						CellStrings.Emplace(TEXT(""));
+					}
 				
 					if ( bShowProgressDialog )
 					{
