Fix TMap key EditInlineNew if value is EditInlineNew.

diff --git a/Engine/Source/Editor/PropertyEditor/Private/PropertyNode.cpp b/Engine/Source/Editor/PropertyEditor/Private/PropertyNode.cpp
--- a/Engine/Source/Editor/PropertyEditor/Private/PropertyNode.cpp
+++ b/Engine/Source/Editor/PropertyEditor/Private/PropertyNode.cpp
@@ -171,7 +171,18 @@ void FPropertyNode::InitNode(const FPropertyNodeInitParams& InitParams)
 		const FProperty* OwnerProperty = MyProperty->GetOwnerProperty();
 
 		const bool bIsObjectOrInterface = CastField<FObjectPropertyBase>(MyProperty) || CastField<FInterfaceProperty>(MyProperty);
-		const bool bIsInsideContainer = CastField<FArrayProperty>(OwnerProperty) || CastField<FSetProperty>(OwnerProperty) || CastField<FMapProperty>(OwnerProperty);
+		bool bIsInsideContainer = CastField<FArrayProperty>(OwnerProperty) || CastField<FSetProperty>(OwnerProperty) || CastField<FMapProperty>(OwnerProperty);
+
+		// Don't consider the container's inline status if the key is a class property that is not inline
+		if (const FMapProperty* MapProperty = CastField<FMapProperty>(OwnerProperty))
+		{
+			const FObjectPropertyBase* KeyObjectProperty = CastField<FObjectPropertyBase>(MapProperty->GetKeyProperty());
+
+			if (KeyObjectProperty && KeyObjectProperty->PropertyClass && !KeyObjectProperty->PropertyClass->HasAnyClassFlags(EClassFlags::CLASS_EditInlineNew))
+			{
+				bIsInsideContainer = false;
+			}
+		}
 
 		// true if the property can be expanded into the property window; that is, instead of seeing
 		// a pointer to the object, you see the object's properties.
