Check if asset is UWorld before opening to prevent loss of work in current level.

diff --git a/Engine/Source/Editor/PropertyEditor/Private/UserInterface/PropertyEditor/SPropertyEditorAsset.cpp b/Engine/Source/Editor/PropertyEditor/Private/UserInterface/PropertyEditor/SPropertyEditorAsset.cpp
--- a/Engine/Source/Editor/PropertyEditor/Private/UserInterface/PropertyEditor/SPropertyEditorAsset.cpp
+++ b/Engine/Source/Editor/PropertyEditor/Private/UserInterface/PropertyEditor/SPropertyEditorAsset.cpp
@@ -29,6 +29,7 @@
 #include "UnrealEdGlobals.h"
 #include "Framework/Notifications/NotificationManager.h"
 #include "Widgets/Notifications/SNotificationList.h"
+#include "FileHelpers.h"
 
 #define LOCTEXT_NAMESPACE "PropertyEditor"
 
@@ -1158,6 +1159,17 @@ void SPropertyEditorAsset::OnOpenAssetEditor()
 	UObject* ObjectToEdit = Value.AssetData.GetAsset();
 	if( ObjectToEdit )
 	{
+		if (UWorld* World = Cast<UWorld>(ObjectToEdit))
+		{
+			constexpr bool bPromptUserToSave = true;
+			constexpr bool bSaveMapPackages = true;
+			constexpr bool bSaveContentPackages = true;
+			if (!FEditorFileUtils::SaveDirtyPackages(bPromptUserToSave, bSaveMapPackages, bSaveContentPackages))
+			{
+				return;
+			}
+		}
+
 		GEditor->EditObject( ObjectToEdit );
 	}
 }
