#UE Move objects referred to by instanced array/set/map properties to the transient package when removed from the container type. This was already the case for ActorComponents, but now applies to all objects referred to by instanced properties.

diff --git a/Engine/Source/Editor/PropertyEditor/Private/PropertyHandleImpl.cpp b/Engine/Source/Editor/PropertyEditor/Private/PropertyHandleImpl.cpp
--- a/Engine/Source/Editor/PropertyEditor/Private/PropertyHandleImpl.cpp
+++ b/Engine/Source/Editor/PropertyEditor/Private/PropertyHandleImpl.cpp
@@ -1137,18 +1137,18 @@ void FPropertyValueImpl::ClearChildren()
 						{
 							FScriptArrayHelper ArrayHelper(ArrayProperty, Addr);
 
-							// If the inner property is an instanced component property we must move the old components to the 
-							// transient package so resetting owned components on the parent doesn't find them
+							// If the inner property is an instanced property we must move the old objects to the 
+							// transient package so code looking for objects of this type on the parent doesn't find them
 							FObjectProperty* InnerObjectProperty = CastField<FObjectProperty>(ArrayProperty->Inner);
-							if (InnerObjectProperty && InnerObjectProperty->HasAnyPropertyFlags(CPF_InstancedReference) && InnerObjectProperty->PropertyClass->IsChildOf(UActorComponent::StaticClass()))
+							if (InnerObjectProperty && InnerObjectProperty->HasAnyPropertyFlags(CPF_InstancedReference))
 							{
 								const int32 ArraySize = ArrayHelper.Num();
 								for (int32 Index = 0; Index < ArraySize; ++Index)
 								{
-									if (UActorComponent* Component = *reinterpret_cast<UActorComponent**>(ArrayHelper.GetRawPtr(Index)))
+									if (UObject* InstancedObject = *reinterpret_cast<UObject**>(ArrayHelper.GetRawPtr(Index)))
 									{
-										Component->Modify();
-										Component->Rename(nullptr, GetTransientPackage(), REN_DontCreateRedirectors);
+										InstancedObject->Modify();
+										InstancedObject->Rename(nullptr, GetTransientPackage(), REN_DontCreateRedirectors);
 									}
 								}
 							}
@@ -1159,10 +1159,10 @@ void FPropertyValueImpl::ClearChildren()
 						{
 							FScriptSetHelper SetHelper(SetProperty, Addr);
 
-							// If the element property is an instanced component property we must move the old components to the 
-							// transient package so resetting owned components on the parent doesn't find them
+							// If the element property is an instanced property we must move the old objects to the 
+							// transient package so code looking for objects of this type on the parent doesn't find them
 							FObjectProperty* ElementObjectProperty = CastField<FObjectProperty>(SetProperty->ElementProp);
-							if (ElementObjectProperty && ElementObjectProperty->HasAnyPropertyFlags(CPF_InstancedReference) && ElementObjectProperty->PropertyClass->IsChildOf(UActorComponent::StaticClass()))
+							if (ElementObjectProperty && ElementObjectProperty->HasAnyPropertyFlags(CPF_InstancedReference))
 							{
 								int32 ElementsToRemove = SetHelper.Num();
 								int32 Index = 0;
@@ -1170,10 +1170,10 @@ void FPropertyValueImpl::ClearChildren()
 								{
 									if (SetHelper.IsValidIndex(Index))
 									{
-										if (UActorComponent* Component = *reinterpret_cast<UActorComponent**>(SetHelper.GetElementPtr(Index)))
+										if (UObject* InstancedObject = *reinterpret_cast<UObject**>(SetHelper.GetElementPtr(Index)))
 										{
-											Component->Modify();
-											Component->Rename(nullptr, GetTransientPackage(), REN_DontCreateRedirectors);
+											InstancedObject->Modify();
+											InstancedObject->Rename(nullptr, GetTransientPackage(), REN_DontCreateRedirectors);
 										}
 										--ElementsToRemove;
 									}
@@ -1187,10 +1187,10 @@ void FPropertyValueImpl::ClearChildren()
 						{
 							FScriptMapHelper MapHelper(MapProperty, Addr);
 
-							// If the map's value property is an instanced component property we must move the old components to the 
-							// transient package so resetting owned components on the parent doesn't find them
+							// If the map's value property is an instanced property we must move the old objects to the 
+							// transient package so code looking for objects of this type on the parent doesn't find them
 							FObjectProperty* ValueObjectProperty = CastField<FObjectProperty>(MapProperty->ValueProp);
-							if (ValueObjectProperty && ValueObjectProperty->HasAnyPropertyFlags(CPF_InstancedReference) && ValueObjectProperty->PropertyClass->IsChildOf(UActorComponent::StaticClass()))
+							if (ValueObjectProperty && ValueObjectProperty->HasAnyPropertyFlags(CPF_InstancedReference))
 							{
 								int32 ElementsToRemove = MapHelper.Num();
 								int32 Index = 0;
@@ -1198,10 +1198,10 @@ void FPropertyValueImpl::ClearChildren()
 								{
 									if (MapHelper.IsValidIndex(Index))
 									{
-										if (UActorComponent* Component = *reinterpret_cast<UActorComponent**>(MapHelper.GetValuePtr(Index)))
+										if (UObject* InstancedObject = *reinterpret_cast<UObject**>(MapHelper.GetValuePtr(Index)))
 										{
-											Component->Modify();
-											Component->Rename(nullptr, GetTransientPackage(), REN_DontCreateRedirectors);
+											InstancedObject->Modify();
+											InstancedObject->Rename(nullptr, GetTransientPackage(), REN_DontCreateRedirectors);
 										}
 										--ElementsToRemove;
 									}
@@ -1423,15 +1423,15 @@ void FPropertyValueImpl::DeleteChild( TSharedPtr<FPropertyNode> ChildNodeToDelet
 					{
 						FScriptArrayHelper ArrayHelper(ArrayProperty, Address);
 
-						// If the inner property is an instanced component property we must move the old component to the 
-						// transient package so resetting owned components on the parent doesn't find it
+						// If the inner property is an instanced property we must move the old object to the 
+						// transient package so code looking for objects of this type on the parent doesn't find it
 						FObjectProperty* InnerObjectProperty = CastField<FObjectProperty>(ArrayProperty->Inner);
-						if (InnerObjectProperty && InnerObjectProperty->HasAnyPropertyFlags(CPF_InstancedReference) && InnerObjectProperty->PropertyClass->IsChildOf(UActorComponent::StaticClass()))
+						if (InnerObjectProperty && InnerObjectProperty->HasAnyPropertyFlags(CPF_InstancedReference))
 						{
-							if (UActorComponent* Component = *reinterpret_cast<UActorComponent**>(ArrayHelper.GetRawPtr(ChildNodePtr->GetArrayIndex())))
+							if (UObject* InstancedObject = *reinterpret_cast<UObject**>(ArrayHelper.GetRawPtr(ChildNodePtr->GetArrayIndex())))
 							{
-								Component->Modify();
-								Component->Rename(nullptr, GetTransientPackage(), REN_DontCreateRedirectors);
+								InstancedObject->Modify();
+								InstancedObject->Rename(nullptr, GetTransientPackage(), REN_DontCreateRedirectors);
 							}
 						}
 
@@ -1441,15 +1441,15 @@ void FPropertyValueImpl::DeleteChild( TSharedPtr<FPropertyNode> ChildNodeToDelet
 					{
 						FScriptSetHelper SetHelper(SetProperty, Address);
 
-						// If the element property is an instanced component property we must move the old component to the 
-						// transient package so resetting owned components on the parent doesn't find it
+						// If the element property is an instanced property we must move the old object to the 
+						// transient package so code looking for objects of this type on the parent doesn't find it
 						FObjectProperty* ElementObjectProperty = CastField<FObjectProperty>(SetProperty->ElementProp);
-						if (ElementObjectProperty && ElementObjectProperty->HasAnyPropertyFlags(CPF_InstancedReference) && ElementObjectProperty->PropertyClass->IsChildOf(UActorComponent::StaticClass()))
+						if (ElementObjectProperty && ElementObjectProperty->HasAnyPropertyFlags(CPF_InstancedReference))
 						{
-							if (UActorComponent* Component = *reinterpret_cast<UActorComponent**>(SetHelper.GetElementPtr(ChildNodePtr->GetArrayIndex())))
+							if (UObject* InstancedObject = *reinterpret_cast<UObject**>(SetHelper.GetElementPtr(ChildNodePtr->GetArrayIndex())))
 							{
-								Component->Modify();
-								Component->Rename(nullptr, GetTransientPackage(), REN_DontCreateRedirectors);
+								InstancedObject->Modify();
+								InstancedObject->Rename(nullptr, GetTransientPackage(), REN_DontCreateRedirectors);
 							}
 						}
 
@@ -1461,15 +1461,15 @@ void FPropertyValueImpl::DeleteChild( TSharedPtr<FPropertyNode> ChildNodeToDelet
 					{
 						FScriptMapHelper MapHelper(MapProperty, Address);
 
-						// If the map's value property is an instanced component property we must move the old component to the 
-						// transient package so resetting owned components on the parent doesn't find it
+						// If the map's value property is an instanced property we must move the old object to the 
+						// transient package so code looking for objects of this type on the parent doesn't find it
 						FObjectProperty* ValueObjectProperty = CastField<FObjectProperty>(MapProperty->ValueProp);
-						if (ValueObjectProperty && ValueObjectProperty->HasAnyPropertyFlags(CPF_InstancedReference) && ValueObjectProperty->PropertyClass->IsChildOf(UActorComponent::StaticClass()))
+						if (ValueObjectProperty && ValueObjectProperty->HasAnyPropertyFlags(CPF_InstancedReference))
 						{
-							if (UActorComponent* Component = *reinterpret_cast<UActorComponent**>(MapHelper.GetValuePtr(ChildNodePtr->GetArrayIndex())))
+							if (UObject* InstancedObject = *reinterpret_cast<UObject**>(MapHelper.GetValuePtr(ChildNodePtr->GetArrayIndex())))
 							{
-								Component->Modify();
-								Component->Rename(nullptr, GetTransientPackage(), REN_DontCreateRedirectors);
+								InstancedObject->Modify();
+								InstancedObject->Rename(nullptr, GetTransientPackage(), REN_DontCreateRedirectors);
 							}
 						}
 
@@ -1561,21 +1561,21 @@ void FPropertyValueImpl::SwapChildren( TSharedPtr<FPropertyNode> FirstChildNode,
 				{
 					FScriptArrayHelper ArrayHelper(ArrayProperty, Address);
 
-					// If the inner property is an instanced component property we must move the old component to the 
-					// transient package so resetting owned components on the parent doesn't find it
+					// If the inner property is an instanced property we must move the old object to the 
+					// transient package so code looking for objects of this type on the parent doesn't find it
 					FObjectProperty* InnerObjectProperty = CastField<FObjectProperty>(ArrayProperty->Inner);
-					if (InnerObjectProperty && InnerObjectProperty->HasAnyPropertyFlags(CPF_InstancedReference) && InnerObjectProperty->PropertyClass->IsChildOf(UActorComponent::StaticClass()))
+					if (InnerObjectProperty && InnerObjectProperty->HasAnyPropertyFlags(CPF_InstancedReference))
 					{
-						if (UActorComponent* Component = *reinterpret_cast<UActorComponent**>(ArrayHelper.GetRawPtr(FirstIndex)))
+						if (UObject* InstancedObject = *reinterpret_cast<UObject**>(ArrayHelper.GetRawPtr(FirstIndex)))
 						{
-							Component->Modify();
-							Component->Rename(nullptr, GetTransientPackage(), REN_DontCreateRedirectors);
+							InstancedObject->Modify();
+							InstancedObject->Rename(nullptr, GetTransientPackage(), REN_DontCreateRedirectors);
 						}
 
-						if (UActorComponent* Component = *reinterpret_cast<UActorComponent**>(ArrayHelper.GetRawPtr(SecondIndex)))
+						if (UObject* InstancedObject = *reinterpret_cast<UObject**>(ArrayHelper.GetRawPtr(SecondIndex)))
 						{
-							Component->Modify();
-							Component->Rename(nullptr, GetTransientPackage(), REN_DontCreateRedirectors);
+							InstancedObject->Modify();
+							InstancedObject->Rename(nullptr, GetTransientPackage(), REN_DontCreateRedirectors);
 						}
 					}
 
