Property matrix now correctly displays EditCondition checkboxes only when applicable to the row. Property matrix rows also hide if the EditConditionHides metadata is set.

diff --git a/Engine/Source/Editor/PropertyEditor/Private/UserInterface/PropertyEditor/SPropertyEditorTableRow.cpp b/Engine/Source/Editor/PropertyEditor/Private/UserInterface/PropertyEditor/SPropertyEditorTableRow.cpp
--- a/Engine/Source/Editor/PropertyEditor/Private/UserInterface/PropertyEditor/SPropertyEditorTableRow.cpp
+++ b/Engine/Source/Editor/PropertyEditor/Private/UserInterface/PropertyEditor/SPropertyEditorTableRow.cpp
@@ -34,7 +34,7 @@ void SPropertyEditorTableRow::Construct( const FArguments& InArgs, const TShared
 	PropertyPath = FPropertyNode::CreatePropertyPath( PropertyEditor->GetPropertyNode() );
 
 	this->SetToolTipText(PropertyEditor->GetToolTipText());
-
+	this->SetVisibility(TAttribute<EVisibility>::CreateSP(this, &SPropertyEditorTableRow::OnGetRowVisibility));
 	SMultiColumnTableRow< TSharedPtr<FPropertyNode*> >::Construct( FSuperRowType::FArguments(), InOwnerTable );
 }
 
@@ -56,6 +56,16 @@ TSharedRef<SWidget> SPropertyEditorTableRow::GenerateWidgetForColumn( const FNam
 	return SNew(STextBlock).Text(NSLOCTEXT("PropertyEditor", "UnknownColumnId", "Unknown Column Id"));
 }
 
+EVisibility SPropertyEditorTableRow::OnGetRowVisibility() const
+{
+	if (PropertyEditor->IsOnlyVisibleWhenEditConditionMet() && !PropertyEditor->IsEditConditionMet())
+	{
+		return EVisibility::Collapsed;
+	}
+
+	return EVisibility::Visible;
+}
+
 TSharedRef< SWidget > SPropertyEditorTableRow::ConstructNameColumnWidget()
 {
 	TSharedRef< SHorizontalBox > NameColumnWidget =
@@ -75,7 +85,7 @@ TSharedRef< SWidget > SPropertyEditorTableRow::ConstructNameColumnWidget()
 			SNew( SEditConditionWidget )
 			.EditConditionValue(PropertyEditor.ToSharedRef(), &FPropertyEditor::IsEditConditionMet)
 			.OnEditConditionValueChanged_Lambda([this](bool bValue) { PropertyEditor->ToggleEditConditionState(); })
-			.Visibility_Lambda([this]() { return PropertyEditor->HasEditCondition() ? EVisibility::Visible : EVisibility::Collapsed; })
+			.Visibility_Lambda([this]() { return PropertyEditor->SupportsEditConditionToggle() ? EVisibility::Visible : EVisibility::Collapsed; })
 		]
 		+SHorizontalBox::Slot()
 		.AutoWidth()
@@ -195,17 +205,6 @@ const FSlateBrush* SPropertyEditorTableRow::OnGetFavoriteImage() const
 	return FEditorStyle::GetBrush(TEXT("PropertyWindow.Favorites_Disabled"));
 }
 
-void SPropertyEditorTableRow::OnEditConditionCheckChanged( ECheckBoxState CheckState )
-{
-	PropertyEditor->ToggleEditConditionState();
-}
-
-ECheckBoxState SPropertyEditorTableRow::OnGetEditConditionCheckState() const
-{
-	bool bEditConditionMet = PropertyEditor->IsEditConditionMet();
-	return bEditConditionMet ? ECheckBoxState::Checked : ECheckBoxState::Unchecked;
-}
-
 FReply SPropertyEditorTableRow::OnNameDoubleClicked()
 {
 	FReply Reply = FReply::Unhandled();
diff --git a/Engine/Source/Editor/PropertyEditor/Private/UserInterface/PropertyEditor/SPropertyEditorTableRow.h b/Engine/Source/Editor/PropertyEditor/Private/UserInterface/PropertyEditor/SPropertyEditorTableRow.h
--- a/Engine/Source/Editor/PropertyEditor/Private/UserInterface/PropertyEditor/SPropertyEditorTableRow.h
+++ b/Engine/Source/Editor/PropertyEditor/Private/UserInterface/PropertyEditor/SPropertyEditorTableRow.h
@@ -47,13 +47,12 @@ public:
 
 private:
 
+	EVisibility OnGetRowVisibility() const;
+
 	EVisibility OnGetFavoritesVisibility() const;
 	const FSlateBrush* OnGetFavoriteImage() const;
 	FReply OnToggleFavoriteClicked();
 
-	void OnEditConditionCheckChanged( ECheckBoxState CheckState );
-	ECheckBoxState OnGetEditConditionCheckState() const;
-
 	FReply OnNameDoubleClicked();
 
 	TSharedRef< SWidget > ConstructNameColumnWidget();
