- Fix crash when toggling a bool property

#rnx

diff --git a/Engine/Source/Editor/PropertyEditor/Private/PropertyRowGenerator.cpp b/Engine/Source/Editor/PropertyEditor/Private/PropertyRowGenerator.cpp
--- a/Engine/Source/Editor/PropertyEditor/Private/PropertyRowGenerator.cpp
+++ b/Engine/Source/Editor/PropertyEditor/Private/PropertyRowGenerator.cpp
@@ -70,7 +70,10 @@ public:
 
 	virtual void NotifyFinishedChangingProperties(const FPropertyChangedEvent& PropertyChangedEvent) override 
 	{
-		Generator->OnFinishedChangingProperties().Broadcast(PropertyChangedEvent);
+		if (Generator)
+		{
+			Generator->OnFinishedChangingProperties().Broadcast(PropertyChangedEvent);
+		}
 	}
 
 	virtual bool DontUpdateValueWhileEditing() const override { return false; }
