Fix ExtensionContent not working for FDetailWidgetRow  created via IDetailPropertyRow::CustomWidget()

diff --git a/Engine/Source/Editor/PropertyEditor/Private/DetailPropertyRow.cpp b/Engine/Source/Editor/PropertyEditor/Private/DetailPropertyRow.cpp
--- a/Engine/Source/Editor/PropertyEditor/Private/DetailPropertyRow.cpp
+++ b/Engine/Source/Editor/PropertyEditor/Private/DetailPropertyRow.cpp
@@ -843,6 +843,12 @@ void FDetailPropertyRow::MakeValueWidget( FDetailWidgetRow& Row, const TSharedPt
 		[
 			InCustomRow->ValueWidget.Widget
 		];
+
+		Row
+		.ExtensionContent()
+		[
+			InCustomRow->ExtensionWidget.Widget
+		];
 	}
 	else if (PropertyEditor.IsValid())
 	{
