Custom copy/paste actions now pulse the details row to indicate they were run

Previously the row only pulsed if it ran the built-in copy/paste routines

diff --git a/Engine/Source/Editor/PropertyEditor/Private/SDetailSingleItemRow.cpp b/Engine/Source/Editor/PropertyEditor/Private/SDetailSingleItemRow.cpp
--- a/Engine/Source/Editor/PropertyEditor/Private/SDetailSingleItemRow.cpp
+++ b/Engine/Source/Editor/PropertyEditor/Private/SDetailSingleItemRow.cpp
@@ -715,11 +715,19 @@ FReply SDetailSingleItemRow::OnMouseButtonUp(const FGeometry& MyGeometry, const
 		if (CopyAction.CanExecute() && MouseEvent.GetEffectingButton() == EKeys::RightMouseButton)
 		{
 			CopyAction.Execute();
+			PulseAnimation.Play(SharedThis(this));
 			bIsHandled = true;
 		}
 		else if (PasteAction.CanExecute() && MouseEvent.GetEffectingButton() == EKeys::LeftMouseButton)
 		{
 			PasteAction.Execute();
+			PulseAnimation.Play(SharedThis(this));
+			if(TSharedPtr<FPropertyNode> PropertyNode = GetPropertyNode())
+			{
+				// Mark property node as animating so we will animate after any potential re-construction
+				IDetailsViewPrivate* DetailsView = OwnerTreeNode.Pin()->GetDetailsView();
+				DetailsView->MarkNodeAnimating(PropertyNode, DetailWidgetConstants::PulseAnimationLength);
+			}
 			bIsHandled = true;
 		}
 
