Speculative fix for crash when an invalid FPropertyNode is passed into FDetailPropertyRow, causing it to store a null IPropertyHandle pointer that is never checked for validity.

This fix always initializes it to a valid pointer, but invalid handle, which is a state that is handled correctly in FPropertyHandleBase.

diff --git a/Engine/Source/Editor/PropertyEditor/Private/DetailPropertyRow.cpp b/Engine/Source/Editor/PropertyEditor/Private/DetailPropertyRow.cpp
--- a/Engine/Source/Editor/PropertyEditor/Private/DetailPropertyRow.cpp
+++ b/Engine/Source/Editor/PropertyEditor/Private/DetailPropertyRow.cpp
@@ -27,11 +27,11 @@ FDetailPropertyRow::FDetailPropertyRow(TSharedPtr<FPropertyNode> InPropertyNode,
 	, bForceAutoExpansion( false )
 	, bCachedCustomTypeInterface(false)
 {
-	if( InPropertyNode.IsValid() )
+	PropertyHandle = InParentCategory->GetParentLayoutImpl().GetPropertyHandle(PropertyNode);
+
+	if (PropertyNode.IsValid())
 	{
 		TSharedRef<FPropertyNode> PropertyNodeRef = PropertyNode.ToSharedRef();
-		
-		PropertyHandle = InParentCategory->GetParentLayoutImpl().GetPropertyHandle(PropertyNodeRef);
 
 		const TSharedRef<IPropertyUtilities> Utilities = InParentCategory->GetParentLayoutImpl().GetPropertyUtilities();
 
