- Updating detail category node generation to consider ShouldShowHiddenProperties flag before discarding a customized node

diff --git a/Engine/Source/Editor/PropertyEditor/Private/DetailCategoryBuilderImpl.cpp b/Engine/Source/Editor/PropertyEditor/Private/DetailCategoryBuilderImpl.cpp
--- a/Engine/Source/Editor/PropertyEditor/Private/DetailCategoryBuilderImpl.cpp
+++ b/Engine/Source/Editor/PropertyEditor/Private/DetailCategoryBuilderImpl.cpp
@@ -800,6 +800,17 @@ static bool IsCustomProperty(const TSharedPtr<FPropertyNode>& PropertyNode)
 	return bIsCustom;
 }
 
+static bool ShouldShowHidden(const TSharedPtr<FPropertyNode>& PropertyNode)
+{
+	// The property node has to be shown even if it's considered hidden if force flag is there
+	if (PropertyNode.IsValid())
+	{
+		return PropertyNode->HasNodeFlags(EPropertyNodeFlags::ShouldShowHiddenProperties);
+	}
+
+	return false;
+}
+
 static bool ShouldBeInlineNode(const TSharedRef<FDetailItemNode>& Node)
 {
 	TSharedPtr<FPropertyNode> PropertyNode = Node->GetPropertyNode();
@@ -888,7 +899,7 @@ void FDetailCategoryImpl::GenerateNodesFromCustomizations(const TArray<FDetailLa
 		if (Customization.IsValidCustomization())
 		{
 			// if a property is customized, skip the default customization
-			if (!IsCustomProperty(Customization.GetPropertyNode()) || Customization.bCustom || bFavoriteCategory)
+			if (!IsCustomProperty(Customization.GetPropertyNode()) || Customization.bCustom || bFavoriteCategory || ShouldShowHidden(Customization.GetPropertyNode()))
 			{
 				TSharedRef<FDetailItemNode> NewNode = MakeShareable(new FDetailItemNode(Customization, AsShared(), IsParentEnabled));
 				// Discard nodes that don't pass the property permission test. There is a special check here for properties in structs that do not have a
