Avoid filtering sub-categories that may yet contain properties

diff --git a/Engine/Source/Editor/PropertyEditor/Private/DetailCategoryBuilderImpl.cpp b/Engine/Source/Editor/PropertyEditor/Private/DetailCategoryBuilderImpl.cpp
--- a/Engine/Source/Editor/PropertyEditor/Private/DetailCategoryBuilderImpl.cpp
+++ b/Engine/Source/Editor/PropertyEditor/Private/DetailCategoryBuilderImpl.cpp
@@ -904,7 +904,7 @@ void FDetailCategoryImpl::GenerateNodesFromCustomizations(const TArray<FDetailLa
 				TSharedRef<FDetailItemNode> NewNode = MakeShareable(new FDetailItemNode(Customization, AsShared(), IsParentEnabled));
 				// Discard nodes that don't pass the property permission test. There is a special check here for properties in structs that do not have a
 				// parent struct node (eg ShowOnlyInnerProperties). Both the node and it's container must pass the filter.
-				if (FPropertyEditorPermissionList::Get().IsEnabled())
+				if (FPropertyEditorPermissionList::Get().IsEnabled() && (NewNode->GetNodeType() == EDetailNodeType::Object || NewNode->GetNodeType() == EDetailNodeType::Item))
 				{
 					if (!FPropertyEditorPermissionList::Get().DoesPropertyPassFilter(NewNode->GetParentBaseStructure(), NewNode->GetNodeName()) || !PassesInlineFilters(NewNode.Get()))
 					{
