PropertyEditor - Expose an api on the property row generator to invalidate cached internal data so that visual state can be refreshed when the displayed object is modified from non property editor code.

diff --git a/Engine/Source/Editor/PropertyEditor/Private/PropertyRowGenerator.cpp b/Engine/Source/Editor/PropertyEditor/Private/PropertyRowGenerator.cpp
--- a/Engine/Source/Editor/PropertyEditor/Private/PropertyRowGenerator.cpp
+++ b/Engine/Source/Editor/PropertyEditor/Private/PropertyRowGenerator.cpp
@@ -332,6 +332,14 @@ void FPropertyRowGenerator::UnregisterInstancedCustomPropertyTypeLayout(FName Pr
 	}
 }
 
+void FPropertyRowGenerator::InvalidateCachedState()
+{
+	for (const TSharedPtr<FComplexPropertyNode>& ComplexRootNode : RootPropertyNodes)
+	{
+		ComplexRootNode->InvalidateCachedState();
+	}
+}
+
 void FPropertyRowGenerator::Tick(float DeltaTime)
 {
 	for (TSharedPtr<IDetailCustomization>& Customization : CustomizationClassInstancesPendingDelete)
diff --git a/Engine/Source/Editor/PropertyEditor/Private/PropertyRowGenerator.h b/Engine/Source/Editor/PropertyEditor/Private/PropertyRowGenerator.h
--- a/Engine/Source/Editor/PropertyEditor/Private/PropertyRowGenerator.h
+++ b/Engine/Source/Editor/PropertyEditor/Private/PropertyRowGenerator.h
@@ -95,6 +95,7 @@ public:
 	{
 		CustomValidatePropertyNodesFunction = MoveTemp(InCustomValidatePropertyNodesFunction);
 	}
+	virtual void InvalidateCachedState() override;
 
 	/** FTickableEditorObject interface */
 	virtual ETickableTickType GetTickableTickType() const override { return ETickableTickType::Always; }
diff --git a/Engine/Source/Editor/PropertyEditor/Public/IPropertyRowGenerator.h b/Engine/Source/Editor/PropertyEditor/Public/IPropertyRowGenerator.h
--- a/Engine/Source/Editor/PropertyEditor/Public/IPropertyRowGenerator.h
+++ b/Engine/Source/Editor/PropertyEditor/Public/IPropertyRowGenerator.h
@@ -111,4 +111,7 @@ public:
 	/* Use this function to set a callback on FPropertyRowGenerator that will override the ValidatePropertyNodes function.
 	 * This is useful if your implementation doesn't need to validate nodes every tick or needs to perform some other form of validation. */
 	virtual void SetCustomValidatePropertyNodesFunction(FOnValidatePropertyRowGeneratorNodes InCustomValidatePropertyNodesFunction) = 0;
+
+	/* Invalidates internal cached state.  Common use of this API is to synchronize the viewed object with changes made by external code. */
+	virtual void InvalidateCachedState() = 0;
 };
\ No newline at end of file
