Fixed removing an instanced element from a map or set not using the internal index for object lookup.

diff --git a/Engine/Source/Editor/PropertyEditor/Private/PropertyHandleImpl.cpp b/Engine/Source/Editor/PropertyEditor/Private/PropertyHandleImpl.cpp
--- a/Engine/Source/Editor/PropertyEditor/Private/PropertyHandleImpl.cpp
+++ b/Engine/Source/Editor/PropertyEditor/Private/PropertyHandleImpl.cpp
@@ -1440,40 +1440,40 @@ void FPropertyValueImpl::DeleteChild( TSharedPtr<FPropertyNode> ChildNodeToDelet
 					else if (SetProperty)
 					{
 						FScriptSetHelper SetHelper(SetProperty, Address);
-
+						int32 InternalIndex = SetHelper.FindInternalIndex(ChildNodePtr->GetArrayIndex());
+						
 						// If the element property is an instanced property we must move the old object to the 
 						// transient package so code looking for objects of this type on the parent doesn't find it
 						FObjectProperty* ElementObjectProperty = CastField<FObjectProperty>(SetProperty->ElementProp);
 						if (ElementObjectProperty && ElementObjectProperty->HasAnyPropertyFlags(CPF_InstancedReference))
 						{
-							if (UObject* InstancedObject = *reinterpret_cast<UObject**>(SetHelper.GetElementPtr(ChildNodePtr->GetArrayIndex())))
+							if (UObject* InstancedObject = *reinterpret_cast<UObject**>(SetHelper.GetElementPtr(InternalIndex)))
 							{
 								InstancedObject->Modify();
 								InstancedObject->Rename(nullptr, GetTransientPackage(), REN_DontCreateRedirectors);
 							}
 						}
 
-						int32 InternalIndex = SetHelper.FindInternalIndex(ChildNodePtr->GetArrayIndex());
 						SetHelper.RemoveAt(InternalIndex);
 						SetHelper.Rehash();
 					}
 					else if (MapProperty)
 					{
 						FScriptMapHelper MapHelper(MapProperty, Address);
-
+						int32 InternalIndex = MapHelper.FindInternalIndex(ChildNodePtr->GetArrayIndex());
+						
 						// If the map's value property is an instanced property we must move the old object to the 
 						// transient package so code looking for objects of this type on the parent doesn't find it
 						FObjectProperty* ValueObjectProperty = CastField<FObjectProperty>(MapProperty->ValueProp);
 						if (ValueObjectProperty && ValueObjectProperty->HasAnyPropertyFlags(CPF_InstancedReference))
 						{
-							if (UObject* InstancedObject = *reinterpret_cast<UObject**>(MapHelper.GetValuePtr(ChildNodePtr->GetArrayIndex())))
+							if (UObject* InstancedObject = *reinterpret_cast<UObject**>(MapHelper.GetValuePtr(InternalIndex)))
 							{
 								InstancedObject->Modify();
 								InstancedObject->Rename(nullptr, GetTransientPackage(), REN_DontCreateRedirectors);
 							}
 						}
 
-						int32 InternalIndex = MapHelper.FindInternalIndex(ChildNodePtr->GetArrayIndex());
 						MapHelper.RemoveAt(InternalIndex);
 						MapHelper.Rehash();
 					}
