Added a bit more information to what failed when an edit condition cannot be parsed, indicating whether it was an issue with the type or the value

diff --git a/Engine/Source/Editor/PropertyEditor/Private/EditConditionParser.cpp b/Engine/Source/Editor/PropertyEditor/Private/EditConditionParser.cpp
--- a/Engine/Source/Editor/PropertyEditor/Private/EditConditionParser.cpp
+++ b/Engine/Source/Editor/PropertyEditor/Private/EditConditionParser.cpp
@@ -579,7 +579,7 @@ static FExpressionResult EnumPropertyEquals(const EditConditionParserTokens::FEn
 
 	if (!TypeName.IsSet())
 	{
-		return MakeError(FText::Format(LOCTEXT("InvalidOperand", "EditCondition attempted to use an invalid operand \"{0}\"."), FText::FromString(Property.PropertyName)));
+		return MakeError(FText::Format(LOCTEXT("InvalidOperand_Type", "EditCondition attempted to use an invalid operand \"{0}\" (type error)."), FText::FromString(Property.PropertyName)));
 	}
 
 	if (TypeName.GetValue() != Enum.Type)
@@ -591,7 +591,7 @@ static FExpressionResult EnumPropertyEquals(const EditConditionParserTokens::FEn
 
 	if (!ValueProp.IsSet())
 	{
-		return MakeError(FText::Format(LOCTEXT("InvalidOperand", "EditCondition attempted to use an invalid operand \"{0}\"."), FText::FromString(Property.PropertyName)));
+		return MakeError(FText::Format(LOCTEXT("InvalidOperand_Value", "EditCondition attempted to use an invalid operand \"{0}\" (value error)."), FText::FromString(Property.PropertyName)));
 	}
 
 	bool bEqual = ValueProp.GetValue() == Enum.Value;
