- Reverting show hidden properties bypassing customized struct

diff --git a/Engine/Source/Editor/PropertyEditor/Private/DetailCategoryBuilderImpl.cpp b/Engine/Source/Editor/PropertyEditor/Private/DetailCategoryBuilderImpl.cpp
--- a/Engine/Source/Editor/PropertyEditor/Private/DetailCategoryBuilderImpl.cpp
+++ b/Engine/Source/Editor/PropertyEditor/Private/DetailCategoryBuilderImpl.cpp
@@ -800,17 +800,6 @@ static bool IsCustomProperty(const TSharedPtr<FPropertyNode>& PropertyNode)
 	return bIsCustom;
 }
 
-static bool ShouldShowHidden(const TSharedPtr<FPropertyNode>& PropertyNode)
-{
-	// The property node has to be shown even if it's considered hidden if force flag is there
-	if (PropertyNode.IsValid())
-	{
-		return PropertyNode->HasNodeFlags(EPropertyNodeFlags::ShouldShowHiddenProperties);
-	}
-
-	return false;
-}
-
 static bool ShouldBeInlineNode(const TSharedRef<FDetailItemNode>& Node)
 {
 	TSharedPtr<FPropertyNode> PropertyNode = Node->GetPropertyNode();
@@ -899,7 +888,7 @@ void FDetailCategoryImpl::GenerateNodesFromCustomizations(const TArray<FDetailLa
 		if (Customization.IsValidCustomization())
 		{
 			// if a property is customized, skip the default customization
-			if (!IsCustomProperty(Customization.GetPropertyNode()) || Customization.bCustom || bFavoriteCategory || ShouldShowHidden(Customization.GetPropertyNode()))
+			if (!IsCustomProperty(Customization.GetPropertyNode()) || Customization.bCustom || bFavoriteCategory)
 			{
 				TSharedRef<FDetailItemNode> NewNode = MakeShareable(new FDetailItemNode(Customization, AsShared(), IsParentEnabled));
 				// Discard nodes that don't pass the property permission test. There is a special check here for properties in structs that do not have a
