FPropertyHandleObject::SetValue when called with an FAssetData no longer always marks an asset as dirty if the new value is identical to the previous value.

This makes an assumption the FObjectPropertyBase::GetExportPath() will continue to return the value in 'ClassName'ObjectPath format, but that seems very unlikely to change and is depended on elsewhere so this can be considered an acceptable risk.

Suggested by UDN: https://epicgames.lightning.force.com/lightning/r/Case/5004z00001ay3pbAAA/view

diff --git a/Engine/Source/Editor/PropertyEditor/Private/PropertyHandleImpl.cpp b/Engine/Source/Editor/PropertyEditor/Private/PropertyHandleImpl.cpp
--- a/Engine/Source/Editor/PropertyEditor/Private/PropertyHandleImpl.cpp
+++ b/Engine/Source/Editor/PropertyEditor/Private/PropertyHandleImpl.cpp
@@ -3806,7 +3806,7 @@ FPropertyAccess::Result FPropertyHandleObject::SetValue(const FAssetData& NewVal
 			NewValue.GetAsset();
 		}
 
-		FString ObjectPathName = NewValue.IsValid() ? NewValue.ObjectPath.ToString() : TEXT("None");
+		FString ObjectPathName = NewValue.IsValid() ? FString::Printf(TEXT("%s'%s'"), *NewValue.GetClass()->GetName(), *NewValue.ObjectPath.ToString()) : TEXT("None");
 		return SetValueFromFormattedString(ObjectPathName, Flags, bSkipResolve);
 	}
 
