Subcategory nodes now hide themselves if no children are visible. This was most evident with eg. EditInstanceOnly metadata set for a property while editing the CDO - the property would be hidden, but the subcategory would still be visible.

diff --git a/Engine/Source/Editor/PropertyEditor/Private/DetailItemNode.cpp b/Engine/Source/Editor/PropertyEditor/Private/DetailItemNode.cpp
--- a/Engine/Source/Editor/PropertyEditor/Private/DetailItemNode.cpp
+++ b/Engine/Source/Editor/PropertyEditor/Private/DetailItemNode.cpp
@@ -863,4 +863,24 @@ void FDetailItemNode::FilterNode(const FDetailFilter& InFilter)
 			}
 		}
 	}
+
+	// if this is a subcategory, it should only be visible if one or more of its children is visible
+	if (Customization.HasPropertyNode() &&
+		Customization.GetPropertyNode()->AsCategoryNode() && 
+		bShouldBeVisibleDueToFiltering)
+	{
+		bool bAnyChildVisible = false;
+		for (const TSharedRef<FDetailTreeNode>& Child : Children)
+		{
+			if (Child->GetVisibility() == ENodeVisibility::Visible)
+			{
+				bAnyChildVisible = true;
+				break;
+			}
+		}
+		if (!bAnyChildVisible)
+		{
+			bShouldBeVisibleDueToFiltering = false;
+		}
+	}
 }
