InvalidEnumValues metadata support (inverse of ValidEnumValues)

diff --git a/Engine/Source/Editor/PropertyEditor/Private/PropertyEditorHelpers.cpp b/Engine/Source/Editor/PropertyEditor/Private/PropertyEditorHelpers.cpp
--- a/Engine/Source/Editor/PropertyEditor/Private/PropertyEditorHelpers.cpp
+++ b/Engine/Source/Editor/PropertyEditor/Private/PropertyEditorHelpers.cpp
@@ -1085,7 +1085,7 @@ namespace PropertyEditorHelpers
 			TArray<FString> ValidEnumValuesAsString;
 
 			Property->GetMetaData(ValidEnumValuesName).ParseIntoArray(ValidEnumValuesAsString, TEXT(","));
-			for(auto& Value : ValidEnumValuesAsString)
+			for(FString& Value : ValidEnumValuesAsString)
 			{
 				Value.TrimStartInline();
 				ValidEnumValues.Add(*InEnum->GenerateFullEnumName(*Value));
@@ -1095,6 +1095,26 @@ namespace PropertyEditorHelpers
 		return ValidEnumValues;
 	}
 
+	TArray<FName> GetInvalidEnumsFromPropertyOverride(const FProperty* Property, const UEnum* InEnum)
+	{
+		TArray<FName> InvalidEnumValues;
+
+		static const FName InvalidEnumValuesName("InvalidEnumValues");
+		if(Property->HasMetaData(InvalidEnumValuesName))
+		{
+			TArray<FString> InvalidEnumValuesAsString;
+
+			Property->GetMetaData(InvalidEnumValuesName).ParseIntoArray(InvalidEnumValuesAsString, TEXT(","));
+			for(FString& Value : InvalidEnumValuesAsString)
+			{
+				Value.TrimStartInline();
+				InvalidEnumValues.Add(*InEnum->GenerateFullEnumName(*Value));
+			}
+		}
+
+		return InvalidEnumValues;
+	}
+
 	bool IsCategoryHiddenByClass(const TSharedPtr<FComplexPropertyNode>& InRootNode, FName CategoryName)
 	{
 		return InRootNode->AsObjectNode() && InRootNode->AsObjectNode()->GetHiddenCategories().Contains(CategoryName);
diff --git a/Engine/Source/Editor/PropertyEditor/Private/PropertyEditorHelpers.h b/Engine/Source/Editor/PropertyEditor/Private/PropertyEditorHelpers.h
--- a/Engine/Source/Editor/PropertyEditor/Private/PropertyEditorHelpers.h
+++ b/Engine/Source/Editor/PropertyEditor/Private/PropertyEditorHelpers.h
@@ -244,9 +244,19 @@ namespace PropertyEditorHelpers
 	 *
 	 * @param Property	The property which may contain the "ValidEnumValues" metadata
 	 * @param InEnum	The enum to search
-	 * @return The array of allowed enums.  NOTE: If an empty array is returned all enum values are allowed.  It is an error for a property to hide all enum values so that state is undefined here.
+	 * @return The array of allowed enums.  NOTE: If an empty array is returned all enum values are allowed. It is an error for a property to hide all enum values so that state is undefined here.
 	 */
 	TArray<FName> GetValidEnumsFromPropertyOverride(const FProperty* Property, const UEnum* InEnum);
+
+	/**
+	 * Returns any enums that are explicitly Disallowed by the "InvalidEnumValues" metadata on FProperty using the specified enum.
+	 *
+	 * @param Property	The property which may contain the "InvalidEnumValues" metadata
+	 * @param InEnum	The enum to search
+	 * @return The array of disallowed enums.
+	 */
+	TArray<FName> GetInvalidEnumsFromPropertyOverride(const FProperty* Property, const UEnum* InEnum);
+	
 	/**
 	 * Whether or not a category is hidden by a given root object
 	 * @param InRootNode	The root node that for the objects we are customizing
diff --git a/Engine/Source/Editor/PropertyEditor/Private/UserInterface/PropertyEditor/SPropertyEditorCombo.cpp b/Engine/Source/Editor/PropertyEditor/Private/UserInterface/PropertyEditor/SPropertyEditorCombo.cpp
--- a/Engine/Source/Editor/PropertyEditor/Private/UserInterface/PropertyEditor/SPropertyEditorCombo.cpp
+++ b/Engine/Source/Editor/PropertyEditor/Private/UserInterface/PropertyEditor/SPropertyEditorCombo.cpp
@@ -138,6 +138,7 @@ void SPropertyEditorCombo::GenerateComboBoxStrings( TArray< TSharedPtr<FString>
 	}
 
 	TArray<FText> BasicTooltips;
+
 	bUsesAlternateDisplayValues = ComboArgs.PropertyHandle->GeneratePossibleValues(OutComboBoxStrings, BasicTooltips, OutRestrictedItems);
 
 	// For enums, look for rich tooltip information
@@ -159,7 +160,8 @@ void SPropertyEditorCombo::GenerateComboBoxStrings( TArray< TSharedPtr<FString>
 			if(Enum)
 			{
 				TArray<FName> AllowedPropertyEnums = PropertyEditorHelpers::GetValidEnumsFromPropertyOverride(Property, Enum);
-
+				TArray<FName> DisallowedPropertyEnums = PropertyEditorHelpers::GetInvalidEnumsFromPropertyOverride(Property, Enum);
+				
 				// Get enum doc link (not just GetDocumentationLink as that is the documentation for the struct we're in, not the enum documentation)
 				FString DocLink = PropertyEditorHelpers::GetEnumDocumentationLink(Property);
 
@@ -168,9 +170,17 @@ void SPropertyEditorCombo::GenerateComboBoxStrings( TArray< TSharedPtr<FString>
 					FString Excerpt = Enum->GetNameStringByIndex(EnumIdx);
 
 					bool bShouldBeHidden = Enum->HasMetaData(TEXT("Hidden"), EnumIdx) || Enum->HasMetaData(TEXT("Spacer"), EnumIdx);
-					if(!bShouldBeHidden && AllowedPropertyEnums.Num() != 0)
+					if(!bShouldBeHidden)
 					{
-						bShouldBeHidden = AllowedPropertyEnums.Find(Enum->GetNameByIndex(EnumIdx)) == INDEX_NONE;
+						if(AllowedPropertyEnums.Num() > 0)
+						{
+							bShouldBeHidden = AllowedPropertyEnums.Find(Enum->GetNameByIndex(EnumIdx)) == INDEX_NONE;
+						}
+						// If both are specified, InvalidEnumValues takes precedence
+						else if(DisallowedPropertyEnums.Num() > 0)
+						{
+							bShouldBeHidden = DisallowedPropertyEnums.Find(Enum->GetNameByIndex(EnumIdx)) != INDEX_NONE;
+						}
 					}
 
 					if (!bShouldBeHidden)
diff --git a/Engine/Source/Editor/PropertyEditor/Private/UserInterface/PropertyEditor/SPropertyEditorNumeric.h b/Engine/Source/Editor/PropertyEditor/Private/UserInterface/PropertyEditor/SPropertyEditorNumeric.h
--- a/Engine/Source/Editor/PropertyEditor/Private/UserInterface/PropertyEditor/SPropertyEditorNumeric.h
+++ b/Engine/Source/Editor/PropertyEditor/Private/UserInterface/PropertyEditor/SPropertyEditorNumeric.h
@@ -90,14 +90,23 @@ public:
 					};
 
 					const TArray<FName> AllowedPropertyEnums = PropertyEditorHelpers::GetValidEnumsFromPropertyOverride(Prop, BitmaskEnum);
+					const TArray<FName> DisallowedPropertyEnums = PropertyEditorHelpers::GetInvalidEnumsFromPropertyOverride(Prop, BitmaskEnum);
 					// Note: This loop doesn't include (BitflagsEnum->NumEnums() - 1) in order to skip the implicit "MAX" value that gets added to the enum type at compile time.
 					for (int32 BitmaskEnumIndex = 0; BitmaskEnumIndex < BitmaskEnum->NumEnums() - 1; ++BitmaskEnumIndex)
 					{
 						const int64 EnumValue = BitmaskEnum->GetValueByIndex(BitmaskEnumIndex);
 						bool bShouldBeHidden = BitmaskEnum->HasMetaData(TEXT("Hidden"), BitmaskEnumIndex);
-						if (!bShouldBeHidden && AllowedPropertyEnums.Num() > 0)
+						if (!bShouldBeHidden)
 						{
-							bShouldBeHidden = AllowedPropertyEnums.Find(BitmaskEnum->GetNameByIndex(BitmaskEnumIndex)) == INDEX_NONE;
+							if(AllowedPropertyEnums.Num() > 0)
+							{
+								bShouldBeHidden = AllowedPropertyEnums.Find(BitmaskEnum->GetNameByIndex(BitmaskEnumIndex)) == INDEX_NONE;
+							}
+							// If both are specified, InvalidEnumValues takes precedence
+							else if(DisallowedPropertyEnums.Num() > 0)
+							{
+								bShouldBeHidden = DisallowedPropertyEnums.Find(BitmaskEnum->GetNameByIndex(BitmaskEnumIndex)) != INDEX_NONE;
+							}
 						}
 						if (EnumValue >= 0 && !bShouldBeHidden)
 						{
