The Reset to Default button is no longer hidden when a details row is disabled, but is rather shown greyed out.

diff --git a/Engine/Source/Editor/PropertyEditor/Private/PropertyNode.cpp b/Engine/Source/Editor/PropertyEditor/Private/PropertyNode.cpp
--- a/Engine/Source/Editor/PropertyEditor/Private/PropertyNode.cpp
+++ b/Engine/Source/Editor/PropertyEditor/Private/PropertyNode.cpp
@@ -2014,7 +2014,7 @@ bool FPropertyNode::GetDiffersFromDefault()
 		}
 
 		FObjectPropertyNode* ObjectNode = FindObjectItemParent();
-		if (ObjectNode && Property.IsValid() && !IsEditConst())
+		if (ObjectNode && Property.IsValid())
 		{
 			// Get an iterator for the enclosing objects.
 			for(int32 ObjIndex = 0; ObjIndex < ObjectNode->GetNumObjects(); ++ObjIndex)
diff --git a/Engine/Source/Editor/PropertyEditor/Private/SDetailSingleItemRow.cpp b/Engine/Source/Editor/PropertyEditor/Private/SDetailSingleItemRow.cpp
--- a/Engine/Source/Editor/PropertyEditor/Private/SDetailSingleItemRow.cpp
+++ b/Engine/Source/Editor/PropertyEditor/Private/SDetailSingleItemRow.cpp
@@ -16,7 +16,6 @@
 #include "SConstrainedBox.h"
 #include "SDetailExpanderArrow.h"
 #include "SDetailRowIndent.h"
-#include "SResetToDefaultPropertyEditor.h"
 
 #include "HAL/PlatformApplicationMisc.h"
 #include "Modules/ModuleInterface.h"
@@ -290,28 +289,28 @@ TSharedPtr<IPropertyHandle> SDetailSingleItemRow::GetPropertyHandle() const
 	return Handle;
 }
 
-bool SDetailSingleItemRow::UpdateResetToDefault()
+void SDetailSingleItemRow::UpdateResetToDefault()
 {
+	bCachedResetToDefaultVisible = false;
+
 	TSharedPtr<IPropertyHandle> PropertyHandle = GetPropertyHandle();
 	if (PropertyHandle.IsValid())
 	{
 		if (PropertyHandle->HasMetaData("NoResetToDefault") || PropertyHandle->GetInstanceMetaData("NoResetToDefault"))
 		{
-			return false;
+			bCachedResetToDefaultVisible = false;
+			return;
 		}
 	}
 
 	if (WidgetRow.CustomResetToDefault.IsSet())
 	{
-		return WidgetRow.CustomResetToDefault.GetValue().IsResetToDefaultVisible(PropertyHandle);
+		bCachedResetToDefaultVisible = WidgetRow.CustomResetToDefault.GetValue().IsResetToDefaultVisible(PropertyHandle);
 	}
 	else if (PropertyHandle.IsValid())
 	{
-		return PropertyHandle->CanResetToDefault();
+		bCachedResetToDefaultVisible = PropertyHandle->CanResetToDefault();
 	}
-
-	return false;
-
 }
 
 void SDetailSingleItemRow::Construct( const FArguments& InArgs, FDetailLayoutCustomization* InCustomization, bool bHasMultipleColumns, TSharedRef<FDetailTreeNode> InOwnerTreeNode, const TSharedRef<STableViewBase>& InOwnerTableView )
@@ -577,7 +576,10 @@ void SDetailSingleItemRow::Construct( const FArguments& InArgs, FDetailLayoutCus
 			ResetToDefault.Label = NSLOCTEXT("PropertyEditor", "ResetToDefault", "Reset to Default");
 			ResetToDefault.UIAction = FUIAction(
 				FExecuteAction::CreateSP(this, &SDetailSingleItemRow::OnResetToDefaultClicked),
-				FCanExecuteAction::CreateSP(this, &SDetailSingleItemRow::IsResetToDefaultEnabled)
+				FCanExecuteAction::CreateLambda([this, IsValueEnabledAttribute]()
+					{
+						return IsResetToDefaultVisible() && IsValueEnabledAttribute.Get(true);
+					})
 			);
 
 			// We could just collapse the Reset to Default button by setting the FIsActionButtonVisible delegate,
@@ -587,13 +589,14 @@ void SDetailSingleItemRow::Construct( const FArguments& InArgs, FDetailLayoutCus
 			static FSlateIcon DisabledResetToDefaultIcon(FAppStyle::Get().GetStyleSetName(), "NoBrush");
 			ResetToDefault.Icon = TAttribute<FSlateIcon>::Create([this]()
 			{
-				return IsResetToDefaultEnabled() ? 
+				return IsResetToDefaultVisible() ?
 					EnabledResetToDefaultIcon :
 					DisabledResetToDefaultIcon;
 			});
+
 			ResetToDefault.ToolTip = TAttribute<FText>::Create([this]() 
 			{
-				return IsResetToDefaultEnabled() ?
+				return IsResetToDefaultVisible() ?
 					NSLOCTEXT("PropertyEditor", "ResetToDefaultToolTip", "Reset this property to its default value.") :
 					FText::GetEmpty();
 			});
@@ -740,9 +743,9 @@ FReply SDetailSingleItemRow::OnMouseButtonUp(const FGeometry& MyGeometry, const
 	return SDetailTableRowBase::OnMouseButtonUp(MyGeometry, MouseEvent);
 }
 
-bool SDetailSingleItemRow::IsResetToDefaultEnabled() const
+bool SDetailSingleItemRow::IsResetToDefaultVisible() const
 {
-	return bCachedResetToDefaultEnabled;
+	return bCachedResetToDefaultVisible;
 }
 
 void SDetailSingleItemRow::OnResetToDefaultClicked() const
@@ -1254,7 +1257,7 @@ TSharedPtr<FDragDropOperation> SDetailSingleItemRow::CreateDragDropOperation()
 
 void SDetailSingleItemRow::Tick(const FGeometry& AllottedGeometry, const double InCurrentTime, const float InDeltaTime)
 {
-	bCachedResetToDefaultEnabled = UpdateResetToDefault();
+	UpdateResetToDefault();
 }
 
 void SArrayRowHandle::Construct(const FArguments& InArgs)
diff --git a/Engine/Source/Editor/PropertyEditor/Private/SDetailSingleItemRow.h b/Engine/Source/Editor/PropertyEditor/Private/SDetailSingleItemRow.h
--- a/Engine/Source/Editor/PropertyEditor/Private/SDetailSingleItemRow.h
+++ b/Engine/Source/Editor/PropertyEditor/Private/SDetailSingleItemRow.h
@@ -79,8 +79,10 @@ private:
 
 	void CreateGlobalExtensionWidgets(TArray<FPropertyRowExtensionButton>& ExtensionButtons) const;
 	void PopulateExtensionWidget();
+
+	void UpdateResetToDefault();
 	void OnResetToDefaultClicked() const;
-	bool IsResetToDefaultEnabled() const;
+	bool IsResetToDefaultVisible() const;
 
 	bool IsHighlighted() const;
 
@@ -111,13 +113,12 @@ private:
 	TSharedPtr<FPropertyNode> GetPropertyNode() const;
 	TSharedPtr<IPropertyHandle> GetPropertyHandle() const;
 
-	bool UpdateResetToDefault();
 private:
 	/** Customization for this widget */
 	FDetailLayoutCustomization* Customization;
 	FDetailWidgetRow WidgetRow;
 	bool bAllowFavoriteSystem;
-	bool bCachedResetToDefaultEnabled;
+	bool bCachedResetToDefaultVisible;
 	TSharedPtr<FPropertyNode> SwappablePropertyNode;
 	TSharedPtr<SButton> ExpanderArrow;
 	TWeakPtr<FDragDropOperation> DragOperation; // last drag initiated by this widget
