InvalidEnumValues metadata support in PropertyHandle
(should have been part of 18805091)

diff --git a/Engine/Source/Editor/PropertyEditor/Private/PropertyHandleImpl.cpp b/Engine/Source/Editor/PropertyEditor/Private/PropertyHandleImpl.cpp
--- a/Engine/Source/Editor/PropertyEditor/Private/PropertyHandleImpl.cpp
+++ b/Engine/Source/Editor/PropertyEditor/Private/PropertyHandleImpl.cpp
@@ -2820,15 +2820,24 @@ bool FPropertyHandleBase::GeneratePossibleValues(TArray< TSharedPtr<FString> >&
 	if( Enum )
 	{
 		const TArray<FName> ValidEnumValues = PropertyEditorHelpers::GetValidEnumsFromPropertyOverride(Property, Enum);
+		const TArray<FName> InvalidEnumValues = PropertyEditorHelpers::GetInvalidEnumsFromPropertyOverride(Property, Enum);
 
 		//NumEnums() - 1, because the last item in an enum is the _MAX item
 		for( int32 EnumIndex = 0; EnumIndex < Enum->NumEnums() - 1; ++EnumIndex )
 		{
 			// Ignore hidden enums
 			bool bShouldBeHidden = Enum->HasMetaData(TEXT("Hidden"), EnumIndex ) || Enum->HasMetaData(TEXT("Spacer"), EnumIndex );
-			if (!bShouldBeHidden && ValidEnumValues.Num() != 0)
+			if (!bShouldBeHidden)
 			{
-				bShouldBeHidden = ValidEnumValues.Find(Enum->GetNameByIndex(EnumIndex)) == INDEX_NONE;
+				if(ValidEnumValues.Num() > 0)
+				{
+					bShouldBeHidden = ValidEnumValues.Find(Enum->GetNameByIndex(EnumIndex)) == INDEX_NONE;
+				}
+				// If both are specified, InvalidEnumValues takes precedence
+				else if(InvalidEnumValues.Num() > 0)
+				{
+					bShouldBeHidden = InvalidEnumValues.Find(Enum->GetNameByIndex(EnumIndex)) != INDEX_NONE;
+				}
 			}
 
 			if (!bShouldBeHidden)
