Property NoClear specifier propagates to item object pickers

diff --git a/Engine/Source/Editor/PropertyEditor/Private/UserInterface/PropertyEditor/SPropertyEditorEditInline.cpp b/Engine/Source/Editor/PropertyEditor/Private/UserInterface/PropertyEditor/SPropertyEditorEditInline.cpp
--- a/Engine/Source/Editor/PropertyEditor/Private/UserInterface/PropertyEditor/SPropertyEditorEditInline.cpp
+++ b/Engine/Source/Editor/PropertyEditor/Private/UserInterface/PropertyEditor/SPropertyEditorEditInline.cpp
@@ -178,7 +178,16 @@ TSharedRef<SWidget> SPropertyEditorEditInline::GenerateClassPicker()
 	FProperty* Property = PropertyNode->GetProperty();
 	ClassFilter->ObjProperty = CastField<FObjectPropertyBase>( Property );
 	ClassFilter->IntProperty = CastField<FInterfaceProperty>( Property );
-	Options.bShowNoneOption = !(Property->PropertyFlags & CPF_NoClear);
+	
+	bool bContainerHasNoClear = false;
+	if(PropertyNode->GetArrayIndex() != INDEX_NONE)
+	{
+		if(const TSharedPtr<FPropertyNode>& ParentNode = PropertyNode->GetParentNodeSharedPtr())
+		{
+			bContainerHasNoClear = ParentNode->GetProperty()->HasAllPropertyFlags(CPF_NoClear);
+		}
+	}
+	Options.bShowNoneOption = !Property->HasAllPropertyFlags(CPF_NoClear) && !bContainerHasNoClear;
 
 	FObjectPropertyNode* ObjectPropertyNode = PropertyNode->FindObjectItemParent();
 	if( ObjectPropertyNode )
