Scrolling a row outside of the Details view now commits any changes that were being made to that row.

diff --git a/Engine/Source/Editor/PropertyEditor/Private/SDetailsView.cpp b/Engine/Source/Editor/PropertyEditor/Private/SDetailsView.cpp
--- a/Engine/Source/Editor/PropertyEditor/Private/SDetailsView.cpp
+++ b/Engine/Source/Editor/PropertyEditor/Private/SDetailsView.cpp
@@ -430,6 +430,7 @@ TSharedRef<SDetailTree> SDetailsView::ConstructTreeView( TSharedRef<SScrollBar>&
 		.OnGetChildren(this, &SDetailsView::OnGetChildrenForDetailTree)
 		.OnSetExpansionRecursive(this, &SDetailsView::SetNodeExpansionStateRecursive)
 		.OnGenerateRow(this, &SDetailsView::OnGenerateRowForDetailTree)
+		.OnRowReleased(this, &SDetailsView::OnRowReleasedForDetailTree)
 		.OnExpansionChanged(this, &SDetailsView::OnItemExpansionChanged)
 		.SelectionMode(ESelectionMode::None)
 		.HandleDirectionalNavigation(false)
diff --git a/Engine/Source/Editor/PropertyEditor/Private/SDetailsViewBase.cpp b/Engine/Source/Editor/PropertyEditor/Private/SDetailsViewBase.cpp
--- a/Engine/Source/Editor/PropertyEditor/Private/SDetailsViewBase.cpp
+++ b/Engine/Source/Editor/PropertyEditor/Private/SDetailsViewBase.cpp
@@ -72,6 +72,23 @@ TSharedRef<ITableRow> SDetailsViewBase::OnGenerateRowForDetailTree(TSharedRef<FD
 	return InTreeNode->GenerateWidgetForTableView(OwnerTable, DetailsViewArgs.bAllowFavoriteSystem);
 }
 
+void SDetailsViewBase::OnRowReleasedForDetailTree(const TSharedRef<ITableRow>& TableRow)
+{
+	// search upwards from the current keyboard-focused widget to see if it's contained in our row
+	TSharedPtr<SWidget> CurrentWidget = FSlateApplication::Get().GetKeyboardFocusedWidget();
+	while (CurrentWidget.IsValid())
+	{
+		if (CurrentWidget == TableRow->AsWidget())
+		{
+			// if so, clear focus so that any pending value changes are committed
+			FSlateApplication::Get().ClearKeyboardFocus();
+			return;
+		}
+
+		CurrentWidget = CurrentWidget->GetParentWidget();
+	}
+}
+
 void SDetailsViewBase::SetRootExpansionStates(const bool bExpand, const bool bRecurse)
 {
 	if(ContainsMultipleTopLevelObjects())
diff --git a/Engine/Source/Editor/PropertyEditor/Private/SDetailsViewBase.h b/Engine/Source/Editor/PropertyEditor/Private/SDetailsViewBase.h
--- a/Engine/Source/Editor/PropertyEditor/Private/SDetailsViewBase.h
+++ b/Engine/Source/Editor/PropertyEditor/Private/SDetailsViewBase.h
@@ -228,6 +228,9 @@ protected:
 	 */
 	TSharedRef<ITableRow> OnGenerateRowForDetailTree( TSharedRef<FDetailTreeNode> InTreeNode, const TSharedRef<STableViewBase>& OwnerTable );
 
+	/** Clear focus when the row is scrolled out of view. */
+	void OnRowReleasedForDetailTree(const TSharedRef<ITableRow>& TableRow);
+
 	/** @return true if show only modified is checked */
 	bool IsShowOnlyModifiedChecked() const { return CurrentFilter.bShowOnlyModified; }
 
