Engine change to allow only building the property nodes that user is going to see.

When property view widgets get reused, we are running (generate nodes) repeatedly while scrolling. Through some trial and error, I've identified that since I am only showing a fraction of the nodes I don't need to have all the nodes built; I can get away with generating the ones I will see, which improves performance from around 130 ms to 10-20ms per widget. I am adding a way to specify the properties you want to build and ignore the rest, which ultimately dramatically speeds up the PropertyView widget for the case where you are only showing a few properties from the object.

diff --git a/Engine/Source/Editor/PropertyEditor/Private/DetailLayoutBuilderImpl.cpp b/Engine/Source/Editor/PropertyEditor/Private/DetailLayoutBuilderImpl.cpp
--- a/Engine/Source/Editor/PropertyEditor/Private/DetailLayoutBuilderImpl.cpp
+++ b/Engine/Source/Editor/PropertyEditor/Private/DetailLayoutBuilderImpl.cpp
@@ -932,3 +932,26 @@ void FDetailLayoutBuilderImpl::SortCategories(const FOnCategorySortOrderFunction
 {
 	CategorySortOrderFunctions.Add(InSortFunction);
 }
+
+void FDetailLayoutBuilderImpl::SetPropertyGenerationAllowListPaths(const TSet<FString>& InPropertyGenerationAllowListPaths)
+{
+	PropertyGenerationAllowListPaths = InPropertyGenerationAllowListPaths;
+}
+
+bool FDetailLayoutBuilderImpl::IsPropertyPathAllowed(const FString& InPath) const
+{
+	if (PropertyGenerationAllowListPaths.IsEmpty())
+	{
+		return true;
+	}
+
+	for (const FString& PropertyName : PropertyGenerationAllowListPaths)
+	{
+		if (PropertyName.StartsWith(InPath))
+		{
+			return true;
+		}
+	}
+
+	return false;
+}
\ No newline at end of file
diff --git a/Engine/Source/Editor/PropertyEditor/Private/DetailLayoutBuilderImpl.h b/Engine/Source/Editor/PropertyEditor/Private/DetailLayoutBuilderImpl.h
--- a/Engine/Source/Editor/PropertyEditor/Private/DetailLayoutBuilderImpl.h
+++ b/Engine/Source/Editor/PropertyEditor/Private/DetailLayoutBuilderImpl.h
@@ -53,6 +53,8 @@ public:
 	virtual bool HasClassDefaultObject() const override;
 	virtual void RegisterInstancedCustomPropertyTypeLayout(FName PropertyTypeName, FOnGetPropertyTypeCustomizationInstance PropertyTypeLayoutDelegate, TSharedPtr<IPropertyTypeIdentifier> Identifier = nullptr) override;
 	virtual void SortCategories(const FOnCategorySortOrderFunction& SortFunction) override;
+	virtual void SetPropertyGenerationAllowListPaths(const TSet<FString>& InPropertyGenerationAllowListPaths) override;
+	virtual bool IsPropertyPathAllowed(const FString& InPath) const override;
 	/**
 	 * Creates a default category. The SDetails view will generate widgets in default categories
 	 *
@@ -244,6 +246,8 @@ private:
 	FClassToPropertyMap& PropertyMap;
 	/** Force hidden categories set by the user */
 	TSet<FName> ForceHiddenCategories;
+	/** If not empty only nodes for property specified paths will be generated */
+	TSet<FString> PropertyGenerationAllowListPaths;
 	/** Nodes that require ticking */
 	TSet<FDetailTreeNode*> TickableNodes;
 	/** Current filter applied to the view */
diff --git a/Engine/Source/Editor/PropertyEditor/Private/DetailPropertyRow.cpp b/Engine/Source/Editor/PropertyEditor/Private/DetailPropertyRow.cpp
--- a/Engine/Source/Editor/PropertyEditor/Private/DetailPropertyRow.cpp
+++ b/Engine/Source/Editor/PropertyEditor/Private/DetailPropertyRow.cpp
@@ -371,7 +371,12 @@ void FDetailPropertyRow::GenerateChildrenForPropertyNode( TSharedPtr<FPropertyNo
 		for( int32 ChildIndex = 0; ChildIndex < RootPropertyNode->GetNumChildNodes(); ++ChildIndex )
 		{
 			TSharedPtr<FPropertyNode> ChildNode = RootPropertyNode->GetChildNode(ChildIndex);
-
+			
+			if (!LayoutBuilder.IsPropertyPathAllowed(ChildNode->GetPropertyPath()))
+			{
+				continue;
+			}
+			
 			if( ChildNode.IsValid() && ChildNode->HasNodeFlags( EPropertyNodeFlags::IsCustomized ) == 0 )
 			{
 				if( ChildNode->AsObjectNode() )
diff --git a/Engine/Source/Editor/PropertyEditor/Private/PropertyNode.h b/Engine/Source/Editor/PropertyEditor/Private/PropertyNode.h
--- a/Engine/Source/Editor/PropertyEditor/Private/PropertyNode.h
+++ b/Engine/Source/Editor/PropertyEditor/Private/PropertyNode.h
@@ -949,6 +949,11 @@ public:
 	 */
 	void SetExpandedChildPropertyNodes(const TSet<FString>& InNodesToExpand);
 
+	/**
+	 * Helper to fetch a PropertyPath 
+	 */
+	const FString& GetPropertyPath() const { return PropertyPath; }
+
 protected:
 
 	TSharedRef<FEditPropertyChain> BuildPropertyChain( FProperty* PropertyAboutToChange );
diff --git a/Engine/Source/Editor/PropertyEditor/Private/PropertyRowGenerator.cpp b/Engine/Source/Editor/PropertyEditor/Private/PropertyRowGenerator.cpp
--- a/Engine/Source/Editor/PropertyEditor/Private/PropertyRowGenerator.cpp
+++ b/Engine/Source/Editor/PropertyEditor/Private/PropertyRowGenerator.cpp
@@ -453,6 +453,11 @@ const TArray<TSharedRef<class IClassViewerFilter>>& FPropertyRowGenerator::GetCl
 	return NotImplemented;
 }
 
+void FPropertyRowGenerator::SetPropertyGenerationAllowListPaths(const TSet<FString>& InPropertyPathsToGeneratePropertyNodes)
+{
+	PropertyGenerationAllowListPaths = InPropertyPathsToGeneratePropertyNodes;
+}
+
 void FPropertyRowGenerator::PreSetObject(int32 NumNewObjects, bool bHasStructRoots)
 {
 	// Save existing expanded items first
@@ -610,7 +615,6 @@ void FPropertyRowGenerator::UpdatePropertyMaps()
 {
 	RootTreeNodes.Empty();
 
-
 	for (FDetailLayoutData& LayoutData : DetailLayouts)
 	{
 		// Check uniqueness.  It is critical that detail layouts can be destroyed
@@ -653,6 +657,7 @@ void FPropertyRowGenerator::UpdateSinglePropertyMap(TSharedPtr<FComplexPropertyN
 	TSharedPtr<FDetailLayoutBuilderImpl> DetailLayout = MakeShareable(new FDetailLayoutBuilderImpl(InRootPropertyNode, LayoutData.ClassToPropertyMap, PropertyUtilities, PropertyGenerationUtilities, nullptr, false));
 	DetailLayout->AddNodeVisibilityChangedHandler(FSimpleMulticastDelegate::FDelegate::CreateRaw(this, &FPropertyRowGenerator::LayoutNodeVisibilityChanged));
 	LayoutData.DetailLayout = DetailLayout;
+	LayoutData.DetailLayout->SetPropertyGenerationAllowListPaths(PropertyGenerationAllowListPaths);
 
 	TSharedPtr<FComplexPropertyNode> RootPropertyNode = InRootPropertyNode;
 	check(RootPropertyNode.IsValid());
diff --git a/Engine/Source/Editor/PropertyEditor/Private/PropertyRowGenerator.h b/Engine/Source/Editor/PropertyEditor/Private/PropertyRowGenerator.h
--- a/Engine/Source/Editor/PropertyEditor/Private/PropertyRowGenerator.h
+++ b/Engine/Source/Editor/PropertyEditor/Private/PropertyRowGenerator.h
@@ -95,6 +95,8 @@ public:
 	{
 		CustomValidatePropertyNodesFunction = MoveTemp(InCustomValidatePropertyNodesFunction);
 	}
+	/** This function sets property paths to generate PropertyNodes.This improves the performance for cases where PropertyView is only showing a few properties of the object by not generating all other PropertyNodes */
+	virtual void SetPropertyGenerationAllowListPaths(const TSet<FString>& InPropertyGenerationAllowListPaths) override;
 	virtual void InvalidateCachedState() override;
 
 	/** FTickableEditorObject interface */
@@ -154,6 +156,8 @@ private:
 	FOnFinishedChangingProperties OnFinishedChangingPropertiesDelegate;
 	/** The ValidatePropertyNodes function can be overridden with this member, if set. Useful if your implementation doesn't require this kind of validation each Tick. */
 	FOnValidatePropertyRowGeneratorNodes CustomValidatePropertyNodesFunction;
+	/** If specified only nodes for these properties will be generated */
+	TSet<FString> PropertyGenerationAllowListPaths;
 
 	bool bViewingClassDefaultObject;
 };
diff --git a/Engine/Source/Editor/PropertyEditor/Public/DetailLayoutBuilder.h b/Engine/Source/Editor/PropertyEditor/Public/DetailLayoutBuilder.h
--- a/Engine/Source/Editor/PropertyEditor/Public/DetailLayoutBuilder.h
+++ b/Engine/Source/Editor/PropertyEditor/Public/DetailLayoutBuilder.h
@@ -255,4 +255,16 @@ public:
 	 * @param DetailLayoutDelegate	The delegate to call when querying for custom detail layouts for the classes properties
 	 */
 	virtual void RegisterInstancedCustomPropertyTypeLayout(FName PropertyTypeName, FOnGetPropertyTypeCustomizationInstance PropertyTypeLayoutDelegate, TSharedPtr<IPropertyTypeIdentifier> Identifier = nullptr) = 0;
+
+	/**
+	 * This function sets property paths to generate PropertyNodes.This improves the performance for cases where PropertyView is only showing a few properties of the object by not generating all other PropertyNodes
+	 *
+	 * @param InPropertyGenerationAllowListPaths Set of the property paths
+	 */
+	virtual void SetPropertyGenerationAllowListPaths(const TSet<FString>& InPropertyGenerationAllowListPaths) = 0;
+
+	/**
+	 * @return True if the property path is contained within our allowed paths
+	 */
+	virtual bool IsPropertyPathAllowed(const FString& InPath) const = 0;
 };
diff --git a/Engine/Source/Editor/PropertyEditor/Public/IPropertyRowGenerator.h b/Engine/Source/Editor/PropertyEditor/Public/IPropertyRowGenerator.h
--- a/Engine/Source/Editor/PropertyEditor/Public/IPropertyRowGenerator.h
+++ b/Engine/Source/Editor/PropertyEditor/Public/IPropertyRowGenerator.h
@@ -112,6 +112,11 @@ public:
 	 * This is useful if your implementation doesn't need to validate nodes every tick or needs to perform some other form of validation. */
 	virtual void SetCustomValidatePropertyNodesFunction(FOnValidatePropertyRowGeneratorNodes InCustomValidatePropertyNodesFunction) = 0;
 
+	/* Use this function to set property paths to generate FPropertyNodes. This improves the performance for cases where UPropertyView is only showing a few properties of the object by not generating all other FPropertyNodes 
+	*  An example of a nested property path: RootPropertyName[0]->NestedPropertyName[0]->PropertyName
+	*/
+	virtual void SetPropertyGenerationAllowListPaths(const TSet<FString>& InPropertyGenerationAllowListPaths) = 0;
+
 	/* Invalidates internal cached state.  Common use of this API is to synchronize the viewed object with changes made by external code. */
 	virtual void InvalidateCachedState() = 0;
 };
\ No newline at end of file
