Removed unnecessary TWeakPtr::IsValid() check.

diff --git a/Engine/Source/Editor/PropertyEditor/Private/PropertyNode.h b/Engine/Source/Editor/PropertyEditor/Private/PropertyNode.h
--- a/Engine/Source/Editor/PropertyEditor/Private/PropertyNode.h
+++ b/Engine/Source/Editor/PropertyEditor/Private/PropertyNode.h
@@ -397,8 +397,8 @@ public:
 	/**
 	 * Returns the parent node in the hierarchy
 	 */
-	FPropertyNode*			GetParentNode() { return ParentNodeWeakPtr.IsValid() ? ParentNodeWeakPtr.Pin().Get() : nullptr; }
-	const FPropertyNode*	GetParentNode() const { return ParentNodeWeakPtr.IsValid() ? ParentNodeWeakPtr.Pin().Get() : nullptr; }
+	FPropertyNode*			GetParentNode() { return ParentNodeWeakPtr.Pin().Get(); }
+	const FPropertyNode*	GetParentNode() const { return ParentNodeWeakPtr.Pin().Get(); }
 	TSharedPtr<FPropertyNode> GetParentNodeSharedPtr() { return ParentNodeWeakPtr.Pin(); }
 	/**
 	 * Returns the Property this Node represents
