Log "Add to Allowlist" events in PropertyEditor

diff --git a/Engine/Source/Editor/PropertyEditor/Private/PropertyEditorPermissionList.cpp b/Engine/Source/Editor/PropertyEditor/Private/PropertyEditorPermissionList.cpp
--- a/Engine/Source/Editor/PropertyEditor/Private/PropertyEditorPermissionList.cpp
+++ b/Engine/Source/Editor/PropertyEditor/Private/PropertyEditorPermissionList.cpp
@@ -6,6 +6,8 @@
 #include "Misc/CoreDelegates.h"
 #include "UObject/UnrealType.h"
 
+DEFINE_LOG_CATEGORY(LogPropertyEditorPermissionList);
+
 namespace
 {
 	const FName PropertyEditorPermissionListOwner = "PropertyEditorPermissionList";
diff --git a/Engine/Source/Editor/PropertyEditor/Private/SDetailSingleItemRow.cpp b/Engine/Source/Editor/PropertyEditor/Private/SDetailSingleItemRow.cpp
--- a/Engine/Source/Editor/PropertyEditor/Private/SDetailSingleItemRow.cpp
+++ b/Engine/Source/Editor/PropertyEditor/Private/SDetailSingleItemRow.cpp
@@ -1015,7 +1015,7 @@ const UStruct* GetExactStructForProperty(const UStruct* MostDerivedStruct, const
 	return MostDerivedStruct;
 }
 
-void SDetailSingleItemRow::CopyRowNameText() const
+FString SDetailSingleItemRow::GetRowNameText() const
 {
 	if (const TSharedPtr<FDetailTreeNode> Owner = OwnerTreeNode.Pin())
 	{
@@ -1023,9 +1023,19 @@ void SDetailSingleItemRow::CopyRowNameText() const
 		if (BaseStructure)
 		{
 			const UStruct* ExactStruct = GetExactStructForProperty(Owner->GetParentBaseStructure(), Owner->GetNodeName());
-			FPlatformApplicationMisc::ClipboardCopy(*FString::Printf(TEXT("(%s, %s)"), *FSoftObjectPtr(ExactStruct).ToString(), *Owner->GetNodeName().ToString()));
+			return FString::Printf(TEXT("(%s, %s)"), *FSoftObjectPtr(ExactStruct).ToString(), *Owner->GetNodeName().ToString());
 		}
 	}
+	return FString();
+}
+
+void SDetailSingleItemRow::CopyRowNameText() const
+{
+	const FString RowNameText = GetRowNameText();
+	if (!RowNameText.IsEmpty())
+	{
+		FPlatformApplicationMisc::ClipboardCopy(*RowNameText);
+	}
 }
 
 void SDetailSingleItemRow::OnToggleAllowList() const
@@ -1038,10 +1048,12 @@ void SDetailSingleItemRow::OnToggleAllowList() const
 		if (IsAllowListChecked())
 		{
 			FPropertyEditorPermissionList::Get().RemoveFromAllowList(ExactStruct, Owner->GetNodeName(), OwnerName);
+			UE_LOG(LogPropertyEditorPermissionList, Log, TEXT("Removing %s from AllowList"), *GetRowNameText());
 		}
 		else
 		{
 			FPropertyEditorPermissionList::Get().AddToAllowList(ExactStruct, Owner->GetNodeName(), OwnerName);
+			UE_LOG(LogPropertyEditorPermissionList, Log, TEXT("Adding %s to AllowList"), *GetRowNameText());
 		}
 	}
 }
@@ -1066,10 +1078,12 @@ void SDetailSingleItemRow::OnToggleDenyList() const
 		if (IsDenyListChecked())
 		{
 			FPropertyEditorPermissionList::Get().RemoveFromDenyList(ExactStruct, Owner->GetNodeName(), OwnerName);
+			UE_LOG(LogPropertyEditorPermissionList, Log, TEXT("Removing %s from DenyList"), *GetRowNameText());
 		}
 		else
 		{
 			FPropertyEditorPermissionList::Get().AddToDenyList(ExactStruct, Owner->GetNodeName(), OwnerName);
+			UE_LOG(LogPropertyEditorPermissionList, Log, TEXT("Adding %s to AllowList"), *GetRowNameText());
 		}
 	}
 }
diff --git a/Engine/Source/Editor/PropertyEditor/Private/SDetailSingleItemRow.h b/Engine/Source/Editor/PropertyEditor/Private/SDetailSingleItemRow.h
--- a/Engine/Source/Editor/PropertyEditor/Private/SDetailSingleItemRow.h
+++ b/Engine/Source/Editor/PropertyEditor/Private/SDetailSingleItemRow.h
@@ -89,8 +89,9 @@ private:
 	void OnFavoriteMenuToggle();
 	bool CanFavorite() const;
 	bool IsFavorite() const;
-
+	
 	/** UIActions to help populate the PropertyEditorPermissionList, which must first be turned on through FPropertyEditorPermissionList::Get().SetShouldShowMenuEntries */
+	FString GetRowNameText() const;
 	void CopyRowNameText() const;
 	void OnToggleAllowList() const;
 	bool IsAllowListChecked() const;
diff --git a/Engine/Source/Editor/PropertyEditor/Public/PropertyEditorPermissionList.h b/Engine/Source/Editor/PropertyEditor/Public/PropertyEditorPermissionList.h
--- a/Engine/Source/Editor/PropertyEditor/Public/PropertyEditorPermissionList.h
+++ b/Engine/Source/Editor/PropertyEditor/Public/PropertyEditorPermissionList.h
@@ -7,6 +7,8 @@
 #include "UObject/WeakObjectPtr.h"
 #include "Misc/NamePermissionList.h"
 
+DECLARE_LOG_CATEGORY_EXTERN(LogPropertyEditorPermissionList, Log, All);
+
 /** Struct, OwnerName */
 DECLARE_MULTICAST_DELEGATE_TwoParams(FPermissionListUpdated, TSoftObjectPtr<UStruct>, FName);
 
