- Fix potential crash when AddObjectPropertyData is called during a customization

diff --git a/Engine/Source/Editor/PropertyEditor/Private/DetailLayoutHelpers.cpp b/Engine/Source/Editor/PropertyEditor/Private/DetailLayoutHelpers.cpp
--- a/Engine/Source/Editor/PropertyEditor/Private/DetailLayoutHelpers.cpp
+++ b/Engine/Source/Editor/PropertyEditor/Private/DetailLayoutHelpers.cpp
@@ -337,26 +337,51 @@ namespace DetailLayoutHelpers
 				{
 					UStruct* Class = WeakClass.Get();
 
-					FClassInstanceToPropertyMap& InstancedPropertyMap = LayoutData.ClassToPropertyMap.FindChecked(Class->GetFName());
-					for (FClassInstanceToPropertyMap::TIterator InstanceIt(InstancedPropertyMap); InstanceIt; ++InstanceIt)
+					TSet<FName> ProcessedClasses;
+					bool bHasItemsToProcess = true;
+					while (bHasItemsToProcess)
 					{
-						FName Key = InstanceIt.Key();
-						LayoutData.DetailLayout->SetCurrentCustomizationClass(Class, Key);
+						bHasItemsToProcess = false;
 
-						const FOnGetDetailCustomizationInstance& DetailDelegate = LayoutIt.Value()->DetailLayoutDelegate;
+						// Copy ClassToProperty map since it could be modified during customization (through AddObjectPropertyData for example)
+						FClassInstanceToPropertyMap InstancedPropertyMapCopy = LayoutData.ClassToPropertyMap.FindChecked(Class->GetFName());
 
-						if (DetailDelegate.IsBound())
+						// Stamp this rounds number of item to detect if current map has changed after going through it
+						const int32 InitialCount = InstancedPropertyMapCopy.Num();
+						for (FClassInstanceToPropertyMap::TIterator InstanceIt(InstancedPropertyMapCopy); InstanceIt; ++InstanceIt)
 						{
-							QueriedClasses.Add(Class);
+							// Stamp classes that have been processed to avoid reprocessing them when doing subsequent rounds
+							const FName Key = InstanceIt.Key();
+							if (ProcessedClasses.Contains(Key))
+							{
+								continue;
+							}
+							ProcessedClasses.Add(Key);
+
+							LayoutData.DetailLayout->SetCurrentCustomizationClass(Class, Key);
 
-							// Create a new instance of the custom detail layout for the current class
-							TSharedRef<IDetailCustomization> CustomizationInstance = DetailDelegate.Execute();
+							const FOnGetDetailCustomizationInstance& DetailDelegate = LayoutIt.Value()->DetailLayoutDelegate;
+
+							if (DetailDelegate.IsBound())
+							{
+								QueriedClasses.Add(Class);
 
-							// Ask for details immediately
-							CustomizationInstance->CustomizeDetails(LayoutData.DetailLayout);
+								// Create a new instance of the custom detail layout for the current class
+								TSharedRef<IDetailCustomization> CustomizationInstance = DetailDelegate.Execute();
 
-							// Save the instance from destruction until we refresh
-							LayoutData.CustomizationClassInstances.Add(CustomizationInstance);
+								// Ask for details immediately
+								CustomizationInstance->CustomizeDetails(LayoutData.DetailLayout);
+
+								// Save the instance from destruction until we refresh
+								LayoutData.CustomizationClassInstances.Add(CustomizationInstance);
+							}
+						}
+						
+						// Verify if current mapping has changed after customizations which would require another pass
+						const int32 NewCount = LayoutData.ClassToPropertyMap.FindChecked(Class->GetFName()).Num();
+						if (NewCount > InitialCount)
+						{
+							bHasItemsToProcess = true;
 						}
 					}
 				}
