Fix config hierarchy crashes due to being enabled on array entries

diff --git a/Engine/Source/Editor/PropertyEditor/Private/DetailPropertyRow.cpp b/Engine/Source/Editor/PropertyEditor/Private/DetailPropertyRow.cpp
--- a/Engine/Source/Editor/PropertyEditor/Private/DetailPropertyRow.cpp
+++ b/Engine/Source/Editor/PropertyEditor/Private/DetailPropertyRow.cpp
@@ -888,7 +888,11 @@ void FDetailPropertyRow::MakeValueWidget( FDetailWidgetRow& Row, const TSharedPt
 			}
 		}
 
-		if (PropertyHandle->HasMetaData(TEXT("ConfigHierarchyEditable")))
+		// Don't add config hierarchy to container children, can't edit child properties at the hiearchy's per file level
+		TSharedPtr<IPropertyHandle> ParentHandle = PropertyHandle->GetParentHandle();
+		bool bIsChildProperty = ParentHandle && (ParentHandle->AsArray() || ParentHandle->AsMap() || ParentHandle->AsSet());
+
+		if (!bIsChildProperty && PropertyHandle->HasMetaData(TEXT("ConfigHierarchyEditable")))
 		{
 			ValueWidget->AddSlot()
 			.AutoWidth()
