SPropertyEditorNumeric now uses ForceUnits even when placed in an array.

diff --git a/Engine/Source/Editor/PropertyEditor/Private/UserInterface/PropertyEditor/SPropertyEditorNumeric.h b/Engine/Source/Editor/PropertyEditor/Private/UserInterface/PropertyEditor/SPropertyEditorNumeric.h
--- a/Engine/Source/Editor/PropertyEditor/Private/UserInterface/PropertyEditor/SPropertyEditorNumeric.h
+++ b/Engine/Source/Editor/PropertyEditor/Private/UserInterface/PropertyEditor/SPropertyEditorNumeric.h
@@ -54,10 +54,11 @@ public:
 
 		const TSharedRef< FPropertyNode > PropertyNode = InPropertyEditor->GetPropertyNode();
 		const FProperty* Property = InPropertyEditor->GetProperty();
+		TSharedPtr<IPropertyHandle> PropertyHandle = InPropertyEditor->GetPropertyHandle();
 
 		if(!Property->IsA(FFloatProperty::StaticClass()) && !Property->IsA(FDoubleProperty::StaticClass()) && Property->HasMetaData(PropertyEditorConstants::MD_Bitmask))
 		{
-			auto CreateBitmaskFlagsArray = [](const FProperty* Prop)
+			auto CreateBitmaskFlagsArray = [PropertyHandle]()
 			{
 				const int32 BitmaskBitCount = sizeof(NumericType) << 3;
 
@@ -65,7 +66,7 @@ public:
 				Result.Empty(BitmaskBitCount);
 
 				const UEnum* BitmaskEnum = nullptr;
-				const FString& BitmaskEnumName = Prop->GetMetaData(PropertyEditorConstants::MD_BitmaskEnum);
+				const FString& BitmaskEnumName = PropertyHandle->GetMetaData(PropertyEditorConstants::MD_BitmaskEnum);
 				if (!BitmaskEnumName.IsEmpty())
 				{
 					// @TODO: Potentially replace this with a parameter passed in from a member variable on the FProperty (e.g. FByteProperty::Enum)
@@ -89,8 +90,8 @@ public:
 						}
 					};
 
-					const TArray<FName> AllowedPropertyEnums = PropertyEditorHelpers::GetValidEnumsFromPropertyOverride(Prop, BitmaskEnum);
-					const TArray<FName> DisallowedPropertyEnums = PropertyEditorHelpers::GetInvalidEnumsFromPropertyOverride(Prop, BitmaskEnum);
+					const TArray<FName> AllowedPropertyEnums = PropertyEditorHelpers::GetValidEnumsFromPropertyOverride(PropertyHandle->GetProperty(), BitmaskEnum);
+					const TArray<FName> DisallowedPropertyEnums = PropertyEditorHelpers::GetInvalidEnumsFromPropertyOverride(PropertyHandle->GetProperty(), BitmaskEnum);
 					// Note: This loop doesn't include (BitflagsEnum->NumEnums() - 1) in order to skip the implicit "MAX" value that gets added to the enum type at compile time.
 					for (int32 BitmaskEnumIndex = 0; BitmaskEnumIndex < BitmaskEnum->NumEnums() - 1; ++BitmaskEnumIndex)
 					{
@@ -142,7 +143,7 @@ public:
 
 			const FComboBoxStyle& ComboBoxStyle = FCoreStyle::Get().GetWidgetStyle< FComboBoxStyle >("ComboBox");
 
-			const auto& GetComboButtonText = [this, CreateBitmaskFlagsArray, Property]() -> FText
+			const auto& GetComboButtonText = [this, CreateBitmaskFlagsArray]() -> FText
 			{
 				TOptional<NumericType> Value = OnGetValue();
 				if (Value.IsSet())
@@ -150,7 +151,7 @@ public:
 					NumericType BitmaskValue = Value.GetValue();
 					if (BitmaskValue != 0)
 					{
-						TArray<FBitmaskFlagInfo> BitmaskFlags = CreateBitmaskFlagsArray(Property);
+						TArray<FBitmaskFlagInfo> BitmaskFlags = CreateBitmaskFlagsArray();
 
 						TArray<FText> SetFlags;
 						SetFlags.Reserve(BitmaskFlags.Num());
@@ -182,12 +183,12 @@ public:
 			// Constructs the UI for bitmask property editing.
 			SAssignNew(PrimaryWidget, SComboButton)
 			.ComboButtonStyle(&ComboBoxStyle.ComboButtonStyle)
-			.ToolTipText_Lambda([this, CreateBitmaskFlagsArray, Property]
+			.ToolTipText_Lambda([this, CreateBitmaskFlagsArray]
 			{
 				TOptional<NumericType> Value = OnGetValue();
 				if (Value.IsSet())
 				{
-					TArray<FBitmaskFlagInfo> BitmaskFlags = CreateBitmaskFlagsArray(Property);
+					TArray<FBitmaskFlagInfo> BitmaskFlags = CreateBitmaskFlagsArray();
 
 					TArray<FText> SetFlags;
 					SetFlags.Reserve(BitmaskFlags.Num());
@@ -211,11 +212,11 @@ public:
 				.Font(InArgs._Font)
 				.Text_Lambda(GetComboButtonText)
 			]
-			.OnGetMenuContent_Lambda([this, CreateBitmaskFlagsArray, Property]()
+			.OnGetMenuContent_Lambda([this, CreateBitmaskFlagsArray]()
 			{
 				FMenuBuilder MenuBuilder(false, nullptr);
 
-				TArray<FBitmaskFlagInfo> BitmaskFlags = CreateBitmaskFlagsArray(Property);
+				TArray<FBitmaskFlagInfo> BitmaskFlags = CreateBitmaskFlagsArray();
 				for (int i = 0; i < BitmaskFlags.Num(); ++i)
 				{
 					MenuBuilder.AddMenuEntry(
@@ -256,10 +257,10 @@ public:
 		else
 		{
 			// Instance metadata overrides per-class metadata.
-			auto GetMetaDataFromKey = [&PropertyNode, &Property](const FName& Key) -> const FString&
+			auto GetMetaDataFromKey = [PropertyHandle](const FName& Key) -> const FString&
 			{
-				const FString* InstanceValue = PropertyNode->GetInstanceMetaData(Key);
-				return (InstanceValue != nullptr) ? *InstanceValue : Property->GetMetaData(Key);
+				const FString* InstanceValue = PropertyHandle->GetInstanceMetaData(Key);
+				return (InstanceValue != nullptr) ? *InstanceValue : PropertyHandle->GetMetaData(Key);
 			};
 
 			const FString& MetaUIMinString = GetMetaDataFromKey("UIMin");
@@ -332,12 +333,12 @@ public:
 
 			FObjectPropertyNode* ObjectPropertyNode = PropertyNode->FindObjectItemParent();
 			const bool bAllowSpin = (!ObjectPropertyNode || (1 == ObjectPropertyNode->GetNumObjects()))
-				&& !PropertyNode->GetProperty()->GetBoolMetaData("NoSpinbox");
+				&& !PropertyHandle->GetBoolMetaData("NoSpinbox");
 
 			// Set up the correct type interface if we want to display units on the property editor
 
 			// First off, check for ForceUnits= meta data. This meta tag tells us to interpret, and always display the value in these units. FUnitConversion::Settings().ShouldDisplayUnits does not apply to such properties
-			const FString& ForcedUnits = InPropertyEditor->GetProperty()->GetMetaData(TEXT("ForceUnits"));
+			const FString& ForcedUnits = PropertyHandle->GetMetaData(TEXT("ForceUnits"));
 			auto PropertyUnits = FUnitConversion::UnitFromString(*ForcedUnits);
 			if (PropertyUnits.IsSet())
 			{
@@ -350,7 +351,7 @@ public:
 			{
 				if (FUnitConversion::Settings().ShouldDisplayUnits())
 				{
-					const FString& DynamicUnits = InPropertyEditor->GetProperty()->GetMetaData(TEXT("Units"));
+					const FString& DynamicUnits = PropertyHandle->GetMetaData(TEXT("Units"));
 					if (!DynamicUnits.IsEmpty())
 					{
 						PropertyUnits = FUnitConversion::UnitFromString(*DynamicUnits);
