Uses of IPropertyHandle::GetMetadataProperty() have been switched to directly call GetMetadata(), which does it internally anyway. It's probably a mistake for GetMetadataProperty() to be exposed anyway, and it should be deprecated.

diff --git a/Engine/Plugins/Editor/AssetManagerEditor/Source/AssetManagerEditor/Private/PrimaryAssetTypeCustomization.cpp b/Engine/Plugins/Editor/AssetManagerEditor/Source/AssetManagerEditor/Private/PrimaryAssetTypeCustomization.cpp
--- a/Engine/Plugins/Editor/AssetManagerEditor/Source/AssetManagerEditor/Private/PrimaryAssetTypeCustomization.cpp
+++ b/Engine/Plugins/Editor/AssetManagerEditor/Source/AssetManagerEditor/Private/PrimaryAssetTypeCustomization.cpp
@@ -33,7 +33,7 @@ void FPrimaryAssetTypeCustomization::CustomizeHeader(TSharedRef<class IPropertyH
 	}
 
 	StructPropertyHandle = InStructPropertyHandle;
-	const bool bAllowClear = !(StructPropertyHandle->GetMetaDataProperty()->PropertyFlags & CPF_NoClear);
+	const bool bAllowClear = !StructPropertyHandle->GetMetaDataProperty()->HasAnyPropertyFlags(CPF_NoClear);
 
 	HeaderRow
 	.NameContent()
diff --git a/Engine/Source/Editor/DetailCustomizations/Private/TextCustomization.cpp b/Engine/Source/Editor/DetailCustomizations/Private/TextCustomization.cpp
--- a/Engine/Source/Editor/DetailCustomizations/Private/TextCustomization.cpp
+++ b/Engine/Source/Editor/DetailCustomizations/Private/TextCustomization.cpp
@@ -22,12 +22,12 @@ namespace
 
 		virtual bool IsMultiLineText() const override
 		{
-			return PropertyHandle->IsValidHandle() && PropertyHandle->GetMetaDataProperty()->GetBoolMetaData("MultiLine");
+			return PropertyHandle->IsValidHandle() && PropertyHandle->GetBoolMetaData("MultiLine");
 		}
 
 		virtual bool IsPassword() const override
 		{
-			return PropertyHandle->IsValidHandle() && PropertyHandle->GetMetaDataProperty()->GetBoolMetaData("PasswordField");
+			return PropertyHandle->IsValidHandle() && PropertyHandle->GetBoolMetaData("PasswordField");
 		}
 
 		virtual bool IsReadOnly() const override
diff --git a/Engine/Source/Editor/PropertyEditor/Private/UserInterface/PropertyEditor/SPropertyEditorText.cpp b/Engine/Source/Editor/PropertyEditor/Private/UserInterface/PropertyEditor/SPropertyEditorText.cpp
--- a/Engine/Source/Editor/PropertyEditor/Private/UserInterface/PropertyEditor/SPropertyEditorText.cpp
+++ b/Engine/Source/Editor/PropertyEditor/Private/UserInterface/PropertyEditor/SPropertyEditorText.cpp
@@ -12,9 +12,9 @@ void SPropertyEditorText::Construct( const FArguments& InArgs, const TSharedRef<
 	PropertyEditor = InPropertyEditor;
 
 	bIsFNameProperty = InPropertyEditor->PropertyIsA(FNameProperty::StaticClass());
-	bIsMultiLine = InPropertyEditor->GetPropertyHandle()->GetMetaDataProperty()->GetBoolMetaData("MultiLine");
+	bIsMultiLine = InPropertyEditor->GetPropertyHandle()->GetBoolMetaData("MultiLine");
 
-	const bool bIsPassword = InPropertyEditor->GetPropertyHandle()->GetMetaDataProperty()->GetBoolMetaData("PasswordField");
+	const bool bIsPassword = InPropertyEditor->GetPropertyHandle()->GetBoolMetaData("PasswordField");
 	
 	TSharedPtr<SHorizontalBox> HorizontalBox;
 	if(bIsMultiLine)
