Temp fix for crashes when editing array properties

diff --git a/Engine/Source/Editor/PropertyEditor/Private/PropertyTextUtilities.cpp b/Engine/Source/Editor/PropertyEditor/Private/PropertyTextUtilities.cpp
--- a/Engine/Source/Editor/PropertyEditor/Private/PropertyTextUtilities.cpp
+++ b/Engine/Source/Editor/PropertyEditor/Private/PropertyTextUtilities.cpp
@@ -18,14 +18,7 @@ void FPropertyTextUtilities::PropertyToTextHelper(FString& OutString, const FPro
 
 void FPropertyTextUtilities::PropertyToTextHelper(FString& OutString, const FPropertyNode* InPropertyNode, const FProperty* Property, const FObjectBaseAddress& ObjectAddress, EPropertyPortFlags PortFlags)
 {
-	if (InPropertyNode->GetArrayIndex() != INDEX_NONE || Property->ArrayDim == 1)
-	{
-		Property->ExportTextItem_InContainer(OutString, ObjectAddress.Object, nullptr, ObjectAddress.Object, PortFlags);
-	}
-	else
-	{
-		FArrayProperty::ExportTextInnerItem(OutString, Property, ObjectAddress.BaseAddress, Property->ArrayDim, ObjectAddress.BaseAddress, Property->ArrayDim, ObjectAddress.Object, PortFlags);
-	}
+	PropertyToTextHelper(OutString, InPropertyNode, Property, ObjectAddress.BaseAddress, ObjectAddress.Object, PortFlags);
 }
 
 void FPropertyTextUtilities::TextToPropertyHelper(const TCHAR* Buffer, const FPropertyNode* InPropertyNode, const FProperty* Property, uint8* ValueAddress, UObject* Object, EPropertyPortFlags PortFlags)
@@ -42,12 +35,5 @@ void FPropertyTextUtilities::TextToPropertyHelper(const TCHAR* Buffer, const FPr
 
 void FPropertyTextUtilities::TextToPropertyHelper(const TCHAR* Buffer, const FPropertyNode* InPropertyNode, const FProperty* Property, const FObjectBaseAddress& ObjectAddress, EPropertyPortFlags PortFlags)
 {
-	if (InPropertyNode->GetArrayIndex() != INDEX_NONE || Property->ArrayDim == 1)
-	{
-		Property->ImportText_InContainer(Buffer, ObjectAddress.Object, ObjectAddress.Object, PortFlags);
-	}
-	else
-	{
-		FArrayProperty::ImportTextInnerItem(Buffer, Property, ObjectAddress.BaseAddress, PortFlags, ObjectAddress.Object);
-	}
+	TextToPropertyHelper(Buffer, InPropertyNode, Property, ObjectAddress.BaseAddress, ObjectAddress.Object, PortFlags);
 }
