Show Only Modified will now also take into account the visibility of a custom reset to default provided to FDetailWidgetRow::OverrideResetToDefault().

diff --git a/Engine/Source/Editor/PropertyEditor/Private/DetailItemNode.cpp b/Engine/Source/Editor/PropertyEditor/Private/DetailItemNode.cpp
--- a/Engine/Source/Editor/PropertyEditor/Private/DetailItemNode.cpp
+++ b/Engine/Source/Editor/PropertyEditor/Private/DetailItemNode.cpp
@@ -609,6 +609,13 @@ static bool PassesAllFilters( FDetailItemNode* ItemNode, const FDetailLayoutCust
 		const bool bPassesCategoryFilter = !bSearchFilterIsEmpty && InFilter.bShowAllChildrenIfCategoryMatches ? Local::StringPassesFilter(InFilter, InCategoryName) : false;
 		const bool bPassesValueFilter = !bSearchFilterIsEmpty && Local::StringPassesFilter(InFilter, Local::GetPropertyNodeValueFilterString(InCustomization, PropertyNodePin));
 
+		const FDetailWidgetRow WidgetRow = InCustomization.GetWidgetRow();
+		bool bIsCustomResetToDefaultVisible = false;
+		if (InFilter.bShowOnlyModified && WidgetRow.CustomResetToDefault.IsSet())
+		{
+			bIsCustomResetToDefaultVisible = WidgetRow.CustomResetToDefault.GetValue().IsResetToDefaultVisible(ItemNode->CreatePropertyHandle());
+		}
+
 		bPassesAllFilters = false;
 		if( PropertyNodePin.IsValid() && !PropertyNodePin->AsCategoryNode())
 		{
@@ -617,7 +624,7 @@ static bool PassesAllFilters( FDetailItemNode* ItemNode, const FDetailLayoutCust
 			const bool bIsParentSeenDueToFiltering = PropertyNodePin->HasNodeFlags(EPropertyNodeFlags::IsParentSeenDueToFiltering) != 0;
 
 			const bool bPassesSearchFilter = bPassesCategoryFilter || bPassesValueFilter || bSearchFilterIsEmpty || ( bIsNotBeingFiltered || bIsSeenDueToFiltering || bIsParentSeenDueToFiltering );
-			const bool bPassesModifiedFilter = bPassesSearchFilter && ( InFilter.bShowOnlyModified == false || PropertyNodePin->GetDiffersFromDefault() == true );
+			const bool bPassesModifiedFilter = bPassesSearchFilter && (InFilter.bShowOnlyModified == false || PropertyNodePin->GetDiffersFromDefault() == true || bIsCustomResetToDefaultVisible);
 			const bool bPassesAllowListFilter = InFilter.bShowOnlyAllowed ? InFilter.PropertyAllowList.Contains(*FPropertyNode::CreatePropertyPath(PropertyNodePin.ToSharedRef())) : true;
 
 			bool bPassesKeyableFilter = true;
@@ -641,18 +648,18 @@ static bool PassesAllFilters( FDetailItemNode* ItemNode, const FDetailLayoutCust
 		}
 		else if (InCustomization.HasCustomWidget())
 		{
-			const bool bPassesTextFilter = bPassesCategoryFilter || bPassesValueFilter || Local::StringPassesFilter(InFilter, InCustomization.WidgetDecl->FilterTextString.ToString());
-			const bool bPassesModifiedFilter = InFilter.bShowOnlyModified == false || InCustomization.WidgetDecl->EditConditionValue.Get(false);
+			const bool bPassesTextFilter = bPassesCategoryFilter || bPassesValueFilter || Local::StringPassesFilter(InFilter, WidgetRow.FilterTextString.ToString());
 			//@todo we need to support custom widgets for keyable, animated, in particular for transforms(ComponentTransformDetails).
+			const bool bPassesModifiedFilter = (InFilter.bShowOnlyModified == false || WidgetRow.EditConditionValue.Get(false) || bIsCustomResetToDefaultVisible);
 			const bool bPassesKeyableFilter = (InFilter.bShowOnlyKeyable == false);
 			const bool bPassesAnimatedFilter = (InFilter.bShowOnlyAnimated == false);
 			bPassesAllFilters = bPassesTextFilter && bPassesModifiedFilter && bPassesKeyableFilter && bPassesAnimatedFilter;
 		}
 		else if (InCustomization.HasCustomBuilder())
 		{
-			const bool bPassesTextFilter = bPassesCategoryFilter || bPassesValueFilter || Local::StringPassesFilter(InFilter, InCustomization.CustomBuilderRow->GetWidgetRow().FilterTextString.ToString());
+			const bool bPassesTextFilter = bPassesCategoryFilter || bPassesValueFilter || Local::StringPassesFilter(InFilter, WidgetRow.FilterTextString.ToString());
 			//@todo we need to support custom builders for modified, keyable, animated, in particular for transforms(ComponentTransformDetails).
-			const bool bPassesModifiedFilter = (InFilter.bShowOnlyModified == false);
+			const bool bPassesModifiedFilter = (InFilter.bShowOnlyModified == false || bIsCustomResetToDefaultVisible);
 			const bool bPassesKeyableFilter = (InFilter.bShowOnlyKeyable == false);
 			const bool bPassesAnimatedFilter = (InFilter.bShowOnlyAnimated == false);
 			bPassesAllFilters = bPassesTextFilter && bPassesModifiedFilter && bPassesKeyableFilter && bPassesAnimatedFilter;
