Details view's IsPropertyReadOnly delegate is now evaluated every frame instead of only when refreshing the tree, similar to how IsPropertyVisible now works.

IsPropertyReadOnly and IsPropertyVisible removed from update args, because they were only evaluated in conjunction with a full refresh.

diff --git a/Engine/Source/Editor/PropertyEditor/Private/DetailItemNode.cpp b/Engine/Source/Editor/PropertyEditor/Private/DetailItemNode.cpp
--- a/Engine/Source/Editor/PropertyEditor/Private/DetailItemNode.cpp
+++ b/Engine/Source/Editor/PropertyEditor/Private/DetailItemNode.cpp
@@ -786,18 +786,26 @@ TAttribute<bool> FDetailItemNode::IsPropertyEditingEnabled() const
 
 bool FDetailItemNode::IsPropertyEditingEnabledImpl() const
 {
-	const bool IsParentEnabledValue = IsParentEnabled.Get(true);
-	if (Customization.HasCustomWidget())
+	bool bIsEnabled = IsParentEnabled.Get(true);
+
+	IDetailsViewPrivate* DetailsView = GetDetailsView();
+	if (DetailsView)
 	{
-		IDetailsViewPrivate* DetailsView = GetDetailsView();
-		if (DetailsView)
+		if (Customization.HasPropertyNode())
+		{
+			TSharedPtr<FPropertyNode> PropertyNode = Customization.GetPropertyNode();
+			if (PropertyNode->GetProperty() != nullptr)
+			{
+				bIsEnabled &= !DetailsView->IsPropertyReadOnly(FPropertyAndParent(PropertyNode.ToSharedRef()));
+			}
+		}
+		else if (Customization.HasCustomWidget())
 		{
-			return IsParentEnabledValue &&
-				!DetailsView->IsCustomRowReadOnly(FName(*Customization.WidgetDecl->FilterTextString.ToString()), FName(*GetParentCategory()->GetDisplayName().ToString()));
+			bIsEnabled &= !DetailsView->IsCustomRowReadOnly(Customization.GetName(), GetParentCategory()->GetCategoryName());
 		}
 	}
 	
-	return IsParentEnabledValue;
+	return bIsEnabled;
 }
 
 TSharedPtr<FPropertyNode> FDetailItemNode::GetPropertyNode() const
diff --git a/Engine/Source/Editor/PropertyEditor/Private/DetailLayoutHelpers.cpp b/Engine/Source/Editor/PropertyEditor/Private/DetailLayoutHelpers.cpp
--- a/Engine/Source/Editor/PropertyEditor/Private/DetailLayoutHelpers.cpp
+++ b/Engine/Source/Editor/PropertyEditor/Private/DetailLayoutHelpers.cpp
@@ -102,9 +102,6 @@ namespace DetailLayoutHelpers
 				// Children of arrays are not visible directly,
 				bVisibleByDefault &= !bIsChildOfContainer;
 
-				FPropertyAndParent PropertyAndParent(ChildNodePtr.ToSharedRef());
-				const bool bIsUserVisible = InUpdateArgs.IsPropertyVisible(PropertyAndParent);
-
 				// Inners of customized in structs should not be taken into consideration for customizing.  They are not designed to be individually customized when their parent is already customized
 				if (!bIsChildOfCustomizedStruct && !LocalUpdateFavoriteSystemOnly)
 				{
@@ -145,7 +142,7 @@ namespace DetailLayoutHelpers
 					PropertyNodeMap.Add(Property->GetFName(), ChildNodePtr);
 				}
 
-				if (bVisibleByDefault && bIsUserVisible && !bPushOutStructProps)
+				if (bVisibleByDefault && !bPushOutStructProps)
 				{
 					FName CategoryName = CurCategory;
 					// For properties inside a struct, add them to their own category unless they just take the name of the parent struct.  
@@ -158,11 +155,6 @@ namespace DetailLayoutHelpers
 
 					if (!LocalUpdateFavoriteSystemOnly)
 					{
-						if (InUpdateArgs.IsPropertyReadOnly(PropertyAndParent))
-						{
-							ChildNode.SetNodeFlags(EPropertyNodeFlags::IsReadOnly, true);
-						}
-
 						// Add a property to the default category
 						FDetailCategoryImpl& CategoryImpl = DetailLayout.DefaultCategory(CategoryName);
 						CategoryImpl.AddPropertyNode(ChildNodePtr.ToSharedRef(), InstanceName);
@@ -178,20 +170,13 @@ namespace DetailLayoutHelpers
 
 							if (LocalUpdateFavoriteSystemOnly)
 							{
-								if (InUpdateArgs.IsPropertyReadOnly(PropertyAndParent))
+								//If the parent has a condition that is not met, mark the child as read-only
+								FDetailLayoutCustomization ParentTmpCustomization;
+								ParentTmpCustomization.PropertyRow = MakeShared<FDetailPropertyRow>(InNode.AsShared(), FavoritesCategory.AsShared());
+								if (ParentTmpCustomization.PropertyRow->GetPropertyEditor()->IsPropertyEditingEnabled() == false)
 								{
 									ChildNode.SetNodeFlags(EPropertyNodeFlags::IsReadOnly, true);
 								}
-								else
-								{
-									//If the parent has a condition that is not met, mark the child as read-only
-									FDetailLayoutCustomization ParentTmpCustomization;
-									ParentTmpCustomization.PropertyRow = MakeShared<FDetailPropertyRow>(InNode.AsShared(), FavoritesCategory.AsShared());
-									if (ParentTmpCustomization.PropertyRow->GetPropertyEditor()->IsPropertyEditingEnabled() == false)
-									{
-										ChildNode.SetNodeFlags(EPropertyNodeFlags::IsReadOnly, true);
-									}
-								}
 							}
 
 							// Add the property to the favorite
@@ -220,7 +205,6 @@ namespace DetailLayoutHelpers
 					&& !bIsCustomizedStruct // Don't recurse into customized structs
 					&& !bIsChildOfContainer // Don't recurse into containers, the children are drawn by the container property parent
 					&& !bIsEditInlineNew // Edit inline new children are not supported for customization yet
-					&&	bIsUserVisible // Properties must be allowed to be visible by a user if they are not then their children are not visible either
 					&& (!bIsStruct || bPushOutStructProps); //  Only recurse into struct properties if they are going to be displayed as standalone properties in categories instead of inside an expandable area inside a category
 
 				if (bRecurseIntoChildren || LocalUpdateFavoriteSystemOnly)
diff --git a/Engine/Source/Editor/PropertyEditor/Private/DetailLayoutHelpers.h b/Engine/Source/Editor/PropertyEditor/Private/DetailLayoutHelpers.h
--- a/Engine/Source/Editor/PropertyEditor/Private/DetailLayoutHelpers.h
+++ b/Engine/Source/Editor/PropertyEditor/Private/DetailLayoutHelpers.h
@@ -14,8 +14,6 @@ namespace DetailLayoutHelpers
 	{
 		FDetailLayoutData* LayoutData;
 		FCustomPropertyTypeLayoutMap* InstancedPropertyTypeToDetailLayoutMap;
-		TFunction<bool(const FPropertyAndParent&)> IsPropertyVisible;
-		TFunction<bool(const FPropertyAndParent&)> IsPropertyReadOnly;
 		bool bEnableFavoriteSystem;
 		bool bUpdateFavoriteSystemOnly;
 	};
diff --git a/Engine/Source/Editor/PropertyEditor/Private/PropertyRowGenerator.cpp b/Engine/Source/Editor/PropertyEditor/Private/PropertyRowGenerator.cpp
--- a/Engine/Source/Editor/PropertyEditor/Private/PropertyRowGenerator.cpp
+++ b/Engine/Source/Editor/PropertyEditor/Private/PropertyRowGenerator.cpp
@@ -666,8 +666,6 @@ void FPropertyRowGenerator::UpdateSinglePropertyMap(TSharedPtr<FComplexPropertyN
 
 	LayoutArgs.LayoutData = &LayoutData;
 	LayoutArgs.InstancedPropertyTypeToDetailLayoutMap = &InstancedTypeToLayoutMap;
-	LayoutArgs.IsPropertyReadOnly = [this](const FPropertyAndParent& PropertyAndParent) { return false; };
-	LayoutArgs.IsPropertyVisible = [this](const FPropertyAndParent& PropertyAndParent) { return true; };
 	LayoutArgs.bEnableFavoriteSystem = false;
 	LayoutArgs.bUpdateFavoriteSystemOnly = false;
 
diff --git a/Engine/Source/Editor/PropertyEditor/Private/SDetailsViewBase.cpp b/Engine/Source/Editor/PropertyEditor/Private/SDetailsViewBase.cpp
--- a/Engine/Source/Editor/PropertyEditor/Private/SDetailsViewBase.cpp
+++ b/Engine/Source/Editor/PropertyEditor/Private/SDetailsViewBase.cpp
@@ -441,8 +441,6 @@ void SDetailsViewBase::UpdateSinglePropertyMap(TSharedPtr<FComplexPropertyNode>
 
 	Args.LayoutData = &LayoutData;
 	Args.InstancedPropertyTypeToDetailLayoutMap = &InstancedTypeToLayoutMap;
-	Args.IsPropertyReadOnly = [this](const FPropertyAndParent& PropertyAndParent) { return IsPropertyReadOnly(PropertyAndParent); };
-	Args.IsPropertyVisible = [this](const FPropertyAndParent& PropertyAndParent) { return IsPropertyVisible(PropertyAndParent); };
 	Args.bEnableFavoriteSystem = bEnableFavoriteSystem;
 	Args.bUpdateFavoriteSystemOnly = false;
 	DetailLayoutHelpers::UpdateSinglePropertyMapRecursive(*RootPropertyNode, NAME_None, RootPropertyNode.Get(), Args);
