Fixed mismerged indentation.

diff --git a/Engine/Source/Editor/PropertyEditor/Private/PropertyHandleImpl.cpp b/Engine/Source/Editor/PropertyEditor/Private/PropertyHandleImpl.cpp
--- a/Engine/Source/Editor/PropertyEditor/Private/PropertyHandleImpl.cpp
+++ b/Engine/Source/Editor/PropertyEditor/Private/PropertyHandleImpl.cpp
@@ -2913,14 +2913,14 @@ bool FPropertyHandleBase::GeneratePossibleValues(TArray< TSharedPtr<FString> >&
 
 						FCachedPropertyPath Path(GetOptionsFunctionName);
 						if (!PropertyPathHelpers::GetPropertyValue(Target, Path, StringOptions))
-					{
-						TArray<FName> NameOptions;
-						if (PropertyPathHelpers::GetPropertyValue(Target, Path, NameOptions))
 						{
-							Algo::Transform(NameOptions, StringOptions, [](const FName& InName) { return InName.ToString(); });
+							TArray<FName> NameOptions;
+							if (PropertyPathHelpers::GetPropertyValue(Target, Path, NameOptions))
+							{
+								Algo::Transform(NameOptions, StringOptions, [](const FName& InName) { return InName.ToString(); });
+							}
 						}
 					}
-					}
 
 					// If this is the first time there won't be any options.
 					if (OptionIntersection.Num() == 0)
