Fixed issues while searching details panels (most notably Project Settings) introduced by change 19171653.

Fixed crash when accessing and invalid PropertyEditor in FDetailPropertyRow.

Fixed performance issue caused by calling GetWidgetRow() on FDetailPropertyRow for every row, which caused the row's widgets to be constructed again when filtering.

Added GetCustomResetToDefault() to IDetailLayoutRow, which allows us to bypass the performance penalty of creating the WidgetRow by directly accessing the FResetToDefaultOverride of the row.

diff --git a/Engine/Source/Editor/PropertyEditor/Private/DetailCategoryBuilderImpl.cpp b/Engine/Source/Editor/PropertyEditor/Private/DetailCategoryBuilderImpl.cpp
--- a/Engine/Source/Editor/PropertyEditor/Private/DetailCategoryBuilderImpl.cpp
+++ b/Engine/Source/Editor/PropertyEditor/Private/DetailCategoryBuilderImpl.cpp
@@ -93,7 +93,7 @@ FDetailWidgetRow FDetailLayoutCustomization::GetWidgetRow() const
 	}
 	else if (HasCustomBuilder())
 	{
-		return CustomBuilderRow->GetWidgetRow();
+		return *CustomBuilderRow->GetWidgetRow();
 	}
 	else if (HasPropertyNode())
 	{
@@ -105,27 +105,45 @@ FDetailWidgetRow FDetailLayoutCustomization::GetWidgetRow() const
 	}
 }
 
-FName FDetailLayoutCustomization::GetName() const
+const IDetailLayoutRow* FDetailLayoutCustomization::GetDetailLayoutRow() const
 {
 	if (HasCustomWidget())
 	{
-		return WidgetDecl->GetRowName();
+		return WidgetDecl.Get();
 	}
 	else if (HasCustomBuilder())
 	{
-		return CustomBuilderRow->GetRowName();
+		return CustomBuilderRow.Get();
 	}
 	else if (HasGroup())
 	{
-		return DetailGroup->GetRowName();
+		return DetailGroup.Get();
 	}
-	else if (GetPropertyNode())
+	else if (PropertyRow.IsValid())
 	{
-		return PropertyRow->GetRowName();
+		return PropertyRow.Get();
+	}
+	return nullptr;
+}
+
+FName FDetailLayoutCustomization::GetName() const
+{
+	if (const IDetailLayoutRow* LayoutRow = GetDetailLayoutRow())
+	{
+		return LayoutRow->GetRowName();
 	}
 	return NAME_None;
 }
 
+TOptional<FResetToDefaultOverride> FDetailLayoutCustomization::GetCustomResetToDefault() const
+{
+	if (const IDetailLayoutRow* LayoutRow = GetDetailLayoutRow())
+	{
+		return LayoutRow->GetCustomResetToDefault();
+	}
+	return TOptional<FResetToDefaultOverride>();
+}
+
 FDetailCategoryImpl::FDetailCategoryImpl(FName InCategoryName, TSharedRef<FDetailLayoutBuilderImpl> InDetailLayout)
 	: HeaderContentWidget(nullptr)
 	, DetailLayoutBuilder(InDetailLayout)
diff --git a/Engine/Source/Editor/PropertyEditor/Private/DetailCategoryBuilderImpl.h b/Engine/Source/Editor/PropertyEditor/Private/DetailCategoryBuilderImpl.h
--- a/Engine/Source/Editor/PropertyEditor/Private/DetailCategoryBuilderImpl.h
+++ b/Engine/Source/Editor/PropertyEditor/Private/DetailCategoryBuilderImpl.h
@@ -60,8 +60,12 @@ struct FDetailLayoutCustomization
 	bool bAdvanced { false };
 	/** Whether or not this customization is custom or a default one. */
 	bool bCustom { false };
+	/** @return The current layout row that this customization represents. May be null if IsValidCustomization() is false. */
+	const IDetailLayoutRow* GetDetailLayoutRow() const;
 	/** @return The name of the row depending on which type of customization this is, then the name of the property node, then NAME_None. */
 	FName GetName() const;
+	/** @return The custom reset to default if one is available. */
+	TOptional<FResetToDefaultOverride> GetCustomResetToDefault() const;
 };
 
 class FDetailLayout
diff --git a/Engine/Source/Editor/PropertyEditor/Private/DetailCustomBuilderRow.cpp b/Engine/Source/Editor/PropertyEditor/Private/DetailCustomBuilderRow.cpp
--- a/Engine/Source/Editor/PropertyEditor/Private/DetailCustomBuilderRow.cpp
+++ b/Engine/Source/Editor/PropertyEditor/Private/DetailCustomBuilderRow.cpp
@@ -11,6 +11,15 @@ FDetailCustomBuilderRow::FDetailCustomBuilderRow( TSharedRef<IDetailCustomNodeBu
 {
 }
 
+TOptional<FResetToDefaultOverride> FDetailCustomBuilderRow::GetCustomResetToDefault() const
+{
+	if (HeaderRow.IsValid())
+	{
+		return HeaderRow->GetCustomResetToDefault();
+	}
+	return TOptional<FResetToDefaultOverride>();
+}
+
 void FDetailCustomBuilderRow::Tick( float DeltaTime ) 
 {
 	return CustomNodeBuilder->Tick( DeltaTime );
@@ -79,9 +88,9 @@ bool FDetailCustomBuilderRow::IsInitiallyCollapsed() const
 	return CustomNodeBuilder->InitiallyCollapsed();
 }
 
-FDetailWidgetRow FDetailCustomBuilderRow::GetWidgetRow()
+TSharedPtr<FDetailWidgetRow> FDetailCustomBuilderRow::GetWidgetRow() const
 {
-	return *HeaderRow;
+	return HeaderRow;
 }
 
 bool FDetailCustomBuilderRow::AreChildCustomizationsHidden() const
diff --git a/Engine/Source/Editor/PropertyEditor/Private/DetailCustomBuilderRow.h b/Engine/Source/Editor/PropertyEditor/Private/DetailCustomBuilderRow.h
--- a/Engine/Source/Editor/PropertyEditor/Private/DetailCustomBuilderRow.h
+++ b/Engine/Source/Editor/PropertyEditor/Private/DetailCustomBuilderRow.h
@@ -21,6 +21,7 @@ public:
 
 	/** IDetailLayoutRow interface */
 	virtual FName GetRowName() const override { return GetCustomBuilderName(); }
+	virtual TOptional<FResetToDefaultOverride> GetCustomResetToDefault() const override;
 
 	void Tick( float DeltaTime );
 	bool RequiresTick() const;
@@ -32,7 +33,7 @@ public:
 	TSharedPtr<IPropertyHandle> GetPropertyHandle() const;
 	void OnGenerateChildren( FDetailNodeList& OutChildren );
 	bool IsInitiallyCollapsed() const;
-	FDetailWidgetRow GetWidgetRow();
+	TSharedPtr<FDetailWidgetRow> GetWidgetRow() const;
 	bool AreChildCustomizationsHidden() const;
 	void SetOriginalPath(FStringView Path) { OriginalPath = Path; }
 	const FString& GetOriginalPath() const { return OriginalPath; }
diff --git a/Engine/Source/Editor/PropertyEditor/Private/DetailGroup.cpp b/Engine/Source/Editor/PropertyEditor/Private/DetailGroup.cpp
--- a/Engine/Source/Editor/PropertyEditor/Private/DetailGroup.cpp
+++ b/Engine/Source/Editor/PropertyEditor/Private/DetailGroup.cpp
@@ -102,6 +102,15 @@ bool FDetailGroup::GetExpansionState() const
 	return false;
 }
 
+TOptional<FResetToDefaultOverride> FDetailGroup::GetCustomResetToDefault() const
+{
+	if (HeaderCustomization.IsValid())
+	{
+		return HeaderCustomization->PropertyRow->GetCustomResetToDefault();
+	}
+	return TOptional<FResetToDefaultOverride>();
+}
+
 TSharedPtr<FDetailPropertyRow> FDetailGroup::GetHeaderPropertyRow() const
 {
 	return HeaderCustomization.IsValid() ? HeaderCustomization->PropertyRow : nullptr;
diff --git a/Engine/Source/Editor/PropertyEditor/Private/DetailGroup.h b/Engine/Source/Editor/PropertyEditor/Private/DetailGroup.h
--- a/Engine/Source/Editor/PropertyEditor/Private/DetailGroup.h
+++ b/Engine/Source/Editor/PropertyEditor/Private/DetailGroup.h
@@ -32,6 +32,7 @@ public:
 
 	/** IDetailLayoutRow interface */
 	virtual FName GetRowName() const override { return GetGroupName(); }
+	virtual TOptional<FResetToDefaultOverride> GetCustomResetToDefault() const override;
 
 	TSharedPtr<FDetailPropertyRow> GetHeaderPropertyRow() const;
 	TSharedPtr<FPropertyNode> GetHeaderPropertyNode() const;
diff --git a/Engine/Source/Editor/PropertyEditor/Private/DetailItemNode.cpp b/Engine/Source/Editor/PropertyEditor/Private/DetailItemNode.cpp
--- a/Engine/Source/Editor/PropertyEditor/Private/DetailItemNode.cpp
+++ b/Engine/Source/Editor/PropertyEditor/Private/DetailItemNode.cpp
@@ -593,6 +593,12 @@ static bool PassesAllFilters( FDetailItemNode* ItemNode, const FDetailLayoutCust
 			return FString();
 		}
 	};
+	
+	auto IsCustomResetToDefaultVisible = [ItemNode, &InCustomization]()
+	{
+		TOptional<FResetToDefaultOverride> CustomResetToDefault = InCustomization.GetCustomResetToDefault();
+		return CustomResetToDefault.IsSet() && CustomResetToDefault.GetValue().IsResetToDefaultVisible(ItemNode->CreatePropertyHandle());
+	};
 
 	bool bPassesAllFilters = true;
 
@@ -609,13 +615,6 @@ static bool PassesAllFilters( FDetailItemNode* ItemNode, const FDetailLayoutCust
 		const bool bPassesCategoryFilter = !bSearchFilterIsEmpty && InFilter.bShowAllChildrenIfCategoryMatches ? Local::StringPassesFilter(InFilter, InCategoryName) : false;
 		const bool bPassesValueFilter = !bSearchFilterIsEmpty && Local::StringPassesFilter(InFilter, Local::GetPropertyNodeValueFilterString(InCustomization, PropertyNodePin));
 
-		const FDetailWidgetRow WidgetRow = InCustomization.GetWidgetRow();
-		bool bIsCustomResetToDefaultVisible = false;
-		if (InFilter.bShowOnlyModified && WidgetRow.CustomResetToDefault.IsSet())
-		{
-			bIsCustomResetToDefaultVisible = WidgetRow.CustomResetToDefault.GetValue().IsResetToDefaultVisible(ItemNode->CreatePropertyHandle());
-		}
-
 		bPassesAllFilters = false;
 		if( PropertyNodePin.IsValid() && !PropertyNodePin->AsCategoryNode())
 		{
@@ -624,7 +623,13 @@ static bool PassesAllFilters( FDetailItemNode* ItemNode, const FDetailLayoutCust
 			const bool bIsParentSeenDueToFiltering = PropertyNodePin->HasNodeFlags(EPropertyNodeFlags::IsParentSeenDueToFiltering) != 0;
 
 			const bool bPassesSearchFilter = bPassesCategoryFilter || bPassesValueFilter || bSearchFilterIsEmpty || ( bIsNotBeingFiltered || bIsSeenDueToFiltering || bIsParentSeenDueToFiltering );
-			const bool bPassesModifiedFilter = bPassesSearchFilter && (InFilter.bShowOnlyModified == false || PropertyNodePin->GetDiffersFromDefault() == true || bIsCustomResetToDefaultVisible);
+
+			bool bPassesModifiedFilter = true;
+			if (bPassesSearchFilter && InFilter.bShowOnlyModified)
+			{
+				bPassesModifiedFilter = PropertyNodePin->GetDiffersFromDefault() || IsCustomResetToDefaultVisible();
+			}
+
 			const bool bPassesAllowListFilter = InFilter.bShowOnlyAllowed ? InFilter.PropertyAllowList.Contains(*FPropertyNode::CreatePropertyPath(PropertyNodePin.ToSharedRef())) : true;
 
 			bool bPassesKeyableFilter = true;
@@ -648,18 +653,18 @@ static bool PassesAllFilters( FDetailItemNode* ItemNode, const FDetailLayoutCust
 		}
 		else if (InCustomization.HasCustomWidget())
 		{
-			const bool bPassesTextFilter = bPassesCategoryFilter || bPassesValueFilter || Local::StringPassesFilter(InFilter, WidgetRow.FilterTextString.ToString());
+			const bool bPassesTextFilter = bPassesCategoryFilter || bPassesValueFilter || Local::StringPassesFilter(InFilter, InCustomization.WidgetDecl->FilterTextString.ToString());
 			//@todo we need to support custom widgets for keyable, animated, in particular for transforms(ComponentTransformDetails).
-			const bool bPassesModifiedFilter = (InFilter.bShowOnlyModified == false || WidgetRow.EditConditionValue.Get(false) || bIsCustomResetToDefaultVisible);
+			const bool bPassesModifiedFilter = (InFilter.bShowOnlyModified == false || InCustomization.WidgetDecl->EditConditionValue.Get(false) || IsCustomResetToDefaultVisible());
 			const bool bPassesKeyableFilter = (InFilter.bShowOnlyKeyable == false);
 			const bool bPassesAnimatedFilter = (InFilter.bShowOnlyAnimated == false);
 			bPassesAllFilters = bPassesTextFilter && bPassesModifiedFilter && bPassesKeyableFilter && bPassesAnimatedFilter;
 		}
 		else if (InCustomization.HasCustomBuilder())
 		{
-			const bool bPassesTextFilter = bPassesCategoryFilter || bPassesValueFilter || Local::StringPassesFilter(InFilter, WidgetRow.FilterTextString.ToString());
+			const bool bPassesTextFilter = bPassesCategoryFilter || bPassesValueFilter || Local::StringPassesFilter(InFilter, InCustomization.CustomBuilderRow->GetWidgetRow()->FilterTextString.ToString());
 			//@todo we need to support custom builders for modified, keyable, animated, in particular for transforms(ComponentTransformDetails).
-			const bool bPassesModifiedFilter = (InFilter.bShowOnlyModified == false || bIsCustomResetToDefaultVisible);
+			const bool bPassesModifiedFilter = (InFilter.bShowOnlyModified == false || IsCustomResetToDefaultVisible());
 			const bool bPassesKeyableFilter = (InFilter.bShowOnlyKeyable == false);
 			const bool bPassesAnimatedFilter = (InFilter.bShowOnlyAnimated == false);
 			bPassesAllFilters = bPassesTextFilter && bPassesModifiedFilter && bPassesKeyableFilter && bPassesAnimatedFilter;
diff --git a/Engine/Source/Editor/PropertyEditor/Private/DetailPropertyRow.cpp b/Engine/Source/Editor/PropertyEditor/Private/DetailPropertyRow.cpp
--- a/Engine/Source/Editor/PropertyEditor/Private/DetailPropertyRow.cpp
+++ b/Engine/Source/Editor/PropertyEditor/Private/DetailPropertyRow.cpp
@@ -644,7 +644,7 @@ bool FDetailPropertyRow::GetEnabledState() const
 		{
 			Result = Result && CustomEditConditionValue.Get();
 		}
-		else
+		else if (PropertyEditor.IsValid())
 		{
 			Result = Result && PropertyEditor->IsEditConditionMet();
 		}
@@ -696,7 +696,7 @@ void FDetailPropertyRow::SetWidgetRowProperties(FDetailWidgetRow& Row) const
 {
 	// set edit condition handlers - use customized if provided
 	TAttribute<bool> EditConditionValue = CustomEditConditionValue;
-	if (!EditConditionValue.IsSet())
+	if (!EditConditionValue.IsSet() && PropertyEditor.IsValid())
 	{
 		EditConditionValue = TAttribute<bool>(PropertyEditor.ToSharedRef(), &FPropertyEditor::IsEditConditionMet);
 	}
@@ -707,7 +707,7 @@ void FDetailPropertyRow::SetWidgetRowProperties(FDetailWidgetRow& Row) const
 		TWeakPtr<FPropertyEditor> PropertyEditorWeak = PropertyEditor;
 		OnEditConditionValueChanged = FOnBooleanValueChanged::CreateStatic(&ExecuteCustomEditConditionToggle, CustomEditConditionValueChanged, PropertyEditorWeak);
 	}
-	else if (PropertyEditor->SupportsEditConditionToggle())
+	else if (PropertyEditor.IsValid() && PropertyEditor->SupportsEditConditionToggle())
 	{
 		TWeakPtr<FPropertyEditor> PropertyEditorWeak = PropertyEditor;
 		OnEditConditionValueChanged = FOnBooleanValueChanged::CreateStatic(&TogglePropertyEditorEditCondition, PropertyEditorWeak);
diff --git a/Engine/Source/Editor/PropertyEditor/Private/DetailPropertyRow.h b/Engine/Source/Editor/PropertyEditor/Private/DetailPropertyRow.h
--- a/Engine/Source/Editor/PropertyEditor/Private/DetailPropertyRow.h
+++ b/Engine/Source/Editor/PropertyEditor/Private/DetailPropertyRow.h
@@ -48,6 +48,8 @@ public:
 
 	/** IDetailLayoutRow interface */
 	virtual FName GetRowName() const override;
+	virtual TOptional<FResetToDefaultOverride> GetCustomResetToDefault() const override { return CustomResetToDefault; }
+
 
 	/** @return true if this row has widgets with columns */
 	bool HasColumns() const;
diff --git a/Engine/Source/Editor/PropertyEditor/Public/DetailCategoryBuilder.h b/Engine/Source/Editor/PropertyEditor/Public/DetailCategoryBuilder.h
--- a/Engine/Source/Editor/PropertyEditor/Public/DetailCategoryBuilder.h
+++ b/Engine/Source/Editor/PropertyEditor/Public/DetailCategoryBuilder.h
@@ -17,7 +17,9 @@ class IDetailPropertyRow;
 /** Interface for the various types of rows that can be created through a detail layout customization */
 class IDetailLayoutRow
 {
+public:
 	virtual FName GetRowName() const = 0;
+	virtual TOptional<FResetToDefaultOverride> GetCustomResetToDefault() const = 0;
 };
 
 /** The location of a property within a category */
diff --git a/Engine/Source/Editor/PropertyEditor/Public/DetailWidgetRow.h b/Engine/Source/Editor/PropertyEditor/Public/DetailWidgetRow.h
--- a/Engine/Source/Editor/PropertyEditor/Public/DetailWidgetRow.h
+++ b/Engine/Source/Editor/PropertyEditor/Public/DetailWidgetRow.h
@@ -146,6 +146,7 @@ public:
 
 	/** IDetailLayoutRow interface */
 	virtual FName GetRowName() const override { return RowTagName; }
+	virtual TOptional<FResetToDefaultOverride> GetCustomResetToDefault() const override { return CustomResetToDefault; }
 
 	/**
 	 * Assigns content to the entire row
