Fixing reset to defaults arrow not working

diff --git a/Engine/Source/Editor/PropertyEditor/Private/PropertyNode.cpp b/Engine/Source/Editor/PropertyEditor/Private/PropertyNode.cpp
--- a/Engine/Source/Editor/PropertyEditor/Private/PropertyNode.cpp
+++ b/Engine/Source/Editor/PropertyEditor/Private/PropertyNode.cpp
@@ -1445,6 +1445,11 @@ public:
 	 */
 	uint8* GetPropertyRootAddress() const { return PropertyValueRoot.ValueAddress; }
 
+	/**
+	 * @return The address of the default value owner object.
+	 */
+	uint8* GetPropertyDefaultRootAddress() const { return PropertyDefaultValueRoot.ValueAddress; }
+
 private:
 		/**
 	 * Determines whether the property bound to this struct exists in the owning object's archetype.
@@ -2082,7 +2087,7 @@ FString FPropertyNode::GetDefaultValueAsStringForObject( FPropertyItemValueDataT
 			else
 			{
 				// Port flags will cause enums to display correctly
-				InProperty->ExportTextItem_InContainer( DefaultValue, ValueTracker.GetPropertyRootAddress(), ValueTracker.GetPropertyDefaultAddress(), nullptr, PortFlags, nullptr );
+				InProperty->ExportTextItem_InContainer( DefaultValue, ValueTracker.GetPropertyDefaultRootAddress(), ValueTracker.GetPropertyDefaultAddress(), nullptr, PortFlags, nullptr );
 			}
 		}
 	}
