The available details sections list is no longer affected when switching between components.

Categories filtered out by the sections panel would not be included in consideration. This does introduce another issue in that certain sections will not contain any elements if a customization hides all children of a category, but that is perhaps less disruptive than the bug this fixes.

diff --git a/Engine/Source/Editor/PropertyEditor/Private/SDetailsView.cpp b/Engine/Source/Editor/PropertyEditor/Private/SDetailsView.cpp
--- a/Engine/Source/Editor/PropertyEditor/Private/SDetailsView.cpp
+++ b/Engine/Source/Editor/PropertyEditor/Private/SDetailsView.cpp
@@ -1384,10 +1384,7 @@ TMap<FName, FText> SDetailsView::GetAllSections() const
 				TArray<TSharedRef<FDetailTreeNode>> Children;
 				Category.GetGeneratedChildren(Children, true, false);
 
-				if (Children.Num() > 0)
-				{
-					LayoutCategories.Add(TreeNode->GetNodeName());
-				}
+				LayoutCategories.Add(TreeNode->GetNodeName());
 			}
 		}
 	}
