PropertyEditorPermissionList now always respects the AllowListAll rule even if an AllowList was manually defined

diff --git a/Engine/Source/Editor/PropertyEditor/Private/PropertyEditorPermissionList.cpp b/Engine/Source/Editor/PropertyEditor/Private/PropertyEditorPermissionList.cpp
--- a/Engine/Source/Editor/PropertyEditor/Private/PropertyEditorPermissionList.cpp
+++ b/Engine/Source/Editor/PropertyEditor/Private/PropertyEditorPermissionList.cpp
@@ -45,7 +45,11 @@ void FPropertyEditorPermissionList::AddPermissionList(TSoftObjectPtr<UStruct> St
 {
 	FPropertyEditorPermissionListEntry& Entry = RawPropertyEditorPermissionList.FindOrAdd(Struct);
 	Entry.PermissionList.Append(PermissionList);
-	Entry.Rules = Rules;
+	// Always use the most permissive rule previously set
+	if (Entry.Rules > Rules)
+	{
+		Entry.Rules = Rules;
+	}
 	// The cache isn't too expensive to recompute, so it is cleared
 	// and lazily repopulated any time the raw PermissionList changes.
 	ClearCache();
@@ -147,7 +151,8 @@ const FNamePermissionList& FPropertyEditorPermissionList::GetCachedPermissionLis
 	check(Struct);
 
 	const FPropertyEditorPermissionListEntry* Entry = RawPropertyEditorPermissionList.Find(Struct);
-	bool bIsThisAllowListEmpty = Entry ? Entry->PermissionList.GetAllowList().Num() == 0 : true;
+	// If an entry is set to AllowListAll, then treat it as if it has no manually-defined properties.
+	const bool bIsThisAllowListEmpty = Entry ? (Entry->PermissionList.GetAllowList().Num() == 0 || Entry->Rules == EPropertyEditorPermissionListRules::AllowListAllProperties) : true;
 
 	// Normally this case would be caught in GetCachedPermissionListForStruct, but when being called recursively from a subclass
 	// we still need to update bInOutShouldAllowListAllProperties so that new PermissionLists cache properly.
@@ -174,7 +179,19 @@ const FNamePermissionList& FPropertyEditorPermissionList::GetCachedPermissionLis
 		// Append this struct's PermissionList on top of the parent's PermissionList
 		if (Entry)
 		{
-			NewPermissionList.Append(Entry->PermissionList);
+			if (bIsThisAllowListEmpty)
+			{
+				// If the AllowList is empty, we only want to append the DenyLists
+				FNamePermissionList DuplicatePermissionList = Entry->PermissionList;
+				// Hack to get around the fact that there's no easy way to only clear an AllowList
+				TMap<FName, FPermissionListOwners>& AllowList = const_cast<TMap<FName, FPermissionListOwners>&>(DuplicatePermissionList.GetAllowList());
+				AllowList.Empty();
+				NewPermissionList.Append(DuplicatePermissionList);
+			}
+			else
+			{
+				NewPermissionList.Append(Entry->PermissionList);
+			}
 
 			if (Entry->Rules == EPropertyEditorPermissionListRules::AllowListAllProperties)
 			{
@@ -184,8 +201,7 @@ const FNamePermissionList& FPropertyEditorPermissionList::GetCachedPermissionLis
 
 		// PermissionList all properties if the flag is set, the parent Struct has a PermissionList, and this Struct has no PermissionList
 		// If the parent Struct's PermissionList is empty then that already implies all properties are visible
-		// If this Struct has a PermissionList, the manually-specified list always overrides the ShouldPermissionListAllProperties rule
-		if (bInOutShouldAllowListAllProperties && NewPermissionList.GetAllowList().Num() > 0 && bIsThisAllowListEmpty)
+		if (bInOutShouldAllowListAllProperties && NewPermissionList.GetAllowList().Num() > 0)
 		{
 			for (TFieldIterator<FProperty> Property(Struct, EFieldIteratorFlags::ExcludeSuper, EFieldIteratorFlags::ExcludeDeprecated); Property; ++Property)
 			{
@@ -195,7 +211,7 @@ const FNamePermissionList& FPropertyEditorPermissionList::GetCachedPermissionLis
 
 		PermissionList = &CachedPropertyEditorPermissionList.Add(Struct, NewPermissionList);
 	}
-
+	
 	// If this Struct has no PermissionList, then the ShouldPermissionListAllProperties rule just forwards its current value on to the next subclass.
 	// This causes an issue in the case where a Struct should have no PermissionListed properties but wants to PermissionList all subclass properties.
 	// In this case, simply add a dummy entry to the Struct's PermissionList that (likely) won't ever collide with a real property name
diff --git a/Engine/Source/Editor/PropertyEditor/Public/PropertyEditorPermissionList.h b/Engine/Source/Editor/PropertyEditor/Public/PropertyEditorPermissionList.h
--- a/Engine/Source/Editor/PropertyEditor/Public/PropertyEditorPermissionList.h
+++ b/Engine/Source/Editor/PropertyEditor/Public/PropertyEditorPermissionList.h
@@ -28,13 +28,13 @@ DECLARE_MULTICAST_DELEGATE_TwoParams(FPermissionListUpdated, TSoftObjectPtr<UStr
  */
 enum class EPropertyEditorPermissionListRules : uint8
 {
-	// If a PermissionList is manually defined for this struct, PermissionList those properties. Otherwise, use the parent Struct's rule.
-	UseExistingPermissionList,
 	// If no PermissionList is manually defined for this Struct, AllowList all properties from this Struct and its subclasses
 	AllowListAllProperties,
 	// If a PermissionList is manually defined for this Struct, AllowList all properties from this Struct's subclasses.
 	// If this functionality is needed without any properties to AllowList, a fake property must be added to AllowList instead.
-	AllowListAllSubclassProperties
+	AllowListAllSubclassProperties,
+	// If a PermissionList is manually defined for this struct, PermissionList those properties. Otherwise, use the parent Struct's rule.
+	UseExistingPermissionList
 };
 
 struct FPropertyEditorPermissionListEntry
