Temp fix for a crash when reverting a struct property value to default in the editor

diff --git a/Engine/Source/Editor/PropertyEditor/Private/PropertyNode.cpp b/Engine/Source/Editor/PropertyEditor/Private/PropertyNode.cpp
--- a/Engine/Source/Editor/PropertyEditor/Private/PropertyNode.cpp
+++ b/Engine/Source/Editor/PropertyEditor/Private/PropertyNode.cpp
@@ -2087,7 +2087,7 @@ FString FPropertyNode::GetDefaultValueAsStringForObject( FPropertyItemValueDataT
 			else
 			{
 				// Port flags will cause enums to display correctly
-				InProperty->ExportTextItem_InContainer( DefaultValue, ValueTracker.GetPropertyDefaultRootAddress(), ValueTracker.GetPropertyDefaultAddress(), nullptr, PortFlags, nullptr );
+				InProperty->ExportTextItem_Direct(DefaultValue, ValueTracker.GetPropertyDefaultAddress(), ValueTracker.GetPropertyDefaultAddress(), nullptr, PortFlags, nullptr);
 			}
 		}
 	}
