Fixed localization duplicate warning in CIS.
#rnx

diff --git a/Engine/Source/Editor/PropertyEditor/Private/SResetToDefaultMenu.cpp b/Engine/Source/Editor/PropertyEditor/Private/SResetToDefaultMenu.cpp
--- a/Engine/Source/Editor/PropertyEditor/Private/SResetToDefaultMenu.cpp
+++ b/Engine/Source/Editor/PropertyEditor/Private/SResetToDefaultMenu.cpp
@@ -29,7 +29,7 @@ void SResetToDefaultMenu::Construct( const FArguments& InArgs )
 	ChildSlot
 	[
 		SNew(SComboButton)
-		.ToolTipText(NSLOCTEXT( "PropertyEditor", "ResetToDefaultToolTip", "Reset to Default" ))
+		.ToolTipText(NSLOCTEXT( "PropertyEditor", "ResetToDefault", "Reset to Default" ))
 		.HasDownArrow(false)
 		.ButtonStyle( FEditorStyle::Get(), "NoBorder" )
 		.ContentPadding(0) 
