Added unit tests to check for mixed-case boolean values to EditConditionParserTests.

diff --git a/Engine/Source/Editor/PropertyEditor/Private/EditConditionParserTests.cpp b/Engine/Source/Editor/PropertyEditor/Private/EditConditionParserTests.cpp
--- a/Engine/Source/Editor/PropertyEditor/Private/EditConditionParserTests.cpp
+++ b/Engine/Source/Editor/PropertyEditor/Private/EditConditionParserTests.cpp
@@ -164,6 +164,10 @@ bool FEditConditionParser_Parse::RunTest(const FString& Parameters)
 	FEditConditionParser Parser;
 	bool bResult = true;
 
+	bResult &= CanParse(Parser, TEXT("false"), 1, 0);
+	bResult &= CanParse(Parser, TEXT("TRUE"), 1, 0);
+	bResult &= CanParse(Parser, TEXT("fAlsE"), 1, 0);
+	bResult &= CanParse(Parser, TEXT("true == false"), 3, 0);
 	bResult &= CanParse(Parser, TEXT("BoolProperty"), 1, 1);
 	bResult &= CanParse(Parser, TEXT("!BoolProperty"), 2, 1);
 	bResult &= CanParse(Parser, TEXT("BoolProperty == true"), 3, 1);
@@ -228,7 +232,10 @@ static bool RunBoolTests(const IEditConditionContext& Context)
 	bool bResult = true;
 
 	bResult &= CanEvaluate(Parser, Context, TEXT("true"), true);
+	bResult &= CanEvaluate(Parser, Context, TEXT("TRUE"), true);
 	bResult &= CanEvaluate(Parser, Context, TEXT("false"), false);
+	bResult &= CanEvaluate(Parser, Context, TEXT("False"), false);
+	bResult &= CanEvaluate(Parser, Context, TEXT("fAlSe"), false);
 	bResult &= CanEvaluate(Parser, Context, TEXT("!true"), false);
 	bResult &= CanEvaluate(Parser, Context, TEXT("!false"), true);
 	bResult &= CanEvaluate(Parser, Context, TEXT("BoolProperty"), true);
