SResetToDefaultMenu is now also shown even if the property isn't editable, as is the case with other reset to default widgets.

diff --git a/Engine/Source/Editor/PropertyEditor/Private/SResetToDefaultMenu.cpp b/Engine/Source/Editor/PropertyEditor/Private/SResetToDefaultMenu.cpp
--- a/Engine/Source/Editor/PropertyEditor/Private/SResetToDefaultMenu.cpp
+++ b/Engine/Source/Editor/PropertyEditor/Private/SResetToDefaultMenu.cpp
@@ -50,7 +50,7 @@ void SResetToDefaultMenu::Tick( const FGeometry& AllottedGeometry, const double
 	for ( int32 PropIndex = 0; PropIndex < Properties.Num(); ++PropIndex )
 	{
 		TSharedRef<IPropertyHandle> Property = Properties[PropIndex];
-		if( Property->DiffersFromDefault() && !Property->IsEditConst() )
+		if( Property->DiffersFromDefault() )
 		{
 			bShouldBeVisible = true;
 			// If one property should show reset to default, the menu must be visible
