Fixed sub-categories being left visible if all their child nodes had been hidden

#rnx

diff --git a/Engine/Source/Editor/PropertyEditor/Private/DetailItemNode.cpp b/Engine/Source/Editor/PropertyEditor/Private/DetailItemNode.cpp
--- a/Engine/Source/Editor/PropertyEditor/Private/DetailItemNode.cpp
+++ b/Engine/Source/Editor/PropertyEditor/Private/DetailItemNode.cpp
@@ -485,6 +485,20 @@ ENodeVisibility FDetailItemNode::GetVisibility() const
 	{
 		Visibility = (bShouldBeVisibleDueToFiltering || bShouldBeVisibleDueToChildFiltering) ? Visibility : ENodeVisibility::HiddenDueToFiltering;
 	}
+	
+	if (Visibility == ENodeVisibility::Visible && GetNodeType() == EDetailNodeType::Category)
+	{
+		Visibility = ENodeVisibility::ForcedHidden;
+		for (const TSharedRef<FDetailTreeNode>& Child : Children)
+		{
+			if (Child->GetVisibility() != ENodeVisibility::ForcedHidden)
+			{
+				Visibility = ENodeVisibility::Visible;
+				break;
+			}
+		}
+	}
+	
 	return Visibility;
 }
 
