Allow a detail customization to override "ShouldAutoExpand" on the header row

#rnx

diff --git a/Engine/Source/Editor/PropertyEditor/Private/DetailPropertyRow.cpp b/Engine/Source/Editor/PropertyEditor/Private/DetailPropertyRow.cpp
--- a/Engine/Source/Editor/PropertyEditor/Private/DetailPropertyRow.cpp
+++ b/Engine/Source/Editor/PropertyEditor/Private/DetailPropertyRow.cpp
@@ -291,7 +291,13 @@ void FDetailPropertyRow::OnItemNodeInitialized( TSharedRef<FDetailCategoryImpl>
 		if (CustomPropertyWidget->IsEnabledAttr.IsSet())
 		{
 			CustomIsEnabledAttrib = CustomPropertyWidget->IsEnabledAttr;
-		}		
+		}
+
+		// set initial value of auto-expand from struct customization
+		if (CustomPropertyWidget->ForceAutoExpansion.IsSet())
+		{
+			bForceAutoExpansion = CustomPropertyWidget->ForceAutoExpansion.GetValue();
+		}
 	}
 
 	if( bShowCustomPropertyChildren && CustomTypeInterface.IsValid() )
diff --git a/Engine/Source/Editor/PropertyEditor/Public/DetailWidgetRow.h b/Engine/Source/Editor/PropertyEditor/Public/DetailWidgetRow.h
--- a/Engine/Source/Editor/PropertyEditor/Public/DetailWidgetRow.h
+++ b/Engine/Source/Editor/PropertyEditor/Public/DetailWidgetRow.h
@@ -340,6 +340,17 @@ public:
 	*/
 	const TArray<TSharedPtr<IPropertyHandle>>& GetPropertyHandles() const { return PropertyHandles;  }
 
+	/**
+	 * Sets whether or not this property should auto-expand
+	 *
+	 * @param bForceExpansion	true to force the property to be expanded
+	 */
+	FDetailWidgetRow& ShouldAutoExpand(bool bForceExpansion = true)
+	{
+		ForceAutoExpansion = bForceExpansion;
+		return *this;
+	}
+
 public:
 	/** Name column content */
 	FDetailWidgetDecl NameWidget;
@@ -390,5 +401,7 @@ public:
 	TSharedPtr<IDetailDragDropHandler> CustomDragDropHandler;
 	/* All property handle that this custom widget represent */
 	TArray<TSharedPtr<IPropertyHandle>> PropertyHandles;
+	/** True to force auto-expansion */
+	TOptional<bool> ForceAutoExpansion;
 };
 
