Ensure uobjects added through external object customizations are not considered top level objects in a multi-top-level object details panel

diff --git a/Engine/Source/Editor/PropertyEditor/Private/DetailLayoutBuilderImpl.cpp b/Engine/Source/Editor/PropertyEditor/Private/DetailLayoutBuilderImpl.cpp
--- a/Engine/Source/Editor/PropertyEditor/Private/DetailLayoutBuilderImpl.cpp
+++ b/Engine/Source/Editor/PropertyEditor/Private/DetailLayoutBuilderImpl.cpp
@@ -432,7 +432,7 @@ void FDetailLayoutBuilderImpl::GenerateDetailLayout()
 	}
 
 	TSharedPtr<FComplexPropertyNode> RootNodePinned = RootNode.Pin();
-	if (DetailsView && DetailsView->GetRootObjectCustomization() && RootNodePinned->GetInstancesNum())
+	if (DetailsView && DetailsView->GetRootObjectCustomization() && RootNodePinned->GetInstancesNum() && !bLayoutForExternalRoot)
 	{
 		FObjectPropertyNode* ObjectNode = RootNodePinned->AsObjectNode();
 
