Fix permissions not correctly refreshing when iterating on dev permissions in editor

Ensure that AddPermissionList assigns rather than appends to existing permissions, otherwise permissions lists are always grow-only.
Start with bool bShouldPermissionListAllProperties = false in FPropertyEditorPermissionList::GetCachedPermissionListForStruct, otherwise for classes with no parent permissions ALL properties will be exposed even if an allow/deby list is set.

diff --git a/Engine/Source/Editor/PropertyEditor/Private/PropertyEditorPermissionList.cpp b/Engine/Source/Editor/PropertyEditor/Private/PropertyEditorPermissionList.cpp
--- a/Engine/Source/Editor/PropertyEditor/Private/PropertyEditorPermissionList.cpp
+++ b/Engine/Source/Editor/PropertyEditor/Private/PropertyEditorPermissionList.cpp
@@ -44,7 +44,7 @@ FPropertyEditorPermissionList::~FPropertyEditorPermissionList()
 void FPropertyEditorPermissionList::AddPermissionList(TSoftObjectPtr<UStruct> Struct, const FNamePermissionList& PermissionList, EPropertyEditorPermissionListRules Rules)
 {
 	FPropertyEditorPermissionListEntry& Entry = RawPropertyEditorPermissionList.FindOrAdd(Struct);
-	Entry.PermissionList.Append(PermissionList);
+	Entry.PermissionList = PermissionList;
 	// Always use the most permissive rule previously set
 	if (Entry.Rules > Rules)
 	{
@@ -142,7 +142,7 @@ const FNamePermissionList& FPropertyEditorPermissionList::GetCachedPermissionLis
 
 	// Default value doesn't matter since it's a no-op until the first PermissionList is encountered, at which
 	// point the rules will re-assign the value.
-	bool bShouldPermissionListAllProperties = true;
+	bool bShouldPermissionListAllProperties = false;
 	return GetCachedPermissionListForStructHelper(Struct, bShouldPermissionListAllProperties);
 }
 
