Fix PropertyEditorPermissionList including all properties on incorrect classes
Improve property editor speed when bulk-changing permission list

diff --git a/Engine/Source/Editor/PropertyEditor/Private/PropertyEditorPermissionList.cpp b/Engine/Source/Editor/PropertyEditor/Private/PropertyEditorPermissionList.cpp
--- a/Engine/Source/Editor/PropertyEditor/Private/PropertyEditorPermissionList.cpp
+++ b/Engine/Source/Editor/PropertyEditor/Private/PropertyEditorPermissionList.cpp
@@ -41,6 +41,18 @@ FPropertyEditorPermissionList::~FPropertyEditorPermissionList()
 	}
 }
 
+void FPropertyEditorPermissionList::ClearCacheAndBroadcast(TSoftObjectPtr<UStruct> ObjectStruct, FName OwnerName)
+{
+	// The cache isn't too expensive to recompute, so it is cleared
+	// and lazily repopulated any time the raw PermissionList changes.
+	ClearCache();
+
+	if (!bSuppressUpdateDelegate)
+	{
+		PermissionListUpdatedDelegate.Broadcast(ObjectStruct, OwnerName);
+	}
+}
+
 void FPropertyEditorPermissionList::AddPermissionList(TSoftObjectPtr<UStruct> Struct, const FNamePermissionList& PermissionList, EPropertyEditorPermissionListRules Rules)
 {
 	FPropertyEditorPermissionListEntry& Entry = RawPropertyEditorPermissionList.FindOrAdd(Struct);
@@ -50,18 +62,15 @@ void FPropertyEditorPermissionList::AddPermissionList(TSoftObjectPtr<UStruct> St
 	{
 		Entry.Rules = Rules;
 	}
-	// The cache isn't too expensive to recompute, so it is cleared
-	// and lazily repopulated any time the raw PermissionList changes.
-	ClearCache();
-	PermissionListUpdatedDelegate.Broadcast(Struct, NAME_None);
+
+	ClearCacheAndBroadcast(Struct);
 }
 
 void FPropertyEditorPermissionList::RemovePermissionList(TSoftObjectPtr<UStruct> Struct)
 {
 	if (RawPropertyEditorPermissionList.Remove(Struct) > 0)
 	{
-		ClearCache();
-		PermissionListUpdatedDelegate.Broadcast(Struct, NAME_None);
+		ClearCacheAndBroadcast(Struct);
 	}
 }
 
@@ -69,7 +78,31 @@ void FPropertyEditorPermissionList::ClearPermissionList()
 {
 	TArray<TSoftObjectPtr<UStruct>> Keys;
 	RawPropertyEditorPermissionList.Reset();
-	PermissionListUpdatedDelegate.Broadcast(nullptr, NAME_None);
+	ClearCacheAndBroadcast();
+}
+
+void FPropertyEditorPermissionList::UnregisterOwner(const FName Owner)
+{
+	TArray<TSoftObjectPtr<UStruct>> StructsToRemove;
+
+	for (TPair<TSoftObjectPtr<UStruct>, FPropertyEditorPermissionListEntry>& Pair : RawPropertyEditorPermissionList)
+	{
+		Pair.Value.PermissionList.UnregisterOwner(Owner);
+		if (Pair.Value.PermissionList.GetOwnerNames().Num() == 0)
+		{
+			StructsToRemove.Add(Pair.Key);
+		}
+	}
+
+	{
+		TGuardValue<bool> SuppressGuard(bSuppressUpdateDelegate, true);
+		for (TSoftObjectPtr<UStruct>& StructToRemove : StructsToRemove)
+		{
+			RemovePermissionList(StructToRemove);
+		}
+	}
+
+	ClearCacheAndBroadcast(nullptr, Owner);
 }
 
 void FPropertyEditorPermissionList::AddToAllowList(TSoftObjectPtr<UStruct> Struct, const FName PropertyName, const FName Owner)
@@ -77,8 +110,7 @@ void FPropertyEditorPermissionList::AddToAllowList(TSoftObjectPtr<UStruct> Struc
 	FPropertyEditorPermissionListEntry& Entry = RawPropertyEditorPermissionList.FindOrAdd(Struct);
 	if (Entry.PermissionList.AddAllowListItem(Owner, PropertyName))
 	{
-		ClearCache();
-		PermissionListUpdatedDelegate.Broadcast(Struct, Owner);
+		ClearCacheAndBroadcast(Struct, Owner);
 	}
 }
 
@@ -87,8 +119,7 @@ void FPropertyEditorPermissionList::RemoveFromAllowList(TSoftObjectPtr<UStruct>
 	FPropertyEditorPermissionListEntry& Entry = RawPropertyEditorPermissionList.FindOrAdd(Struct);
 	if (Entry.PermissionList.RemoveAllowListItem(Owner, PropertyName))
 	{
-		ClearCache();
-		PermissionListUpdatedDelegate.Broadcast(Struct, Owner);
+		ClearCacheAndBroadcast(Struct, Owner);
 	}
 }
 
@@ -97,8 +128,7 @@ void FPropertyEditorPermissionList::AddToDenyList(TSoftObjectPtr<UStruct> Struct
 	FPropertyEditorPermissionListEntry& Entry = RawPropertyEditorPermissionList.FindOrAdd(Struct);
 	if (Entry.PermissionList.AddDenyListItem(Owner, PropertyName))
 	{
-		ClearCache();
-		PermissionListUpdatedDelegate.Broadcast(Struct, Owner);
+		ClearCacheAndBroadcast(Struct, Owner);
 	}
 }
 
@@ -107,8 +137,7 @@ void FPropertyEditorPermissionList::RemoveFromDenyList(TSoftObjectPtr<UStruct> S
 	FPropertyEditorPermissionListEntry& Entry = RawPropertyEditorPermissionList.FindOrAdd(Struct);
 	if (Entry.PermissionList.RemoveDenyListItem(Owner, PropertyName))
 	{
-		ClearCache();
-		PermissionListUpdatedDelegate.Broadcast(Struct, Owner);
+		ClearCacheAndBroadcast(Struct, Owner);
 	}
 }
 
@@ -193,10 +222,7 @@ const FNamePermissionList& FPropertyEditorPermissionList::GetCachedPermissionLis
 				NewPermissionList.Append(Entry->PermissionList);
 			}
 
-			if (Entry->Rules == EPropertyEditorPermissionListRules::AllowListAllProperties)
-			{
-				bInOutShouldAllowListAllProperties = true;
-			}
+			bInOutShouldAllowListAllProperties = Entry->Rules == EPropertyEditorPermissionListRules::AllowListAllProperties;
 		}
 
 		// PermissionList all properties if the flag is set, the parent Struct has a PermissionList, and this Struct has no PermissionList
diff --git a/Engine/Source/Editor/PropertyEditor/Public/PropertyEditorPermissionList.h b/Engine/Source/Editor/PropertyEditor/Public/PropertyEditorPermissionList.h
--- a/Engine/Source/Editor/PropertyEditor/Public/PropertyEditorPermissionList.h
+++ b/Engine/Source/Editor/PropertyEditor/Public/PropertyEditorPermissionList.h
@@ -59,6 +59,9 @@ public:
 	/** Remove all rules */
 	void ClearPermissionList();
 
+	/** Unregister an owner from all permission lists currently stored */
+	void UnregisterOwner(const FName Owner);
+
 	/** Add a specific property to a UStruct's AllowList */
 	void AddToAllowList(TSoftObjectPtr<UStruct> Struct, const FName PropertyName, const FName Owner = NAME_None);
 	/** Remove a specific property from a UStruct's AllowList */
@@ -101,11 +104,16 @@ public:
 	/** Clear CachedPropertyEditorPermissionList to cause the PermissionListed property list to be regenerated next time it's queried */
 	void ClearCache();
 
+	/** If true, PermissionListUpdatedDelegate will not broadcast when modifying the permission lists */
+	bool bSuppressUpdateDelegate = false;
+
 private:
 	FPropertyEditorPermissionList();
 	~FPropertyEditorPermissionList();
 
 	void RegisterOnBlueprintCompiled();
+
+	void ClearCacheAndBroadcast(TSoftObjectPtr<UStruct> ObjectStruct = nullptr, FName OwnerName = NAME_None);
 	
 	/** Whether DoesPropertyPassFilter should perform its PermissionList check or always return true */
 	bool bEnablePropertyEditorPermissionList = false;
