PR #8878: Adding CanEditChange signature that takes a full property chain to the property in question (Contributed by [REDACTED_PERSON])

diff --git a/Engine/Source/Editor/PropertyEditor/Private/PropertyNode.cpp b/Engine/Source/Editor/PropertyEditor/Private/PropertyNode.cpp
--- a/Engine/Source/Editor/PropertyEditor/Private/PropertyNode.cpp
+++ b/Engine/Source/Editor/PropertyEditor/Private/PropertyNode.cpp
@@ -935,12 +935,14 @@ bool FPropertyNode::IsEditConst() const
 
 			if (!bIsEditConst)
 			{
+				TSharedRef<FEditPropertyChain> PropertyChain = BuildPropertyChain(Property.Get());
+				
 				for (TPropObjectConstIterator CurObjectIt(ObjectPropertyNode->ObjectConstIterator()); CurObjectIt; ++CurObjectIt)
 				{
 					const TWeakObjectPtr<UObject> CurObject = *CurObjectIt;
 					if (CurObject.IsValid())
 					{
-						if (!CurObject->CanEditChange(Property.Get()))
+						if (!CurObject->CanEditChange(*PropertyChain))
 						{
 							// At least one of the objects didn't like the idea of this property being changed.
 							bIsEditConst = true;
@@ -2802,13 +2804,13 @@ TSharedPtr< FPropertyItemValueDataTrackerSlate > FPropertyNode::GetValueTracker(
 
 }
 
-TSharedRef<FEditPropertyChain> FPropertyNode::BuildPropertyChain( FProperty* InProperty )
+TSharedRef<FEditPropertyChain> FPropertyNode::BuildPropertyChain( FProperty* InProperty ) const
 {
 	TSharedRef<FEditPropertyChain> PropertyChain( MakeShareable( new FEditPropertyChain ) );
 
-	FPropertyNode* ItemNode = this;
+	const FPropertyNode* ItemNode = this;
 
-	FComplexPropertyNode* ComplexNode = FindComplexParent();
+	const FComplexPropertyNode* ComplexNode = FindComplexParent();
 	FProperty* MemberProperty = InProperty;
 
 	do
@@ -2818,7 +2820,7 @@ TSharedRef<FEditPropertyChain> FPropertyNode::BuildPropertyChain( FProperty* InP
 			MemberProperty = PropertyChain->GetHead()->GetValue();
 		}
 
-		FProperty* TheProperty	= ItemNode->GetProperty();
+		FProperty* TheProperty	= ItemNode->Property.Get();
 		if ( TheProperty )
 		{
 			// Skip over property window items that correspond to a single element in a static array,
@@ -2844,14 +2846,14 @@ TSharedRef<FEditPropertyChain> FPropertyNode::BuildPropertyChain( FProperty* InP
 	return PropertyChain;
 }
 
-TSharedRef<FEditPropertyChain> FPropertyNode::BuildPropertyChain( FProperty* InProperty, const TSet<UObject*>& InAffectedArchetypeInstances )
+TSharedRef<FEditPropertyChain> FPropertyNode::BuildPropertyChain( FProperty* InProperty, const TSet<UObject*>& InAffectedArchetypeInstances ) const
 {
 	TSharedRef<FEditPropertyChain> PropertyChain = BuildPropertyChain(InProperty);
 	PropertyChain->SetAffectedArchetypeInstances(InAffectedArchetypeInstances);
 	return PropertyChain;
 }
 
-TSharedRef<FEditPropertyChain> FPropertyNode::BuildPropertyChain( FProperty* InProperty, TSet<UObject*>&& InAffectedArchetypeInstances )
+TSharedRef<FEditPropertyChain> FPropertyNode::BuildPropertyChain( FProperty* InProperty, TSet<UObject*>&& InAffectedArchetypeInstances ) const
 {
 	TSharedRef<FEditPropertyChain> PropertyChain = BuildPropertyChain(InProperty);
 	PropertyChain->SetAffectedArchetypeInstances(MoveTemp(InAffectedArchetypeInstances));
diff --git a/Engine/Source/Editor/PropertyEditor/Private/PropertyNode.h b/Engine/Source/Editor/PropertyEditor/Private/PropertyNode.h
--- a/Engine/Source/Editor/PropertyEditor/Private/PropertyNode.h
+++ b/Engine/Source/Editor/PropertyEditor/Private/PropertyNode.h
@@ -948,10 +948,9 @@ public:
 	const FString& GetPropertyPath() const { return PropertyPath; }
 
 protected:
-
-	TSharedRef<FEditPropertyChain> BuildPropertyChain( FProperty* PropertyAboutToChange );
-	TSharedRef<FEditPropertyChain> BuildPropertyChain( FProperty* PropertyAboutToChange, const TSet<UObject*>& InAffectedArchetypeInstances );
-	TSharedRef<FEditPropertyChain> BuildPropertyChain( FProperty* PropertyAboutToChange, TSet<UObject*>&& InAffectedArchetypeInstances );
+	TSharedRef<FEditPropertyChain> BuildPropertyChain( FProperty* PropertyAboutToChange ) const;
+	TSharedRef<FEditPropertyChain> BuildPropertyChain( FProperty* PropertyAboutToChange, const TSet<UObject*>& InAffectedArchetypeInstances ) const;
+	TSharedRef<FEditPropertyChain> BuildPropertyChain( FProperty* PropertyAboutToChange, TSet<UObject*>&& InAffectedArchetypeInstances ) const;
 
 	void NotifyPreChangeInternal(TSharedRef<FEditPropertyChain> PropertyChain, FProperty* PropertyAboutToChange, FNotifyHook* InNotifyHook);
 
diff --git a/Engine/Source/Runtime/CoreUObject/Private/UObject/Obj.cpp b/Engine/Source/Runtime/CoreUObject/Private/UObject/Obj.cpp
--- a/Engine/Source/Runtime/CoreUObject/Private/UObject/Obj.cpp
+++ b/Engine/Source/Runtime/CoreUObject/Private/UObject/Obj.cpp
@@ -479,6 +479,11 @@ bool UObject::CanEditChange( const FProperty* InProperty ) const
 	return bIsMutable;
 }
 
+bool UObject::CanEditChange(const FEditPropertyChain& PropertyChain) const
+{
+	return CanEditChange(PropertyChain.GetActiveNode()->GetValue());
+}
+
 void UObject::PropagatePreEditChange( TArray<UObject*>& AffectedObjects, FEditPropertyChain& PropertyAboutToChange )
 {
 	TArray<UObject*> Instances;
diff --git a/Engine/Source/Runtime/CoreUObject/Public/UObject/Object.h b/Engine/Source/Runtime/CoreUObject/Public/UObject/Object.h
--- a/Engine/Source/Runtime/CoreUObject/Public/UObject/Object.h
+++ b/Engine/Source/Runtime/CoreUObject/Public/UObject/Object.h
@@ -408,6 +408,17 @@ public:
 	 */
 	virtual bool CanEditChange( const FProperty* InProperty ) const;
 
+	/**
+	 * Alternate version of CanEditChange that includes the full property chain leading to the property in question.
+	 * The head of the chain is the FStructProperty member variable that contains the property that was modified.
+	 * The active property in the chain is the specific FProperty in question (the one given to the other signature of CanEditChange)
+	 * 
+	 * @param PropertyChain The chain to the property in question. The ActiveNode of the chain corresponds to the specific property in question.
+	 *
+	 * @return	true if the property can be modified in the editor, otherwise false
+	 */
+	virtual bool CanEditChange(const FEditPropertyChain& PropertyChain) const;
+
 	/** 
 	 * Intentionally non-virtual as it calls the FPropertyChangedEvent version
 	 */
