Added the owning struct name to EditCondition error messages to make it easier to find which property is breaking.

diff --git a/Engine/Source/Editor/PropertyEditor/Private/EditConditionContext.cpp b/Engine/Source/Editor/PropertyEditor/Private/EditConditionContext.cpp
--- a/Engine/Source/Editor/PropertyEditor/Private/EditConditionContext.cpp
+++ b/Engine/Source/Editor/PropertyEditor/Private/EditConditionContext.cpp
@@ -13,16 +13,33 @@ FEditConditionContext::FEditConditionContext(FPropertyNode& InPropertyNode)
 	PropertyNode = InPropertyNode.AsShared();
 }
 
+FName FEditConditionContext::GetContextName() const
+{
+	TSharedPtr<FPropertyNode> PinnedNode = PropertyNode.Pin();
+	if (!PinnedNode.IsValid())
+	{
+		return FName();
+	}
+
+	return PinnedNode->GetProperty()->GetOwnerStruct()->GetFName();
+}
+
 const FBoolProperty* FEditConditionContext::GetSingleBoolProperty(const TSharedPtr<FEditConditionExpression>& Expression) const
 {
-	if (!PropertyNode.IsValid())
+	TSharedPtr<FPropertyNode> PinnedNode = PropertyNode.Pin();
+	if (!PinnedNode.IsValid())
 	{
 		return nullptr;
 	}
 
-	const FProperty* Property = PropertyNode.Pin()->GetProperty();
+	const FProperty* Property = PinnedNode->GetProperty();
+	if (Property == nullptr)
+	{
+		return nullptr;
+	}
 
 	const FBoolProperty* BoolProperty = nullptr;
+
 	for (const FCompiledToken& Token : Expression->Tokens)
 	{
 		if (const EditConditionParserTokens::FPropertyToken* PropertyToken = Token.Node.Cast<EditConditionParserTokens::FPropertyToken>())
diff --git a/Engine/Source/Editor/PropertyEditor/Private/EditConditionContext.h b/Engine/Source/Editor/PropertyEditor/Private/EditConditionContext.h
--- a/Engine/Source/Editor/PropertyEditor/Private/EditConditionContext.h
+++ b/Engine/Source/Editor/PropertyEditor/Private/EditConditionContext.h
@@ -9,6 +9,8 @@ DECLARE_LOG_CATEGORY_EXTERN(LogEditCondition, Log, All);
 class IEditConditionContext
 {
 public:
+	virtual FName GetContextName() const = 0;
+
 	virtual TOptional<bool> GetBoolValue(const FString& PropertyName) const = 0;
 	virtual TOptional<int64> GetIntegerValue(const FString& PropertyName) const = 0;
 	virtual TOptional<double> GetNumericValue(const FString& PropertyName) const = 0;
@@ -29,6 +31,8 @@ public:
 	FEditConditionContext(FPropertyNode& InPropertyNode);
 	virtual ~FEditConditionContext() {}
 
+	virtual FName GetContextName() const override;
+
 	virtual TOptional<bool> GetBoolValue(const FString& PropertyName) const override;
 	virtual TOptional<int64> GetIntegerValue(const FString& PropertyName) const override;
 	virtual TOptional<double> GetNumericValue(const FString& PropertyName) const override;
diff --git a/Engine/Source/Editor/PropertyEditor/Private/EditConditionParser.cpp b/Engine/Source/Editor/PropertyEditor/Private/EditConditionParser.cpp
--- a/Engine/Source/Editor/PropertyEditor/Private/EditConditionParser.cpp
+++ b/Engine/Source/Editor/PropertyEditor/Private/EditConditionParser.cpp
@@ -119,20 +119,25 @@ static TOptional<FExpressionError> ConsumePropertyName(FExpressionTokenConsumer&
 }
 
 template <typename ValueType>
-static void LogEditConditionError(const TValueOrError<ValueType, FExpressionError>& Error)
+static void LogEditConditionError(const TValueOrError<ValueType, FExpressionError>& Error, const IEditConditionContext* Context = nullptr)
 {
 	if (!Error.HasError())
 	{
 		return;
 	}
 
+	const FString Message = Error.GetError().Text.ToString();
+	FString Formatted = Message;
+	if (Context != nullptr)
+	{
+		Formatted = FString::Printf(TEXT("%s - %s"), *Context->GetContextName().ToString(), *Message);
+	}
+
 	static TSet<FString> ErrorsAlreadyLogged;
-	
-	const FString& Message = Error.GetError().Text.ToString();
-	if (!ErrorsAlreadyLogged.Find(Message))
+	if (!ErrorsAlreadyLogged.Find(Formatted))
 	{
-		ErrorsAlreadyLogged.Add(Message);
-		UE_LOG(LogEditCondition, Error, TEXT("%s"), *Message);
+		ErrorsAlreadyLogged.Add(Formatted);
+		UE_LOG(LogEditCondition, Error, TEXT("%s"), *Formatted);
 	}
 }
 
@@ -576,7 +581,6 @@ void CreateNumberOperators(TOperatorJumpTable<IEditConditionContext>& OperatorJu
 static FExpressionResult EnumPropertyEquals(const EditConditionParserTokens::FEnumToken& Enum, const EditConditionParserTokens::FPropertyToken& Property, const IEditConditionContext& Context, bool bNegate)
 {
 	TOptional<FString> TypeName = Context.GetTypeName(Property.PropertyName);
-
 	if (!TypeName.IsSet())
 	{
 		return MakeError(FText::Format(LOCTEXT("InvalidOperand_Type", "EditCondition attempted to use an invalid operand \"{0}\" (type error)."), FText::FromString(Property.PropertyName)));
@@ -588,7 +592,6 @@ static FExpressionResult EnumPropertyEquals(const EditConditionParserTokens::FEn
 	}
 
 	TOptional<FString> ValueProp = Context.GetEnumValue(Property.PropertyName);
-
 	if (!ValueProp.IsSet())
 	{
 		return MakeError(FText::Format(LOCTEXT("InvalidOperand_Value", "EditCondition attempted to use an invalid operand \"{0}\" (value error)."), FText::FromString(Property.PropertyName)));
@@ -735,10 +738,10 @@ TValueOrError<bool, FText> FEditConditionParser::Evaluate(const FEditConditionEx
 	}
 	else
 	{
-		LogEditConditionError(Result);
+		LogEditConditionError(Result, &Context);
 	}
 
-	const FText ErrorText = Result.HasError() ? Result.GetError().Text : FText::GetEmpty();
+	const FText ErrorText = Result.HasError() ? Result.StealError().Text : FText::GetEmpty();
 	return MakeError(ErrorText);
 }
 
diff --git a/Engine/Source/Editor/PropertyEditor/Private/EditConditionParserTests.cpp b/Engine/Source/Editor/PropertyEditor/Private/EditConditionParserTests.cpp
--- a/Engine/Source/Editor/PropertyEditor/Private/EditConditionParserTests.cpp
+++ b/Engine/Source/Editor/PropertyEditor/Private/EditConditionParserTests.cpp
@@ -25,6 +25,11 @@ struct FTestEditConditionContext : IEditConditionContext
 	FTestEditConditionContext(){}
 	virtual ~FTestEditConditionContext() {}
 
+	virtual FName GetContextName() const
+	{
+		return "TestEditConditionContext";
+	}
+
 	virtual TOptional<bool> GetBoolValue(const FString& PropertyName) const override
 	{
 		if (const bool* Value = BoolValues.Find(PropertyName))
