SPropertyEditorCombo now correctly checks the editability of the IPropertyHandle passed to it in ComboArgs.PropertyHandle.

Added ensure to check that either the PropertyEditor or PropertyHandle passed is valid.

Reported through UDN.

diff --git a/Engine/Source/Editor/PropertyEditor/Private/UserInterface/PropertyEditor/SPropertyEditorCombo.cpp b/Engine/Source/Editor/PropertyEditor/Private/UserInterface/PropertyEditor/SPropertyEditorCombo.cpp
--- a/Engine/Source/Editor/PropertyEditor/Private/UserInterface/PropertyEditor/SPropertyEditorCombo.cpp
+++ b/Engine/Source/Editor/PropertyEditor/Private/UserInterface/PropertyEditor/SPropertyEditorCombo.cpp
@@ -73,6 +73,8 @@ void SPropertyEditorCombo::Construct( const FArguments& InArgs, const TSharedPtr
 		TooltipAttribute = TAttribute<FText>::Create(TAttribute<FText>::FGetter::CreateSP(PropertyEditor.ToSharedRef(), &FPropertyEditor::GetValueAsText));
 	}
 
+	ensureMsgf(ComboArgs.PropertyHandle.IsValid(), TEXT("Either PropertyEditor or ComboArgs.PropertyHandle must be set!"));
+
 	TArray<TSharedPtr<FString>> ComboItems;
 	TArray<bool> Restrictions;
 	TArray<TSharedPtr<SToolTip>> RichToolTips;
@@ -159,8 +161,8 @@ void SPropertyEditorCombo::GenerateComboBoxStrings( TArray< TSharedPtr<FString>
 
 			if(Enum)
 			{
-				TArray<FName> AllowedPropertyEnums = PropertyEditorHelpers::GetValidEnumsFromPropertyOverride(Property, Enum);
-				TArray<FName> DisallowedPropertyEnums = PropertyEditorHelpers::GetInvalidEnumsFromPropertyOverride(Property, Enum);
+				TArray<FName> ValidPropertyEnums = PropertyEditorHelpers::GetValidEnumsFromPropertyOverride(Property, Enum);
+				TArray<FName> InvalidPropertyEnums = PropertyEditorHelpers::GetInvalidEnumsFromPropertyOverride(Property, Enum);
 				
 				// Get enum doc link (not just GetDocumentationLink as that is the documentation for the struct we're in, not the enum documentation)
 				FString DocLink = PropertyEditorHelpers::GetEnumDocumentationLink(Property);
@@ -172,18 +174,19 @@ void SPropertyEditorCombo::GenerateComboBoxStrings( TArray< TSharedPtr<FString>
 					bool bShouldBeHidden = Enum->HasMetaData(TEXT("Hidden"), EnumIdx) || Enum->HasMetaData(TEXT("Spacer"), EnumIdx);
 					if(!bShouldBeHidden)
 					{
-						if(AllowedPropertyEnums.Num() > 0)
+						if (ValidPropertyEnums.Num() > 0)
 						{
-							bShouldBeHidden = AllowedPropertyEnums.Find(Enum->GetNameByIndex(EnumIdx)) == INDEX_NONE;
+							bShouldBeHidden = ValidPropertyEnums.Find(Enum->GetNameByIndex(EnumIdx)) == INDEX_NONE;
 						}
-						// If both are specified, InvalidEnumValues takes precedence
-						else if(DisallowedPropertyEnums.Num() > 0)
+
+						// If both are specified, the metadata "InvalidEnumValues" takes precedence
+						if (InvalidPropertyEnums.Num() > 0)
 						{
-							bShouldBeHidden = DisallowedPropertyEnums.Find(Enum->GetNameByIndex(EnumIdx)) != INDEX_NONE;
+							bShouldBeHidden = InvalidPropertyEnums.Find(Enum->GetNameByIndex(EnumIdx)) != INDEX_NONE;
 						}
 					}
 
-					if (!bShouldBeHidden)
+					if(!bShouldBeHidden)
 					{
 						bShouldBeHidden = ComboArgs.PropertyHandle->IsHidden(Excerpt);
 					}
@@ -260,7 +263,11 @@ void SPropertyEditorCombo::SendToObjects( const FString& NewValue )
 
 bool SPropertyEditorCombo::CanEdit() const
 {
-	return PropertyEditor.IsValid() ? !PropertyEditor->IsEditConst() : true;
+	if (ComboArgs.PropertyHandle.IsValid())
+	{
+		return ComboArgs.PropertyHandle->IsEditable();
+	}
+	return true;
 }
 
 #undef LOCTEXT_NAMESPACE
