Added default value of true to IsParentEnabled in FDetailPropertyRow::GetEnabledState().

This is a pretty rare case that only occurs if you bypass OnItemNodeInitialized() by creating and adding a property row outside of a CustomizeDetails() call in a property customization. However, the use case is pretty innocuous and this seems more correct anyway, so there's not really any harm to changing the behaviour.

Reported in UDN case 00426747.

diff --git a/Engine/Source/Editor/PropertyEditor/Private/DetailPropertyRow.cpp b/Engine/Source/Editor/PropertyEditor/Private/DetailPropertyRow.cpp
--- a/Engine/Source/Editor/PropertyEditor/Private/DetailPropertyRow.cpp
+++ b/Engine/Source/Editor/PropertyEditor/Private/DetailPropertyRow.cpp
@@ -640,7 +640,7 @@ bool FDetailPropertyRow::HasEditCondition() const
 
 bool FDetailPropertyRow::GetEnabledState() const
 {
-	bool Result = IsParentEnabled.Get();
+	bool Result = IsParentEnabled.Get(true);
 
 	Result = Result && CustomIsEnabledAttrib.Get(true);
 
