Ensure now takes into account the third way that SPropertyEditorCombo can be validly used by passing in three delegates in the args struct.

diff --git a/Engine/Source/Editor/PropertyEditor/Private/UserInterface/PropertyEditor/SPropertyEditorCombo.cpp b/Engine/Source/Editor/PropertyEditor/Private/UserInterface/PropertyEditor/SPropertyEditorCombo.cpp
--- a/Engine/Source/Editor/PropertyEditor/Private/UserInterface/PropertyEditor/SPropertyEditorCombo.cpp
+++ b/Engine/Source/Editor/PropertyEditor/Private/UserInterface/PropertyEditor/SPropertyEditorCombo.cpp
@@ -73,7 +73,7 @@ void SPropertyEditorCombo::Construct( const FArguments& InArgs, const TSharedPtr
 		TooltipAttribute = TAttribute<FText>::Create(TAttribute<FText>::FGetter::CreateSP(PropertyEditor.ToSharedRef(), &FPropertyEditor::GetValueAsText));
 	}
 
-	ensureMsgf(ComboArgs.PropertyHandle.IsValid(), TEXT("Either PropertyEditor or ComboArgs.PropertyHandle must be set!"));
+	ensureMsgf(ComboArgs.PropertyHandle.IsValid() || (ComboArgs.OnGetStrings.IsBound() && ComboArgs.OnGetValue.IsBound() && ComboArgs.OnValueSelected.IsBound()), TEXT("Either PropertyEditor or ComboArgs.PropertyHandle must be set!"));
 
 	TArray<TSharedPtr<FString>> ComboItems;
 	TArray<bool> Restrictions;
