FPropertyRowGenerator::RequestRefresh() now triggers a refresh on the next Tick().

diff --git a/Engine/Source/Editor/PropertyEditor/Private/PropertyRowGenerator.cpp b/Engine/Source/Editor/PropertyEditor/Private/PropertyRowGenerator.cpp
--- a/Engine/Source/Editor/PropertyEditor/Private/PropertyRowGenerator.cpp
+++ b/Engine/Source/Editor/PropertyEditor/Private/PropertyRowGenerator.cpp
@@ -59,7 +59,13 @@ public:
 			Generator->ForceRefresh();
 		}
 	}
-	virtual void RequestRefresh() override {}
+	virtual void RequestRefresh() override 
+	{ 
+		if (Generator != nullptr)
+		{
+			Generator->RequestRefresh();
+		}
+	}
 
 	virtual TSharedPtr<class FAssetThumbnailPool> GetThumbnailPool() const override
 	{
@@ -370,6 +376,13 @@ void FPropertyRowGenerator::Tick(float DeltaTime)
 		DeferredActions.Empty();
 	}
 
+	if (bRefreshPending)
+	{
+		bRefreshPending = false;
+		
+		ForceRefresh();
+	}
+
 	bool bFullRefresh = ValidatePropertyNodes(RootPropertyNodes);
 
 	for (FDetailLayoutData& LayoutData : DetailLayouts)
@@ -383,7 +396,6 @@ void FPropertyRowGenerator::Tick(float DeltaTime)
 			LayoutData.DetailLayout->Tick(DeltaTime);
 		}
 	}
-
 }
 
 TStatId FPropertyRowGenerator::GetStatId() const 
diff --git a/Engine/Source/Editor/PropertyEditor/Private/PropertyRowGenerator.h b/Engine/Source/Editor/PropertyEditor/Private/PropertyRowGenerator.h
--- a/Engine/Source/Editor/PropertyEditor/Private/PropertyRowGenerator.h
+++ b/Engine/Source/Editor/PropertyEditor/Private/PropertyRowGenerator.h
@@ -109,6 +109,7 @@ public:
 	virtual void EnqueueDeferredAction(FSimpleDelegate DeferredAction);	
 	virtual bool IsPropertyEditingEnabled() const;
 	virtual void ForceRefresh();
+	virtual void RequestRefresh() { bRefreshPending = true; }
 	virtual TSharedPtr<class FAssetThumbnailPool> GetThumbnailPool() const;
 	virtual bool HasClassDefaultObject() const { return bViewingClassDefaultObject; }
 	virtual const TArray<TSharedRef<class IClassViewerFilter>>& GetClassViewerFilters() const;
@@ -159,5 +160,6 @@ private:
 	/** If specified only nodes for these properties will be generated */
 	TSet<FString> PropertyGenerationAllowListPaths;
 
-	bool bViewingClassDefaultObject;
+	bool bViewingClassDefaultObject = false;
+	bool bRefreshPending = false;
 };
