Restored grouping for situations where the same struct or component exists multiple times in the same details view, eg. when two components derive from the same base.

diff --git a/Engine/Source/Editor/PropertyEditor/Private/DetailCategoryBuilderImpl.cpp b/Engine/Source/Editor/PropertyEditor/Private/DetailCategoryBuilderImpl.cpp
--- a/Engine/Source/Editor/PropertyEditor/Private/DetailCategoryBuilderImpl.cpp
+++ b/Engine/Source/Editor/PropertyEditor/Private/DetailCategoryBuilderImpl.cpp
@@ -942,20 +942,37 @@ void FDetailCategoryImpl::GenerateNodesFromCustomizations(const TArray<FDetailLa
 	}
 }
 
+void FDetailCategoryImpl::GenerateChildrenForSingleLayout(const FDetailLayout& Layout, FDetailNodeList& OutChildren)
+{
+	FDetailNodeList GeneratedChildren;
+	GenerateNodesFromCustomizations(Layout.GetSimpleLayouts(), GeneratedChildren);
+
+	const FName InstanceName = Layout.GetInstanceName();
+	if (LayoutMap.ShouldShowGroup(InstanceName))
+	{
+		TSharedRef<FDetailCategoryGroupNode> GroupNode = MakeShareable(new FDetailCategoryGroupNode(InstanceName, AsShared()));
+		GroupNode->SetChildren(GeneratedChildren);
+		OutChildren.Add(GroupNode);
+	}
+	else
+	{
+		OutChildren.Append(GeneratedChildren);
+	}
+}
+
 void FDetailCategoryImpl::GenerateChildrenForLayouts()
 {
-	// note: this can't be a ranged-for, because the map may have items added to it by customizations during iteration
+	// note: these can't be ranged-for, because the map may have items added to it by customizations during iteration
 	for (int32 LayoutIndex = 0; LayoutIndex < LayoutMap.Num(); ++LayoutIndex)
 	{
 		const FDetailLayout& Layout = LayoutMap[LayoutIndex];
-		GenerateNodesFromCustomizations(Layout.GetSimpleLayouts(), SimpleChildNodes);
+		GenerateChildrenForSingleLayout(Layout, SimpleChildNodes);
 	}
 
-	// note: this can't be a ranged-for, because the map may have items added to it by customizations during iteration
 	for (int32 LayoutIndex = 0; LayoutIndex < LayoutMap.Num(); ++LayoutIndex)
 	{
 		const FDetailLayout& Layout = LayoutMap[LayoutIndex];
-		GenerateNodesFromCustomizations(Layout.GetAdvancedLayouts(), AdvancedChildNodes);
+		GenerateChildrenForSingleLayout(Layout, AdvancedChildNodes);
 	}
 
 	// Generate nodes for advanced dropdowns
diff --git a/Engine/Source/Editor/PropertyEditor/Private/DetailCategoryBuilderImpl.h b/Engine/Source/Editor/PropertyEditor/Private/DetailCategoryBuilderImpl.h
--- a/Engine/Source/Editor/PropertyEditor/Private/DetailCategoryBuilderImpl.h
+++ b/Engine/Source/Editor/PropertyEditor/Private/DetailCategoryBuilderImpl.h
@@ -324,14 +324,14 @@ private:
 	virtual void OnItemExpansionChanged(bool bIsExpanded, bool bShouldSaveState) override;
 
 	/**
-	 * Adds a new filter widget to this category (for checking if anything is visible in the category when filtered)
+	 * Generates children for all layouts.
 	 */
-	void AddFilterWidget(TSharedRef<SWidget> InWidget);
+	void GenerateChildrenForLayouts();
 
 	/**
-	 * Generates children for each layout
+	 * Generate children for the single given layout and append them to OutChildren.
 	 */
-	void GenerateChildrenForLayouts();
+	void GenerateChildrenForSingleLayout(const FDetailLayout& Layout, FDetailNodeList& OutChildren);
 
 	/**
 	 * Generates nodes from a list of customization in a single layout
@@ -393,9 +393,9 @@ private:
 	/** Layouts that appear in this category category */
 	FDetailLayoutMap LayoutMap;
 	/** All Simple child nodes */
-	TArray< TSharedRef<FDetailTreeNode> > SimpleChildNodes;
+	FDetailNodeList SimpleChildNodes;
 	/** All Advanced child nodes */
-	TArray< TSharedRef<FDetailTreeNode> > AdvancedChildNodes;
+	FDetailNodeList AdvancedChildNodes;
 	/** Advanced dropdown node. */
 	TSharedPtr<FDetailTreeNode> AdvancedDropdownNode;
 	/** Delegate called when expansion of the category changes */
