Siimple nodes were also being duplicated into the Advanced dropdown as a result of my previous fix.

diff --git a/Engine/Source/Editor/PropertyEditor/Private/DetailCategoryBuilderImpl.cpp b/Engine/Source/Editor/PropertyEditor/Private/DetailCategoryBuilderImpl.cpp
--- a/Engine/Source/Editor/PropertyEditor/Private/DetailCategoryBuilderImpl.cpp
+++ b/Engine/Source/Editor/PropertyEditor/Private/DetailCategoryBuilderImpl.cpp
@@ -942,10 +942,10 @@ void FDetailCategoryImpl::GenerateNodesFromCustomizations(const TArray<FDetailLa
 	}
 }
 
-void FDetailCategoryImpl::GenerateChildrenForSingleLayout(const FDetailLayout& Layout, FDetailNodeList& OutChildren)
+void FDetailCategoryImpl::GenerateChildrenForSingleLayout(const FDetailLayout& Layout, const TArray<FDetailLayoutCustomization>& Customizations, FDetailNodeList& OutChildren)
 {
 	FDetailNodeList GeneratedChildren;
-	GenerateNodesFromCustomizations(Layout.GetSimpleLayouts(), GeneratedChildren);
+	GenerateNodesFromCustomizations(Customizations, GeneratedChildren);
 
 	const FName InstanceName = Layout.GetInstanceName();
 	if (LayoutMap.ShouldShowGroup(InstanceName))
@@ -966,13 +966,13 @@ void FDetailCategoryImpl::GenerateChildrenForLayouts()
 	for (int32 LayoutIndex = 0; LayoutIndex < LayoutMap.Num(); ++LayoutIndex)
 	{
 		const FDetailLayout& Layout = LayoutMap[LayoutIndex];
-		GenerateChildrenForSingleLayout(Layout, SimpleChildNodes);
+		GenerateChildrenForSingleLayout(Layout, Layout.GetSimpleLayouts(), SimpleChildNodes);
 	}
 
 	for (int32 LayoutIndex = 0; LayoutIndex < LayoutMap.Num(); ++LayoutIndex)
 	{
 		const FDetailLayout& Layout = LayoutMap[LayoutIndex];
-		GenerateChildrenForSingleLayout(Layout, AdvancedChildNodes);
+		GenerateChildrenForSingleLayout(Layout, Layout.GetAdvancedLayouts(), AdvancedChildNodes);
 	}
 
 	// Generate nodes for advanced dropdowns
diff --git a/Engine/Source/Editor/PropertyEditor/Private/DetailCategoryBuilderImpl.h b/Engine/Source/Editor/PropertyEditor/Private/DetailCategoryBuilderImpl.h
--- a/Engine/Source/Editor/PropertyEditor/Private/DetailCategoryBuilderImpl.h
+++ b/Engine/Source/Editor/PropertyEditor/Private/DetailCategoryBuilderImpl.h
@@ -331,7 +331,7 @@ private:
 	/**
 	 * Generate children for the single given layout and append them to OutChildren.
 	 */
-	void GenerateChildrenForSingleLayout(const FDetailLayout& Layout, FDetailNodeList& OutChildren);
+	void GenerateChildrenForSingleLayout(const FDetailLayout& Layout, const TArray<FDetailLayoutCustomization>& Customizations, FDetailNodeList& OutChildren);
 
 	/**
 	 * Generates nodes from a list of customization in a single layout
