Pass PropertyHandle through MaterialList PropertyEntryBox so that AssetReferenceFilter can access the property source

diff --git a/Engine/Source/Editor/PropertyEditor/Private/MaterialList.cpp b/Engine/Source/Editor/PropertyEditor/Private/MaterialList.cpp
--- a/Engine/Source/Editor/PropertyEditor/Private/MaterialList.cpp
+++ b/Engine/Source/Editor/PropertyEditor/Private/MaterialList.cpp
@@ -183,6 +183,7 @@ TSharedRef<SWidget> FMaterialItemView::CreateValueContent(IDetailLayoutBuilder&
 			[
 				SNew(SObjectPropertyEntryBox)
 				.ObjectPath(this, &FMaterialItemView::OnGetObjectPath)
+				.PropertyHandle(GetMaterialHandle(InDetailBuilder, InActorComponent))
 				.AllowClear(false)
 				.AllowedClass(UMaterialInterface::StaticClass())
 				.OnObjectChanged(this, &FMaterialItemView::OnSetObject)
