SStructureDetailView: Fixed: setting custom name wasn't working

diff --git a/Engine/Source/Editor/PropertyEditor/Private/SStructureDetailsView.cpp b/Engine/Source/Editor/PropertyEditor/Private/SStructureDetailsView.cpp
--- a/Engine/Source/Editor/PropertyEditor/Private/SStructureDetailsView.cpp
+++ b/Engine/Source/Editor/PropertyEditor/Private/SStructureDetailsView.cpp
@@ -311,7 +311,19 @@ const TSharedPtr<class FComplexPropertyNode> SStructureDetailsView::GetRootNode(
 
 void SStructureDetailsView::CustomUpdatePropertyMap(TSharedPtr<FDetailLayoutBuilderImpl>& InDetailLayout)
 {
-	InDetailLayout->DefaultCategory(NAME_None).SetDisplayName(NAME_None, CustomName);
+	FName StructCategoryName = NAME_None;
+	
+	if (StructData.IsValid() && StructData->GetStruct())
+	{
+		TArray<FName> CategoryNames;
+		InDetailLayout->GetCategoryNames(CategoryNames);
+
+		int32 StructCategoryNameIndex = INDEX_NONE;
+		CategoryNames.Find(StructData->GetStruct()->GetFName(), StructCategoryNameIndex);
+		StructCategoryName = StructCategoryNameIndex != INDEX_NONE ? CategoryNames[StructCategoryNameIndex] : NAME_None;
+	}
+	
+	InDetailLayout->DefaultCategory(StructCategoryName).SetDisplayName(NAME_None, CustomName);
 }
 
 EVisibility SStructureDetailsView::GetPropertyEditingVisibility() const
