Allow the temporary used by PropagatePropertyChange to behave as if it had been loaded

Some property types may rely on PPF_SerializedAsImportText to allocate extra state that's required for accurate comparison against another instance

#rnx

diff --git a/Engine/Source/Editor/PropertyEditor/Private/PropertyNode.cpp b/Engine/Source/Editor/PropertyEditor/Private/PropertyNode.cpp
--- a/Engine/Source/Editor/PropertyEditor/Private/PropertyNode.cpp
+++ b/Engine/Source/Editor/PropertyEditor/Private/PropertyNode.cpp
@@ -3345,7 +3345,7 @@ void FPropertyNode::PropagatePropertyChange( UObject* ModifiedObject, const TCHA
 					// This ensures that shared state keeps the correct value, even if the destination property itself isn't imported (or only partly imported, as is the case with arrays/maps/sets)
 					FString CurrentValue;
 					ComplexProperty->ExportText_Direct(CurrentValue, ModifiedComplexPropAddr, ModifiedComplexPropAddr, ModifiedObject, PPF_None);
-					ComplexProperty->ImportText_Direct(*PreviousValue, TempComplexPropAddr, ModifiedObject, PPF_None);
+					ComplexProperty->ImportText_Direct(*PreviousValue, TempComplexPropAddr, ModifiedObject, PPF_SerializedAsImportText);
 					bShouldImport = ComplexProperty->Identical(DestComplexPropAddr, TempComplexPropAddr, PPF_DeepComparison);
 					ComplexProperty->ImportText_Direct(*CurrentValue, TempComplexPropAddr, ModifiedObject, PPF_None);
 				}
