UStructs can now provide a CanEditChange override by setting the TStructOpsTypeTraitsBase2::WithCanEditChange trait.

PR #8877: Contributed by [REDACTED_PERSON]

diff --git a/Engine/Source/Editor/PropertyEditor/Private/PropertyNode.cpp b/Engine/Source/Editor/PropertyEditor/Private/PropertyNode.cpp
--- a/Engine/Source/Editor/PropertyEditor/Private/PropertyNode.cpp
+++ b/Engine/Source/Editor/PropertyEditor/Private/PropertyNode.cpp
@@ -919,24 +919,65 @@ bool FPropertyNode::IsEditConst() const
 		const FObjectPropertyNode* ObjectPropertyNode = FindObjectItemParent();
 
 		bIsEditConst = IsPropertyConst();
-		if (!bIsEditConst && Property != nullptr && ObjectPropertyNode)
+		if (!bIsEditConst && Property.IsValid() && ObjectPropertyNode)
 		{
-			// travel up the chain to see if this property's owner struct is editconst - if it is, so is this property
-			FPropertyNode* NextParent = ParentNode;
-			while (NextParent != nullptr && CastField<FStructProperty>(NextParent->GetProperty()) != NULL)
+			TSharedRef<FEditPropertyChain> PropertyChain = BuildPropertyChain(Property.Get());
+			
+			// travel up the chain to see if this property's owner struct is EditConst - if it is, so is this property
+			FPropertyNode* CurParent = ParentNode;
+			while (CurParent != nullptr)
 			{
-				if (NextParent->IsEditConst())
+				FStructProperty* StructProperty = CastField<FStructProperty>(CurParent->GetProperty());
+				if (StructProperty == nullptr)
 				{
-						bIsEditConst = true;
-						break;
+					break;
+				}
+
+				if (CurParent->IsEditConst())
+				{
+					// An owning struct is edit const, so the child property is too
+					bIsEditConst = true;
+				}
+				else
+				{
+					// See if the struct has a problem with this property being editable
+					UScriptStruct* ScriptStruct = StructProperty->Struct;
+					if (ScriptStruct && ScriptStruct->StructFlags & STRUCT_CanEditChange)
+					{
+						UScriptStruct::ICppStructOps* TheCppStructOps = ScriptStruct->GetCppStructOps();
+						check(TheCppStructOps);
+
+						const int32 NumInstances = ObjectPropertyNode->GetInstancesNum();
+
+						TArray<const void*> StructAddresses;
+						StructAddresses.Reset(NumInstances);
+
+						for (int32 Index = 0; Index < NumInstances; ++Index)
+						{
+							StructAddresses.Add(CurParent->GetValueAddressFromObject(ObjectPropertyNode->GetUObject(Index)));
+						}
+
+						for (const void* StructAddr : StructAddresses)
+						{
+							if (!TheCppStructOps->CanEditChange(*PropertyChain, StructAddr))
+							{
+								bIsEditConst = true;
+								break;
+							}
+						}
+					}
 				}
-				NextParent = NextParent->ParentNode;
+
+				if (bIsEditConst)
+				{
+					break;
+				}
+				
+				CurParent = CurParent->ParentNode;
 			}
 
 			if (!bIsEditConst)
 			{
-				TSharedRef<FEditPropertyChain> PropertyChain = BuildPropertyChain(Property.Get());
-				
 				for (TPropObjectConstIterator CurObjectIt(ObjectPropertyNode->ObjectConstIterator()); CurObjectIt; ++CurObjectIt)
 				{
 					const TWeakObjectPtr<UObject> CurObject = *CurObjectIt;
diff --git a/Engine/Source/Runtime/CoreUObject/Private/UObject/Class.cpp b/Engine/Source/Runtime/CoreUObject/Private/UObject/Class.cpp
--- a/Engine/Source/Runtime/CoreUObject/Private/UObject/Class.cpp
+++ b/Engine/Source/Runtime/CoreUObject/Private/UObject/Class.cpp
@@ -2684,6 +2684,13 @@ void UScriptStruct::PrepareCppStructOps()
 		UE_LOG(LogClass, Verbose, TEXT("Native struct %s has native SerializeFromMismatchedTag."),*GetName());
 		StructFlags = EStructFlags(StructFlags | STRUCT_SerializeFromMismatchedTag);
 	}
+#if WITH_EDITOR
+	if (CppStructOps->HasCanEditChange())
+	{
+		UE_LOG(LogClass, Verbose, TEXT("Native struct %s has native CanEditChange."), *GetName());
+		StructFlags = EStructFlags(StructFlags | STRUCT_CanEditChange);
+	}
+#endif
 
 	check(!bPrepareCppStructOpsCompleted); // recursion is unacceptable
 	bPrepareCppStructOpsCompleted = true;
diff --git a/Engine/Source/Runtime/CoreUObject/Public/UObject/Class.h b/Engine/Source/Runtime/CoreUObject/Public/UObject/Class.h
--- a/Engine/Source/Runtime/CoreUObject/Public/UObject/Class.h
+++ b/Engine/Source/Runtime/CoreUObject/Public/UObject/Class.h
@@ -707,6 +707,9 @@ enum EStructFlags
 	/** If set, this structure has been replaced via reinstancing */
 	STRUCT_NewerVersionExists = 0x01000000,
 
+	/** If set, this struct will have CanEditChange on it in the editor to determine if a child property can be edited */
+	STRUCT_CanEditChange = 0x02000000,
+	
 	/** Struct flags that are automatically inherited */
 	STRUCT_Inherit				= STRUCT_HasInstancedReference|STRUCT_Atomic,
 
@@ -741,6 +744,7 @@ struct TStructOpsTypeTraitsBase2
 		WithNetSharedSerialization     = false,                         // struct has a NetSerialize function that does not require the package map to serialize its state.
 		WithGetPreloadDependencies     = false,                         // struct has a GetPreloadDependencies function to return all objects that will be Preload()ed when the struct is serialized at load time.
 		WithPureVirtual                = false,                         // struct has PURE_VIRTUAL functions and cannot be constructed when CHECK_PUREVIRTUALS is true
+		WithCanEditChange			   = false,							// struct has an editor-only CanEditChange function that can conditionally make child properties read-only in the details panel (same idea as UObject::CanEditChange)
 	};
 };
 
@@ -925,6 +929,14 @@ public:
 
 		/** return true if this struct is abstract **/
 		virtual bool IsAbstract() const = 0;
+
+#if WITH_EDITOR
+		/** Returns true if this struct wants to indicate whether a property can be edited in the details panel */
+		virtual bool HasCanEditChange() const = 0;
+		/** Returns true if this struct would allow the given property to be edited in the details panel. */
+		virtual bool CanEditChange(const FEditPropertyChain& PropertyChain, const void* Data) const = 0;
+#endif
+		
 	private:
 		/** sizeof() of the structure **/
 		const int32 Size;
@@ -1299,6 +1311,25 @@ public:
 		{
 			return TIsAbstract<CPPSTRUCT>::Value;
 		}
+
+#if WITH_EDITOR
+		virtual bool HasCanEditChange() const override
+		{
+			return TTraits::WithCanEditChange;
+		}
+
+		virtual bool CanEditChange(const FEditPropertyChain& PropertyChain, const void* Data) const override
+		{
+			if constexpr (TStructOpsTypeTraits<CPPSTRUCT>::WithCanEditChange)
+			{
+				return ((const CPPSTRUCT*)Data)->CanEditChange(PropertyChain);
+			}
+			else
+			{
+				return false;
+			}
+		}
+#endif // WITH_EDITOR
 	};
 
 	/** Template for noexport classes to autoregister before main starts **/
