We now cache the paths to the objects that we modify in FPropertyNode::NotifyPostEditChange() before we call CanEditChange(Chain)Property.

This causes all the modified objects to have their PostEditChangeProperty() handlers called even if the construction script of a blueprint class is triggered by the first invocation.

Most notably, this fixes the case where multiple components of an actor are multi-selected and edited at the same time, but only one of them will receive a PostEditChangeProperty() call.

Reported from UDN: https://epicgames.lightning.force.com/lightning/r/Case/5004z00001eEzehAAC/view

diff --git a/Engine/Source/Editor/PropertyEditor/Private/PropertyNode.cpp b/Engine/Source/Editor/PropertyEditor/Private/PropertyNode.cpp
--- a/Engine/Source/Editor/PropertyEditor/Private/PropertyNode.cpp
+++ b/Engine/Source/Editor/PropertyEditor/Private/PropertyNode.cpp
@@ -2556,58 +2556,57 @@ void FPropertyNode::NotifyPostChange( FPropertyChangedEvent& InPropertyChangedEv
 		// Call PostEditChange on the object chain.
 		while ( true )
 		{
-			TArray<TWeakObjectPtr<UObject>> Containers;
+			TArray<FString> ObjectPaths;
+			TArray<TWeakObjectPtr<UObject>> WeakObjects;
 			// It's possible that PostEditChangeProperty may cause a construction script to re-run
 			// which will invalidate the PropObjectIterator. We need to instead cache all of the objects
 			// before emitting any change events to ensure there is a PostChange for every PreChange.
 			for (TPropObjectIterator Itor(ObjectNode->ObjectIterator()); Itor; ++Itor)
 			{
-				Containers.Add(*Itor);
+				WeakObjects.Add(*Itor);
+				ObjectPaths.Add((*Itor)->GetPathName());
 			}
-			for (int32 CurrentObjectIndex = 0; CurrentObjectIndex < Containers.Num(); ++CurrentObjectIndex)
+
+			for (int32 CurrentObjectIndex = 0; CurrentObjectIndex < WeakObjects.Num(); ++CurrentObjectIndex)
 			{
-				UObject* Object = Containers[CurrentObjectIndex].Get();
+				UObject* Object = WeakObjects[CurrentObjectIndex].Get();
+				if (Object == nullptr)
+				{
+					// If our weak pointer has gone out of scope, it means that a prior object has destroyed it, 
+					// eg. by causing a blueprint construction script to run (which is triggered by PostEditChangeProperty())
+					// Find a new copy now.
+					Object = FindObject<UObject>(nullptr, *ObjectPaths[CurrentObjectIndex]);
+					if (Object == nullptr)
+					{
+						continue;
+					}
+				}
 
 				// Use a scope to ensure that only local variable are use in the loop.
-				//Since this object can be destroyed in this loop.
-				auto ScopePostEditChange = [&PropertyChain, &InPropertyChangedEvent, &CurProperty, &CurrentObjectIndex, &LevelDirtyCallback] (UObject* Object)
+				// Since this object can be destroyed in this loop.
+				auto ScopePostEditChange = [&PropertyChain, &InPropertyChangedEvent, &CurProperty, CurrentObjectIndex](UObject* Object)
 				{
-					// It is possible that the PostChange for the first object deletes the second object, so
-					// we must check to make sure the object still exists first.
-					if (Object)
+					// copy the property changed event
+					FPropertyChangedEvent ChangedEvent = CurProperty != InPropertyChangedEvent.Property ? 
+						FPropertyChangedEvent(CurProperty, InPropertyChangedEvent.ChangeType) : 
+						InPropertyChangedEvent;
+					ChangedEvent.ObjectIteratorIndex = CurrentObjectIndex;
+
+					if (PropertyChain->Num() == 0)
 					{
-						if (PropertyChain->Num() == 0)
-						{
-							//copy 
-							FPropertyChangedEvent ChangedEvent = InPropertyChangedEvent;
-							if (CurProperty != InPropertyChangedEvent.Property)
-							{
-								ChangedEvent = FPropertyChangedEvent(CurProperty, InPropertyChangedEvent.ChangeType);
-							}
-							ChangedEvent.ObjectIteratorIndex = CurrentObjectIndex;
-							if (Object)
-							{
-								Object->PostEditChangeProperty(ChangedEvent);
-							}
-						}
-						else
-						{
-							FPropertyChangedEvent ChangedEvent = InPropertyChangedEvent;
-							if (CurProperty != InPropertyChangedEvent.Property)
-							{
-								ChangedEvent = FPropertyChangedEvent(CurProperty, InPropertyChangedEvent.ChangeType);
-							}
-							FPropertyChangedChainEvent ChainEvent(*PropertyChain, ChangedEvent);
-							ChainEvent.ObjectIteratorIndex = CurrentObjectIndex;
-							if (Object)
-							{
-								Object->PostEditChangeChainProperty(ChainEvent);
-							}
-						}
-						LevelDirtyCallback.Request();
+						Object->PostEditChangeProperty(ChangedEvent);
+					}
+					else
+					{
+						FPropertyChangedChainEvent ChainEvent(*PropertyChain, ChangedEvent);
+						ChainEvent.ObjectIteratorIndex = CurrentObjectIndex;
+
+						Object->PostEditChangeChainProperty(ChainEvent);
 					}
 				};
+
 				ScopePostEditChange(Object);
+				LevelDirtyCallback.Request();
 			}
 
 			if (!ThisAsWeakPtr.IsValid())
