Fix Private Materials unable to be referenced within same Plugin via asset picker

diff --git a/Engine/Source/Editor/PropertyEditor/Private/MaterialList.cpp b/Engine/Source/Editor/PropertyEditor/Private/MaterialList.cpp
--- a/Engine/Source/Editor/PropertyEditor/Private/MaterialList.cpp
+++ b/Engine/Source/Editor/PropertyEditor/Private/MaterialList.cpp
@@ -168,6 +168,9 @@ namespace ValueExtender
 
 TSharedRef<SWidget> FMaterialItemView::CreateValueContent(IDetailLayoutBuilder& InDetailBuilder, const TArray<FAssetData>& OwnerAssetDataArray, UActorComponent* InActorComponent)
 {	
+	// Always consider the InActorComponent's asset location (BP, Level, etc.) as part of the OwnerAssetArray
+	TArray<FAssetData> AssetDataArray = OwnerAssetDataArray;
+	AssetDataArray.Add((FAssetData)InActorComponent->GetOuter());
 
 	return
 		SNew(SVerticalBox)
@@ -188,7 +191,7 @@ TSharedRef<SWidget> FMaterialItemView::CreateValueContent(IDetailLayoutBuilder&
 				.OnObjectChanged(this, &FMaterialItemView::OnSetObject)
 				.ThumbnailPool(InDetailBuilder.GetThumbnailPool())
 				.DisplayCompactSize(true)
-				.OwnerAssetDataArray(OwnerAssetDataArray)
+				.OwnerAssetDataArray(AssetDataArray)
 				.CustomContentSlot()
 				[
 					SNew( SBox )
