Fix Null access crash in MaterialList

diff --git a/Engine/Source/Editor/PropertyEditor/Private/MaterialList.cpp b/Engine/Source/Editor/PropertyEditor/Private/MaterialList.cpp
--- a/Engine/Source/Editor/PropertyEditor/Private/MaterialList.cpp
+++ b/Engine/Source/Editor/PropertyEditor/Private/MaterialList.cpp
@@ -170,7 +170,10 @@
 {	
 	// Always consider the InActorComponent's asset location (BP, Level, etc.) as part of the OwnerAssetArray
 	TArray<FAssetData> AssetDataArray = OwnerAssetDataArray;
-	AssetDataArray.Add((FAssetData)InActorComponent->GetOuter());
+	if (InActorComponent)
+	{
+		AssetDataArray.Add((FAssetData)InActorComponent->GetOuter());
+	}
 
 	return
 		SNew(SVerticalBox)
