TitleProperty values are now retrieved recursively. This was a regression from UE4 that caused issues in cases that should have worked, eg. an array of instanced objects, because the category node would cause the search to end.

FPropertyNode::FindChildNode is now actually breadth-first, rather than only being breadth-first for a single level. This was causing an issue whereby title properties were being incorrectly retrieved from the first child rather than the topmost child.

Reported from UDN: https://epicgames.lightning.force.com/lightning/r/Case/5004z00001eFeYlAAK/view

diff --git a/Engine/Source/Editor/PropertyEditor/Private/PropertyNode.cpp b/Engine/Source/Editor/PropertyEditor/Private/PropertyNode.cpp
--- a/Engine/Source/Editor/PropertyEditor/Private/PropertyNode.cpp
+++ b/Engine/Source/Editor/PropertyEditor/Private/PropertyNode.cpp
@@ -4,6 +4,7 @@
 #include "Misc/ConfigCacheIni.h"
 #include "Serialization/ArchiveReplaceObjectRef.h"
 #include "Components/ActorComponent.h"
+#include "Containers/Deque.h"
 #include "Editor/UnrealEdEngine.h"
 #include "Engine/UserDefinedStruct.h"
 #include "EditConditionContext.h"
@@ -868,24 +869,42 @@
 
 TSharedPtr<FPropertyNode> FPropertyNode::FindChildPropertyNode( const FName InPropertyName, bool bRecurse )
 {
-	// Search Children
-	for (const TSharedPtr<FPropertyNode>& ChildNode : ChildNodes)
+	// search children breadth-first, so that identically-named properties are first picked up in top-level classes, eg:
+	// class UFoo 
+	// { 
+	//    struct FBar
+	//    { 
+	//       int ID; 
+	//    } Bar;
+	//    int ID;
+	// };
+	// depth-first search would find FBar::ID before UFoo::ID when searching for "ID", which is rarely what was intended
+
+	TDeque<TSharedPtr<FPropertyNode>> NodesToSearch;
+
+	auto PushAll = [&NodesToSearch](const TArray<TSharedPtr<FPropertyNode>>& Nodes)
 	{
-		if( ChildNode->GetProperty() && ChildNode->GetProperty()->GetFName() == InPropertyName )
+		NodesToSearch.Reserve(NodesToSearch.Num() + Nodes.Num());
+		for (const TSharedPtr<FPropertyNode>& Node : Nodes)
 		{
-			return ChildNode;
+			NodesToSearch.PushLast(Node);
 		}
-	}
+	};
 
-	if (bRecurse)
+	PushAll(ChildNodes);
+	while (!NodesToSearch.IsEmpty())
 	{
-		for (const TSharedPtr<FPropertyNode>& ChildNode : ChildNodes)
+		TSharedPtr<FPropertyNode> Node = NodesToSearch.First();
+		NodesToSearch.PopFirst();
+
+		if (Node->GetProperty() && Node->GetProperty()->GetFName() == InPropertyName)
 		{
-			TSharedPtr<FPropertyNode> PropertyNode = ChildNode->FindChildPropertyNode(InPropertyName, bRecurse);
-			if (PropertyNode.IsValid())
-			{
-				return PropertyNode;
-			}
+			return Node;
+		}
+
+		if (bRecurse)
+		{
+			PushAll(Node->ChildNodes);
 		}
 	}
 
diff --git a/Engine/Source/Editor/PropertyEditor/Private/UserInterface/PropertyEditor/SPropertyEditorArrayItem.cpp b/Engine/Source/Editor/PropertyEditor/Private/UserInterface/PropertyEditor/SPropertyEditorArrayItem.cpp
--- a/Engine/Source/Editor/PropertyEditor/Private/UserInterface/PropertyEditor/SPropertyEditorArrayItem.cpp
+++ b/Engine/Source/Editor/PropertyEditor/Private/UserInterface/PropertyEditor/SPropertyEditorArrayItem.cpp
@@ -22,7 +22,7 @@
 			FText::GetFormatPatternParameters(TitleFormatter->Format, OutParameterNames);
 			for (const FString& ParameterName : OutParameterNames)
 			{
-				TSharedPtr<IPropertyHandle> Handle = RootProperty->GetChildHandle(FName(*ParameterName), false /*bRecurse*/);
+				TSharedPtr<IPropertyHandle> Handle = RootProperty->GetChildHandle(FName(*ParameterName), true);
 				if (Handle.IsValid())
 				{
 					TitleFormatter->PropertyHandles.Add(Handle);
@@ -38,7 +38,7 @@
 		}
 		else // Support the old style where it was just a name of a property with no formatting.
 		{
-			TSharedPtr<IPropertyHandle> Handle = RootProperty->GetChildHandle(FName(*TitlePropertyRaw), false /*bRecurse*/);
+			TSharedPtr<IPropertyHandle> Handle = RootProperty->GetChildHandle(FName(*TitlePropertyRaw), true);
 			if (Handle.IsValid())
 			{
 				TitleFormatter->PropertyHandles.Add(Handle);
