Fixing SPropertyEditorAsset not supporting class path names when filtering assets

diff --git a/Engine/Source/Editor/PropertyEditor/Private/UserInterface/PropertyEditor/SPropertyEditorAsset.cpp b/Engine/Source/Editor/PropertyEditor/Private/UserInterface/PropertyEditor/SPropertyEditorAsset.cpp
--- a/Engine/Source/Editor/PropertyEditor/Private/UserInterface/PropertyEditor/SPropertyEditorAsset.cpp
+++ b/Engine/Source/Editor/PropertyEditor/Private/UserInterface/PropertyEditor/SPropertyEditorAsset.cpp
@@ -275,8 +275,14 @@ bool SPropertyEditorAsset::IsAssetAllowed(const FAssetData& InAssetData)
 	{
 		for (const auto& RequiredTagAndValue : *RequiredAssetDataTags.Get())
 		{
-			if (!InAssetData.TagsAndValues.ContainsKeyValue(RequiredTagAndValue.Key, RequiredTagAndValue.Value))
+			if (!InAssetData.TagsAndValues.ContainsKeyValue(RequiredTagAndValue.Key, FPackageName::GetShortName(RequiredTagAndValue.Value)))
 			{
+				// For backwards compatibility compare against short name version of the tag value. This will be removed with ANY_PACKAGE removal
+				if (!FPackageName::IsShortPackageName(RequiredTagAndValue.Value) &&
+					InAssetData.TagsAndValues.ContainsKeyValue(RequiredTagAndValue.Key, FPackageName::ObjectPathToObjectName(RequiredTagAndValue.Value)))
+				{
+					continue;
+				}
 				return false;
 			}
 		}
