Duplicating array entries containing a null UObject*/TObjectPtr no longer crashes.

We were blindly assuming that CurrentObject was non-null and crashing on CurrentObject->GetOuter(), even though DuplicateObject() itself was safe to call with null.

diff --git a/Engine/Source/Editor/PropertyEditor/Private/PropertyNode.cpp b/Engine/Source/Editor/PropertyEditor/Private/PropertyNode.cpp
--- a/Engine/Source/Editor/PropertyEditor/Private/PropertyNode.cpp
+++ b/Engine/Source/Editor/PropertyEditor/Private/PropertyNode.cpp
@@ -3134,7 +3134,7 @@
 			UObject* CurrentObject = ObjProp->GetObjectPropertyValue(DestAddress);
 
 			// Make a deep copy
-			UObject* DuplicatedObject = DuplicateObject(CurrentObject, CurrentObject->GetOuter());
+			UObject* DuplicatedObject = DuplicateObject(CurrentObject, CurrentObject ? CurrentObject->GetOuter() : nullptr);
 			ObjProp->SetObjectPropertyValue(SrcAddress, DuplicatedObject);
 		}
 	}
