Do not mark package as dirty when editing CPF_SkipSerialization property.

diff --git a/Engine/Source/Editor/PropertyEditor/Private/PropertyHandleImpl.cpp b/Engine/Source/Editor/PropertyEditor/Private/PropertyHandleImpl.cpp
--- a/Engine/Source/Editor/PropertyEditor/Private/PropertyHandleImpl.cpp
+++ b/Engine/Source/Editor/PropertyEditor/Private/PropertyHandleImpl.cpp
@@ -420,7 +420,7 @@ FPropertyAccess::Result FPropertyValueImpl::ImportText( const TArray<FObjectBase
 			// If the values before and after setting the property differ, mark the object dirty.
 			if (FCString::Strcmp(*PreviousValue, *ValueAfterImport) != 0)
 			{
-				if (CurObject)
+				if (CurObject && !InPropertyNode->ShouldSkipSerialization())
 				{
 					CurObject->MarkPackageDirty();
 				}
diff --git a/Engine/Source/Editor/PropertyEditor/Private/PropertyNode.cpp b/Engine/Source/Editor/PropertyEditor/Private/PropertyNode.cpp
--- a/Engine/Source/Editor/PropertyEditor/Private/PropertyNode.cpp
+++ b/Engine/Source/Editor/PropertyEditor/Private/PropertyNode.cpp
@@ -1025,6 +1025,11 @@ bool FPropertyNode::IsEditConst() const
 	return bIsEditConst;
 }
 
+bool FPropertyNode::ShouldSkipSerialization() const
+{
+	return Property != nullptr && Property->HasAnyPropertyFlags(CPF_SkipSerialization);
+}
+
 bool FPropertyNode::HasEditCondition() const 
 { 
 	return EditConditionExpression.IsValid();
diff --git a/Engine/Source/Editor/PropertyEditor/Private/PropertyNode.h b/Engine/Source/Editor/PropertyEditor/Private/PropertyNode.h
--- a/Engine/Source/Editor/PropertyEditor/Private/PropertyNode.h
+++ b/Engine/Source/Editor/PropertyEditor/Private/PropertyNode.h
@@ -453,6 +453,11 @@ public:
 	/** @return whether this window's property is constant (can't be edited by the user) */
 	bool IsEditConst() const;
 
+	/**
+	 * Returns whether this window's property should not be serialized (determined by the CPF_SkipSerialization flag).
+	 */
+	bool ShouldSkipSerialization() const;
+
 	/**
 	 * Gets the full name of this node
 	 * @param PathPlusIndex - return value with full path of node
diff --git a/Engine/Source/Runtime/CoreUObject/Private/UObject/Obj.cpp b/Engine/Source/Runtime/CoreUObject/Private/UObject/Obj.cpp
--- a/Engine/Source/Runtime/CoreUObject/Private/UObject/Obj.cpp
+++ b/Engine/Source/Runtime/CoreUObject/Private/UObject/Obj.cpp
@@ -360,7 +360,7 @@ void UObject::PostLoad()
 #if WITH_EDITOR
 void UObject::PreEditChange(FProperty* PropertyAboutToChange)
 {
-	Modify(!GIsTransacting);
+	Modify(!GIsTransacting && !(PropertyAboutToChange && PropertyAboutToChange->HasAnyPropertyFlags(CPF_SkipSerialization)));
 }
 
 
