Fixed "Browse To" failing for an actor picker that didn't provide a PropertyEditor instance

This happens for things that are using a detail customization and creating the picker via a SObjectPropertyEntryBox, however they may provide an explicit PropertyHandle to use

This change uses the explicit property handle, if given, before falling back to trying to use the PropertyEditor instance to get the property handle

#rnx

diff --git a/Engine/Source/Editor/PropertyEditor/Private/UserInterface/PropertyEditor/SPropertyEditorAsset.cpp b/Engine/Source/Editor/PropertyEditor/Private/UserInterface/PropertyEditor/SPropertyEditorAsset.cpp
--- a/Engine/Source/Editor/PropertyEditor/Private/UserInterface/PropertyEditor/SPropertyEditorAsset.cpp
+++ b/Engine/Source/Editor/PropertyEditor/Private/UserInterface/PropertyEditor/SPropertyEditorAsset.cpp
@@ -1187,10 +1187,17 @@ void SPropertyEditorAsset::OnBrowse()
 	FObjectOrAssetData Value;
 	GetValue( Value );
 
-	if(PropertyEditor.IsValid() && Value.Object)
+	if (bIsActor)
 	{
-		// This code only works on loaded objects
-		FPropertyEditor::SyncToObjectsInNode(PropertyEditor->GetPropertyNode());		
+		TSharedPtr<IPropertyHandle> PropertyHandleToUse = GetMostSpecificPropertyHandle();
+		if (PropertyHandleToUse && Value.Object)
+		{
+			// This code only works on loaded objects
+			if (TSharedPtr<FPropertyNode> PropertyNodeToSync = PropertyHandleToUse->GetPropertyNode())
+			{
+				FPropertyEditor::SyncToObjectsInNode(PropertyNodeToSync);
+			}
+		}
 	}
 	else
 	{
@@ -1211,7 +1218,7 @@ FText SPropertyEditorAsset::GetOnBrowseToolTip() const
 		Args.Add(TEXT("Asset"), FText::AsCultureInvariant(Value.Object->GetName()));
 		if (bIsActor)
 		{
-			return FText::Format(LOCTEXT( "BrowseToAssetInViewport", "Select '{Asset}' in the viewport"), Args);
+			return FText::Format(LOCTEXT( "SelectSpecificActorInViewport", "Select '{Asset}' in the viewport"), Args);
 		}
 		else
 		{
@@ -1219,7 +1226,14 @@ FText SPropertyEditorAsset::GetOnBrowseToolTip() const
 		}
 	}
 	
-	return LOCTEXT( "BrowseToAssetInContentBrowser", "Browse to Asset in Content Browser");
+	if (bIsActor)
+	{
+		return LOCTEXT("SelectActorInViewport", "Select Actor in the viewport");
+	}
+	else
+	{
+		return LOCTEXT("BrowseToAssetInContentBrowser", "Browse to Asset in Content Browser");
+	}
 }
 
 void SPropertyEditorAsset::OnUse()
