Fixing a fix to SPropertyEditorAsset not supporting class path names when filtering assets (became broken now that all meta data tags use path names)

diff --git a/Engine/Source/Editor/PropertyEditor/Private/UserInterface/PropertyEditor/SPropertyEditorAsset.cpp b/Engine/Source/Editor/PropertyEditor/Private/UserInterface/PropertyEditor/SPropertyEditorAsset.cpp
--- a/Engine/Source/Editor/PropertyEditor/Private/UserInterface/PropertyEditor/SPropertyEditorAsset.cpp
+++ b/Engine/Source/Editor/PropertyEditor/Private/UserInterface/PropertyEditor/SPropertyEditorAsset.cpp
@@ -277,9 +277,9 @@ bool SPropertyEditorAsset::IsAssetAllowed(const FAssetData& InAssetData)
 	{
 		for (const auto& RequiredTagAndValue : *RequiredAssetDataTags.Get())
 		{
-			if (!InAssetData.TagsAndValues.ContainsKeyValue(RequiredTagAndValue.Key, FPackageName::GetShortName(RequiredTagAndValue.Value)))
+			if (!InAssetData.TagsAndValues.ContainsKeyValue(RequiredTagAndValue.Key, RequiredTagAndValue.Value))
 			{
-				// For backwards compatibility compare against short name version of the tag value. This will be removed with ANY_PACKAGE removal
+				// For backwards compatibility compare against short name version of the tag value.
 				if (!FPackageName::IsShortPackageName(RequiredTagAndValue.Value) &&
 					InAssetData.TagsAndValues.ContainsKeyValue(RequiredTagAndValue.Key, FPackageName::ObjectPathToObjectName(RequiredTagAndValue.Value)))
 				{
