Unified text formatting of double properties on set (the formatting is defined by Expose_TFormatSpecifier(double, "%f") in UnrealTypeTraits.h) : the format needs to match with the implementation of GetNumericPropertyValueToString for doubles because of the previous/new value comparison that is performed in FPropertyValueImpl::ImportText, which is string-based rather than value-based.

diff --git a/Engine/Source/Editor/PropertyEditor/Private/PropertyHandleImpl.cpp b/Engine/Source/Editor/PropertyEditor/Private/PropertyHandleImpl.cpp
--- a/Engine/Source/Editor/PropertyEditor/Private/PropertyHandleImpl.cpp
+++ b/Engine/Source/Editor/PropertyEditor/Private/PropertyHandleImpl.cpp
@@ -3548,8 +3548,10 @@ FPropertyAccess::Result FPropertyHandleDouble::SetValue( const double& NewValue,
 	FPropertyAccess::Result Res;
 	// Clamp the value from any meta data ranges stored on the property value
 	double FinalValue = ClampValueFromMetaData<double>( NewValue, *Implementation->GetPropertyNode() );
+	
+	FDoubleProperty* NumericProperty = CastFieldChecked<FDoubleProperty>(Implementation->GetPropertyNode()->GetProperty());
+	const FString ValueStr = NumericProperty->GetNumericPropertyValueToString(&FinalValue);
 
-	const FString ValueStr = FString::Printf( TEXT("%.12f"), FinalValue );
 	Res = Implementation->ImportText( ValueStr, Flags );
 
 	return Res;
