Fix for crash when calling SDetailsView::ForceRefresh while already performing a refresh

diff --git a/Engine/Source/Editor/PropertyEditor/Private/SDetailsView.cpp b/Engine/Source/Editor/PropertyEditor/Private/SDetailsView.cpp
--- a/Engine/Source/Editor/PropertyEditor/Private/SDetailsView.cpp
+++ b/Engine/Source/Editor/PropertyEditor/Private/SDetailsView.cpp
@@ -19,6 +19,7 @@
 #include "PropertyEditorHelpers.h"
 #include "SDetailNameArea.h"
 #include "Styling/StyleColors.h"
+#include "Templates/UnrealTemplate.h"
 #include "UserInterface/PropertyDetails/PropertyDetailsUtilities.h"
 #include "Widgets/Colors/SColorPicker.h"
 #include "Widgets/Images/SImage.h"
@@ -73,6 +74,7 @@ void SDetailsView::Construct(const FArguments& InArgs, const FDetailsViewArgs& I
 	SetClassViewerFilters(InDetailsViewArgs.ClassViewerFilters);
 
 	bViewingClassDefaultObject = false;
+	bIsRefreshing = false;
 
 	PropertyUtilities = MakeShareable( new FPropertyDetailsUtilities( *this ) );
 	PropertyGenerationUtilities = MakeShareable( new FDetailsViewPropertyGenerationUtilities(*this) );
@@ -676,6 +678,14 @@ int32 SDetailsView::GetNumObjects() const
 
 void SDetailsView::SetObjectArrayPrivate(const TArray<UObject*>& InObjects)
 {
+	if (bIsRefreshing)
+	{
+		// SetObjectArrayPrivate can not be performed while in the middle of a previous refresh
+		return;
+	}
+
+	TGuardValue<bool> RefreshGuard(bIsRefreshing, true);
+
 	double StartTime = FPlatformTime::Seconds();
 
 	const TArray<FDetailsViewObjectRoot> Roots = ObjectFilter->FilterObjects(InObjects);
diff --git a/Engine/Source/Editor/PropertyEditor/Private/SDetailsView.h b/Engine/Source/Editor/PropertyEditor/Private/SDetailsView.h
--- a/Engine/Source/Editor/PropertyEditor/Private/SDetailsView.h
+++ b/Engine/Source/Editor/PropertyEditor/Private/SDetailsView.h
@@ -220,4 +220,6 @@ private:
 	FDelegateHandle PostUndoRedoDelegateHandle;
 	/** The section selector widget to show if DetailsViewArgs.bShowSectionSelector is true. */
 	TSharedPtr<class SWrapBox> SectionSelectorBox;
+	/** True when currently in the middle of a refreshing the object array */
+	bool bIsRefreshing;
 };
