Added NoEditInline property metadata to disable inline editing of properties in the details view (EditInline metadata/EditInlineNew classes)

diff --git a/Engine/Source/Editor/PropertyEditor/Private/PropertyNode.cpp b/Engine/Source/Editor/PropertyEditor/Private/PropertyNode.cpp
--- a/Engine/Source/Editor/PropertyEditor/Private/PropertyNode.cpp
+++ b/Engine/Source/Editor/PropertyEditor/Private/PropertyNode.cpp
@@ -189,9 +189,10 @@ void FPropertyNode::InitNode(const FPropertyNodeInitParams& InitParams)
 		// a pointer to the object, you see the object's properties.
 		static const FName Name_EditInline("EditInline");
 		static const FName Name_ShowInnerProperties("ShowInnerProperties");
+		static const FName Name_NoEditInline("NoEditInline");
 
 		// we are EditInlineNew if this property has the flag, or if inside a container that has the flag.
-		bIsEditInlineNew = GotReadAddresses && bIsObjectOrInterface &&
+		bIsEditInlineNew = GotReadAddresses && bIsObjectOrInterface && !MyProperty->HasMetaData(Name_NoEditInline) && 
 			(MyProperty->HasMetaData(Name_EditInline) || (bIsInsideContainer && OwnerProperty->HasMetaData(Name_EditInline)));
 		bShowInnerObjectProperties = bIsObjectOrInterface && MyProperty->HasMetaData(Name_ShowInnerProperties);
 
diff --git a/Engine/Source/Runtime/CoreUObject/Public/UObject/ObjectMacros.h b/Engine/Source/Runtime/CoreUObject/Public/UObject/ObjectMacros.h
--- a/Engine/Source/Runtime/CoreUObject/Public/UObject/ObjectMacros.h
+++ b/Engine/Source/Runtime/CoreUObject/Public/UObject/ObjectMacros.h
@@ -1371,6 +1371,9 @@ namespace UM
 		/// [PropertyMetadata] Property wont have a 'reset to default' button when displayed in property windows
 		NoResetToDefault,
 
+		/// [PropertyMetadata] Object properties pointing to an UObject instance whos class is marked editinline will not show their properties inline in property windows. Useful for getting actor components to appear in the component tree but not inline in the root actor details panel.
+		NoEditInline,
+
 		/// [PropertyMetadata] Used for integer and float properties. Indicates that the spin box element of the number editing widget should not be displayed.
 		NoSpinbox,
 
