Removed AddExternalObjectProperty overload that was deprecated in 4.24.

Removed inline forward declarations that are against the coding standard.

diff --git a/Engine/Source/Editor/PropertyEditor/Public/IDetailChildrenBuilder.h b/Engine/Source/Editor/PropertyEditor/Public/IDetailChildrenBuilder.h
--- a/Engine/Source/Editor/PropertyEditor/Public/IDetailChildrenBuilder.h
+++ b/Engine/Source/Editor/PropertyEditor/Public/IDetailChildrenBuilder.h
@@ -5,11 +5,14 @@
 #include "CoreMinimal.h"
 #include "DetailBuilderTypes.h"
 
-class SWidget;
-class IPropertyHandle;
+class FDetailWidgetRow;
+class FStructOnScope;
 class IDetailCategoryBuilder;
 class IDetailCustomNodeBuilder;
-class FStructOnScope;
+class IDetailGroup;
+class IDetailPropertyRow;
+class IPropertyHandle;
+class SWidget;
 
 /**
  * Builder for adding children to a detail customization
@@ -34,7 +37,7 @@ public:
 	 * @param LocalizedDisplayName	The display name of the group
 	 * @param true if the group should appear in the advanced section of the category
 	 */
-	virtual class IDetailGroup& AddGroup(FName GroupName, const FText& LocalizedDisplayName) = 0;
+	virtual IDetailGroup& AddGroup(FName GroupName, const FText& LocalizedDisplayName) = 0;
 
 	/**
 	 * Adds new custom content as a child to the struct
@@ -43,7 +46,7 @@ public:
 	 * @return A row that accepts widgets
 	 */
 
-	virtual class FDetailWidgetRow& AddCustomRow(const FText& SearchString) = 0;
+	virtual FDetailWidgetRow& AddCustomRow(const FText& SearchString) = 0;
 
 	/**
 	 * Adds a property to the struct
@@ -51,7 +54,7 @@ public:
 	 * @param PropertyHandle	The handle to the property to add
 	 * @return An interface to the property row that can be used to customize the appearance of the property
 	 */
-	virtual class IDetailPropertyRow& AddProperty(TSharedRef<IPropertyHandle> PropertyHandle) = 0;
+	virtual IDetailPropertyRow& AddProperty(TSharedRef<IPropertyHandle> PropertyHandle) = 0;
 
 	/**
 	 * Adds a set of objects to as a child.  Similar to details panels, all objects will be visible in the details panel as set of properties from the common base class from the list of objects
@@ -60,7 +63,7 @@ public:
 	 * @param  UniqueIdName		Optional identifier that uniquely identifies this object among other objects of the same type.  If this is empty, saving and restoring expansion state of this object may not work
 	 * @return The header row generated for this set of objects by the details panel
 	 */
-	virtual class IDetailPropertyRow* AddExternalObjects(const TArray<UObject*>& Objects, FName UniqueIdName = NAME_None) = 0;
+	virtual IDetailPropertyRow* AddExternalObjects(const TArray<UObject*>& Objects, FName UniqueIdName = NAME_None) = 0;
 
 	/**
 	 * Adds a set of objects to as a child.  Similar to details panels, all objects will be visible in the details panel as set of properties from the common base class from the list of objects
@@ -68,8 +71,8 @@ public:
 	 * @param  Objects			The objects to add
 	 * @param  PropertyName		Name of a property inside the object(s) to add.
 	 * @param  UniqueIdName		Optional identifier that uniquely identifies this object among other objects of the same type.  If this is empty, saving and restoring expansion state of this object may not work
-	 * @param  bAllowChildrenOverride Allows customization of how the new root property node is expanded when this is added. 
-	 * @param  bCreateCategoryNodesOverride Allows customization of how the new root node's category is displayed (or not). 
+	 * @param  bAllowChildrenOverride Allows customization of how the new root property node is expanded when this is added.
+	 * @param  bCreateCategoryNodesOverride Allows customization of how the new root node's category is displayed (or not).
 	 * @return The header row generated for this set of objects by the details panel
 	 */
 	virtual IDetailPropertyRow* AddExternalObjectProperty(const TArray<UObject*>& Objects, FName PropertyName, const FAddPropertyParams& Params = FAddPropertyParams()) = 0;
@@ -90,7 +93,7 @@ public:
 	 * @param ChildStructure	The structure to add
 	 * @param UniqueIdName		Optional identifier that uniquely identifies this structure among other structures of the same type.  If this is empty, saving and restoring expansion state of this structure may not work
 	 */
-	virtual class IDetailPropertyRow* AddExternalStructure(TSharedRef<FStructOnScope> ChildStructure, FName UniqueIdName = NAME_None) = 0;
+	virtual IDetailPropertyRow* AddExternalStructure(TSharedRef<FStructOnScope> ChildStructure, FName UniqueIdName = NAME_None) = 0;
 
 	/**
 	 * Adds all the properties of an external structure as a children
@@ -117,25 +120,4 @@ public:
 	* @return the parent group on the customized object that this children is in (if there is one)
 	*/
 	virtual IDetailGroup* GetParentGroup() const = 0;
-
-	UE_DEPRECATED(4.24, "Please use the overload that accepts an FAddPropertyParams structure")
-	IDetailPropertyRow* AddExternalObjectProperty(const TArray<UObject*>& Objects, FName PropertyName, FName UniqueIdName = NAME_None, TOptional<bool> bAllowChildrenOverride = TOptional<bool>(), TOptional<bool> bCreateCategoryNodesOverride = TOptional<bool>())
-	{
-		FAddPropertyParams Params;
-		Params.UniqueId(UniqueIdName);
-		if (bAllowChildrenOverride.IsSet())
-		{
-			Params.AllowChildren(bAllowChildrenOverride.GetValue());
-		}
-		if (bCreateCategoryNodesOverride.IsSet())
-		{
-			Params.CreateCategoryNodes(bCreateCategoryNodesOverride.GetValue());
-		}
-		return AddExternalObjectProperty(Objects, PropertyName, Params);
-	}
-	UE_DEPRECATED(4.24, "Please use the overload that accepts an FAddPropertyParams structure")
-	IDetailPropertyRow* AddExternalStructureProperty(TSharedRef<FStructOnScope> ChildStructure, FName PropertyName, FName UniqueIdName)
-	{
-		return AddExternalStructureProperty(ChildStructure, PropertyName, FAddPropertyParams().UniqueId(UniqueIdName));
-	}
 };
\ No newline at end of file
