IDetailLayoutBuilder no longer returns a const TSharedRef, unlike all other calls in the API.

diff --git a/Engine/Source/Editor/PropertyEditor/Private/DetailLayoutBuilderImpl.cpp b/Engine/Source/Editor/PropertyEditor/Private/DetailLayoutBuilderImpl.cpp
--- a/Engine/Source/Editor/PropertyEditor/Private/DetailLayoutBuilderImpl.cpp
+++ b/Engine/Source/Editor/PropertyEditor/Private/DetailLayoutBuilderImpl.cpp
@@ -852,7 +852,7 @@ void FDetailLayoutBuilderImpl::GetStructsBeingCustomized( TArray< TSharedPtr<FSt
 	}
 }
 
-const TSharedRef< IPropertyUtilities > FDetailLayoutBuilderImpl::GetPropertyUtilities() const
+TSharedRef< IPropertyUtilities > FDetailLayoutBuilderImpl::GetPropertyUtilities() const
 {
 	return PropertyDetailsUtilities.Pin().ToSharedRef();
 }
diff --git a/Engine/Source/Editor/PropertyEditor/Private/DetailLayoutBuilderImpl.h b/Engine/Source/Editor/PropertyEditor/Private/DetailLayoutBuilderImpl.h
--- a/Engine/Source/Editor/PropertyEditor/Private/DetailLayoutBuilderImpl.h
+++ b/Engine/Source/Editor/PropertyEditor/Private/DetailLayoutBuilderImpl.h
@@ -47,7 +47,7 @@ public:
 	virtual bool IsPropertyVisible(TSharedRef<IPropertyHandle> PropertyHandle) const override;
 	virtual bool IsPropertyVisible(const struct FPropertyAndParent& PropertyAndParent) const override;
 	virtual void HideCategory(FName CategoryName) override;
-	virtual const TSharedRef<IPropertyUtilities> GetPropertyUtilities() const override;
+	virtual TSharedRef<IPropertyUtilities> GetPropertyUtilities() const override;
 	virtual UClass* GetBaseClass() const override;
 	virtual const TArray<TWeakObjectPtr<UObject>>& GetSelectedObjects() const override;
 	virtual bool HasClassDefaultObject() const override;
diff --git a/Engine/Source/Editor/PropertyEditor/Public/DetailLayoutBuilder.h b/Engine/Source/Editor/PropertyEditor/Public/DetailLayoutBuilder.h
--- a/Engine/Source/Editor/PropertyEditor/Public/DetailLayoutBuilder.h
+++ b/Engine/Source/Editor/PropertyEditor/Public/DetailLayoutBuilder.h
@@ -11,6 +11,7 @@
 
 class IDetailCategoryBuilder;
 class IDetailsView;
+class IPropertyUtilities;
 
 namespace ECategoryPriority
 {
@@ -92,7 +93,7 @@ public:
 	/**
 	 *	@return the utilities various widgets need access to certain features of PropertyDetails
 	 */
-	virtual const TSharedRef< class IPropertyUtilities > GetPropertyUtilities() const = 0; 
+	virtual TSharedRef<IPropertyUtilities> GetPropertyUtilities() const = 0; 
 
 
 	/**
