Added doc comment to FDetailLayoutBuilderImpl::GetAllRootTreeNodes().

#rnx

diff --git a/Engine/Source/Editor/PropertyEditor/Private/DetailLayoutBuilderImpl.h b/Engine/Source/Editor/PropertyEditor/Private/DetailLayoutBuilderImpl.h
--- a/Engine/Source/Editor/PropertyEditor/Private/DetailLayoutBuilderImpl.h
+++ b/Engine/Source/Editor/PropertyEditor/Private/DetailLayoutBuilderImpl.h
@@ -122,6 +122,9 @@ public:
 	 */
 	FDetailNodeList& GetFilteredRootTreeNodes() { return FilteredRootTreeNodes; }
 
+	/**
+	 * @return All root tree nodes, regardless of visibility.
+	 */
 	FDetailNodeList& GetAllRootTreeNodes() { return AllRootTreeNodes; }
 
 	/**
