GetParentLayoutImpl() now returns a TSharedPtr rather than requiring an IsParentLayoutValid() call first.

IDetailsView::IsCustomRowVisible is now called using the new FDetailPropertyRow::GetRowName() method instead of having to check the type.

#rnx

diff --git a/Engine/Source/Editor/PropertyEditor/Private/DetailCategoryBuilderImpl.cpp b/Engine/Source/Editor/PropertyEditor/Private/DetailCategoryBuilderImpl.cpp
--- a/Engine/Source/Editor/PropertyEditor/Private/DetailCategoryBuilderImpl.cpp
+++ b/Engine/Source/Editor/PropertyEditor/Private/DetailCategoryBuilderImpl.cpp
@@ -219,11 +219,11 @@ void FDetailCategoryImpl::AddCustomBuilder(TSharedRef<IDetailCustomNodeBuilder>
 	{
 		if (GetDetailsView()->IsCustomBuilderFavorite(NewCustomization.CustomBuilderRow->GetOriginalPath()))
 		{
-			TSharedPtr<FDetailLayoutBuilderImpl> LayoutBuilder = DetailLayoutBuilder.Pin();
-			if (LayoutBuilder.IsValid())
+			TSharedPtr<FDetailLayoutBuilderImpl> ParentLayout = GetParentLayoutImpl();
+			if (ParentLayout.IsValid())
 			{
 				static const FName FavoritesCategoryName(TEXT("Favorites"));
-				FDetailCategoryImpl& FavoritesCategory = LayoutBuilder->DefaultCategory(FavoritesCategoryName);
+				FDetailCategoryImpl& FavoritesCategory = ParentLayout->DefaultCategory(FavoritesCategoryName);
 
 				FDetailLayoutCustomization FavoritesCustomization;
 				FavoritesCustomization.bCustom = true;
@@ -265,7 +265,7 @@ int32 FDetailCategoryImpl::GetNumCustomizations() const
 
 void FDetailCategoryImpl::GetDefaultProperties(TArray<TSharedRef<IPropertyHandle> >& OutDefaultProperties, bool bSimpleProperties, bool bAdvancedProperties)
 {
-	FDetailLayoutBuilderImpl& DetailLayoutBuilderRef = GetParentLayoutImpl();
+	TSharedPtr<FDetailLayoutBuilderImpl> ParentLayout = GetParentLayoutImpl();
 	for (const FDetailLayout& Layout : LayoutMap)
 	{
 		if (bSimpleProperties)
@@ -276,7 +276,7 @@ void FDetailCategoryImpl::GetDefaultProperties(TArray<TSharedRef<IPropertyHandle
 				{
 					const TSharedPtr<FPropertyNode>& Node = Customization.GetPropertyNode();
 
-					TSharedRef<IPropertyHandle> PropertyHandle = DetailLayoutBuilderRef.GetPropertyHandle(Node);
+					TSharedRef<IPropertyHandle> PropertyHandle = ParentLayout->GetPropertyHandle(Node);
 					if (PropertyHandle->IsValidHandle())
 					{
 						OutDefaultProperties.Add(PropertyHandle);
@@ -292,7 +292,7 @@ void FDetailCategoryImpl::GetDefaultProperties(TArray<TSharedRef<IPropertyHandle
 				if (Customization.HasPropertyNode())
 				{
 					const TSharedPtr<FPropertyNode>& Node = Customization.GetPropertyNode();
-					TSharedRef<IPropertyHandle> PropertyHandle = DetailLayoutBuilderRef.GetPropertyHandle(Node);
+					TSharedRef<IPropertyHandle> PropertyHandle = ParentLayout->GetPropertyHandle(Node);
 					if (PropertyHandle->IsValidHandle())
 					{
 						OutDefaultProperties.Add(PropertyHandle);
@@ -313,9 +313,11 @@ void FDetailCategoryImpl::SetCategoryVisibility(bool bIsVisible)
 		{
 			GetDetailsView()->RerunCurrentFilter();
 		}
-		if (DetailLayoutBuilder.IsValid())
+
+		TSharedPtr<FDetailLayoutBuilderImpl> ParentLayout = GetParentLayoutImpl();
+		if (ParentLayout.IsValid())
 		{
-			DetailLayoutBuilder.Pin()->NotifyNodeVisibilityChanged();
+			ParentLayout->NotifyNodeVisibilityChanged();
 		}
 	}
 }
@@ -350,10 +352,11 @@ IDetailPropertyRow& FDetailCategoryImpl::AddProperty(FName PropertyPath, UClass*
 	FDetailLayoutCustomization NewCustomization;
 	NewCustomization.bCustom = true;
 
-	TSharedPtr<FPropertyNode> PropertyNode = GetParentLayoutImpl().GetPropertyNode(PropertyPath, ClassOutermost, InstanceName);
+	TSharedPtr<FDetailLayoutBuilderImpl> ParentLayout = GetParentLayoutImpl();
+	TSharedPtr<FPropertyNode> PropertyNode = ParentLayout->GetPropertyNode(PropertyPath, ClassOutermost, InstanceName);
 	if (PropertyNode.IsValid())
 	{
-		GetParentLayoutImpl().SetCustomProperty(PropertyNode);
+		ParentLayout->SetCustomProperty(PropertyNode);
 	}
 
 	NewCustomization.PropertyRow = MakeShareable(new FDetailPropertyRow(PropertyNode, AsShared()));
@@ -369,11 +372,11 @@ IDetailPropertyRow& FDetailCategoryImpl::AddProperty(TSharedPtr<IPropertyHandle>
 	FDetailLayoutCustomization NewCustomization;
 	NewCustomization.bCustom = true;
 
-	TSharedPtr<FPropertyNode> PropertyNode = GetParentLayoutImpl().GetPropertyNode(PropertyHandle);
-
+	TSharedPtr<FDetailLayoutBuilderImpl> ParentLayout = GetParentLayoutImpl();
+	TSharedPtr<FPropertyNode> PropertyNode = ParentLayout->GetPropertyNode(PropertyHandle);
 	if (PropertyNode.IsValid())
 	{
-		GetParentLayoutImpl().SetCustomProperty(PropertyNode);
+		ParentLayout->SetCustomProperty(PropertyNode);
 	}
 
 	NewCustomization.PropertyRow = MakeShareable(new FDetailPropertyRow(PropertyNode, AsShared()));
@@ -473,12 +476,13 @@ TArray<TSharedPtr<IPropertyHandle>> FDetailCategoryImpl::AddAllExternalStructure
 
 	TArray<TSharedPtr<IPropertyHandle>> Handles;
 
-	FDetailLayoutBuilderImpl& DetailLayoutBuilderRef = GetParentLayoutImpl();
 
 	if (RootPropertyNode.IsValid())
 	{
 		RootPropertyNode->RebuildChildren();
-		DetailLayoutBuilder.Pin()->AddExternalRootPropertyNode(RootPropertyNode.ToSharedRef());
+
+		TSharedPtr<FDetailLayoutBuilderImpl> ParentLayout = GetParentLayoutImpl();
+		ParentLayout->AddExternalRootPropertyNode(RootPropertyNode.ToSharedRef());
 
 		for (int32 ChildIdx = 0; ChildIdx < RootPropertyNode->GetNumChildNodes(); ++ChildIdx)
 		{
@@ -490,7 +494,7 @@ TArray<TSharedPtr<IPropertyHandle>> FDetailCategoryImpl::AddAllExternalStructure
 				NewCustomization.bAdvanced = Location == EPropertyLocation::Advanced;
 				AddDefaultLayout(NewCustomization, NAME_None);
 
-				Handles.Add(DetailLayoutBuilderRef.GetPropertyHandle(PropertyNode));
+				Handles.Add(ParentLayout->GetPropertyHandle(PropertyNode));
 				if (OutPropertiesRow)
 				{
 					OutPropertiesRow->Add(NewCustomization.PropertyRow.Get());
@@ -524,7 +528,7 @@ bool FDetailCategoryImpl::IsAdvancedLayout(const FDetailLayoutCustomization& Lay
 void FDetailCategoryImpl::AddCustomLayout(const FDetailLayoutCustomization& LayoutInfo)
 {
 	ensure(LayoutInfo.bCustom == true);
-	GetLayoutForInstance(GetParentLayoutImpl().GetCurrentCustomizationVariableName()).AddLayout(LayoutInfo);
+	GetLayoutForInstance(GetParentLayoutImpl()->GetCurrentCustomizationVariableName()).AddLayout(LayoutInfo);
 }
 
 void FDetailCategoryImpl::AddDefaultLayout(const FDetailLayoutCustomization& LayoutInfo, FName InstanceName)
@@ -550,7 +554,7 @@ void FDetailCategoryImpl::OnAdvancedDropdownClicked()
 
 FDetailLayoutCustomization* FDetailCategoryImpl::GetDefaultCustomization(TSharedRef<FPropertyNode> PropertyNode)
 {
-	FDetailLayout& Layout = GetLayoutForInstance(GetParentLayoutImpl().GetCurrentCustomizationVariableName());
+	FDetailLayout& Layout = GetLayoutForInstance(GetParentLayoutImpl()->GetCurrentCustomizationVariableName());
 	
 	FDetailLayoutCustomization* Customization = Layout.GetDefaultLayout(PropertyNode);
 	return Customization;
@@ -598,11 +602,11 @@ void FDetailCategoryImpl::RefreshTree(bool bRefilterCategory)
 {
 	if (bRefilterCategory)
 	{
-		TSharedPtr<FDetailLayoutBuilderImpl> DetailLayoutBuilderPtr = DetailLayoutBuilder.Pin();
-		if (DetailLayoutBuilderPtr.IsValid())
+		TSharedPtr<FDetailLayoutBuilderImpl> ParentLayout = GetParentLayoutImpl();
+		if (ParentLayout.IsValid())
 		{
-			FilterNode(DetailLayoutBuilderPtr->GetCurrentFilter());
-			DetailLayoutBuilderPtr->GetPropertyGenerationUtilities().RebuildTreeNodes();
+			FilterNode(ParentLayout->GetCurrentFilter());
+			ParentLayout->GetPropertyGenerationUtilities().RebuildTreeNodes();
 		}
 	}
 	else
@@ -616,26 +620,26 @@ void FDetailCategoryImpl::RefreshTree(bool bRefilterCategory)
 
 void FDetailCategoryImpl::AddTickableNode(FDetailTreeNode& TickableNode)
 {
-	TSharedPtr<FDetailLayoutBuilderImpl> DetailLayoutBuilderPtr = DetailLayoutBuilder.Pin();
-	if (DetailLayoutBuilderPtr.IsValid())
+	TSharedPtr<FDetailLayoutBuilderImpl> ParentLayout = GetParentLayoutImpl();
+	if (ParentLayout.IsValid())
 	{
-		DetailLayoutBuilderPtr->AddTickableNode(TickableNode);
+		ParentLayout->AddTickableNode(TickableNode);
 	}
 }
 
 void FDetailCategoryImpl::RemoveTickableNode(FDetailTreeNode& TickableNode)
 {
-	TSharedPtr<FDetailLayoutBuilderImpl> DetailLayoutBuilderPtr = DetailLayoutBuilder.Pin();
-	if (DetailLayoutBuilderPtr.IsValid())
+	TSharedPtr<FDetailLayoutBuilderImpl> ParentLayout = GetParentLayoutImpl();
+	if (ParentLayout.IsValid())
 	{
-		DetailLayoutBuilderPtr->RemoveTickableNode(TickableNode);
+		ParentLayout->RemoveTickableNode(TickableNode);
 	}
 }
 
 void FDetailCategoryImpl::SaveExpansionState(FDetailTreeNode& InTreeNode)
 {
-	TSharedPtr<FDetailLayoutBuilderImpl> DetailLayoutBuilderPtr = DetailLayoutBuilder.Pin();
-	if (DetailLayoutBuilderPtr.IsValid())
+	TSharedPtr<FDetailLayoutBuilderImpl> ParentLayout = GetParentLayoutImpl();
+	if (ParentLayout.IsValid())
 	{
 		bool bIsExpanded = InTreeNode.ShouldBeExpanded();
 
@@ -643,20 +647,20 @@ void FDetailCategoryImpl::SaveExpansionState(FDetailTreeNode& InTreeNode)
 		Key += TEXT(".");
 		Key += InTreeNode.GetNodeName().ToString();
 
-		DetailLayoutBuilderPtr->SaveExpansionState(Key, bIsExpanded);
+		ParentLayout->SaveExpansionState(Key, bIsExpanded);
 	}
 }
 
 bool FDetailCategoryImpl::GetSavedExpansionState(FDetailTreeNode& InTreeNode) const
 {
-	TSharedPtr<FDetailLayoutBuilderImpl> DetailLayoutBuilderPtr = DetailLayoutBuilder.Pin();
-	if (DetailLayoutBuilderPtr.IsValid())
+	TSharedPtr<FDetailLayoutBuilderImpl> ParentLayout = GetParentLayoutImpl();
+	if (ParentLayout.IsValid())
 	{
 		FString Key = CategoryPathName;
 		Key += TEXT(".");
 		Key += InTreeNode.GetNodeName().ToString();
 
-		return DetailLayoutBuilderPtr->GetSavedExpansionState(Key);
+		return ParentLayout->GetSavedExpansionState(Key);
 	}
 
 	return false;
@@ -699,7 +703,7 @@ void FDetailCategoryImpl::SetDisplayName(FName InCategoryName, const FText& Loca
 	else
 	{
 		// Use the base class name if there is one otherwise this is a generic category not specific to a class
-		const UStruct* BaseStruct = GetParentLayoutImpl().GetRootNode()->GetBaseStructure();
+		const UStruct* BaseStruct = GetParentLayoutImpl()->GetRootNode()->GetBaseStructure();
 		if (BaseStruct)
 		{
 			DisplayName = BaseStruct->GetDisplayNameText();
@@ -713,10 +717,10 @@ void FDetailCategoryImpl::SetDisplayName(FName InCategoryName, const FText& Loca
 
 IDetailsViewPrivate* FDetailCategoryImpl::GetDetailsView() const
 {
-	TSharedPtr<FDetailLayoutBuilderImpl> DetailLayoutBuilderPtr = DetailLayoutBuilder.Pin();
-	if (DetailLayoutBuilderPtr.IsValid())
+	TSharedPtr<FDetailLayoutBuilderImpl> ParentLayout = GetParentLayoutImpl();
+	if (ParentLayout.IsValid())
 	{
-		return DetailLayoutBuilderPtr->GetDetailsView();
+		return ParentLayout->GetDetailsView();
 	}
 
 	return nullptr;
@@ -732,8 +736,10 @@ TSharedRef<ITableRow> FDetailCategoryImpl::GenerateWidgetForTableView(const TSha
 		HeaderContent = Row.ValueWidget.Widget;
 	}
 
+	TSharedPtr<FDetailLayoutBuilderImpl> ParentLayout = GetParentLayoutImpl();
+
 	return SNew(SDetailCategoryTableRow, AsShared(), OwnerTable)
-		.InnerCategory(DetailLayoutBuilder.IsValid() ? DetailLayoutBuilder.Pin()->IsLayoutForExternalRoot() : false)
+		.InnerCategory(ParentLayout.IsValid() ? ParentLayout->IsLayoutForExternalRoot() : false)
 		.DisplayName(GetDisplayName())
 		.HeaderContent(HeaderContent);
 }
@@ -749,7 +755,7 @@ bool FDetailCategoryImpl::GenerateStandaloneWidget(FDetailWidgetRow& OutRow) con
 		HeaderContent = Row.ValueWidget.Widget;
 	}
 
-	const bool bIsInnerCategory = DetailLayoutBuilder.Pin()->IsLayoutForExternalRoot();
+	const bool bIsInnerCategory = GetParentLayoutImpl()->IsLayoutForExternalRoot();
 	FTextBlockStyle NameStyle = bIsInnerCategory ? FAppStyle::Get().GetWidgetStyle<FTextBlockStyle>("NormalText") : FAppStyle::Get().GetWidgetStyle<FTextBlockStyle>("DetailsView.CategoryTextStyle");
 	OutRow.NameContent()
 	[
@@ -1118,9 +1124,10 @@ void FDetailCategoryImpl::FilterNode(const FDetailFilter& InFilter)
 
 FCustomPropertyTypeLayoutMap FDetailCategoryImpl::GetCustomPropertyTypeLayoutMap() const
 {
-	if (DetailLayoutBuilder.IsValid())
+	TSharedPtr<FDetailLayoutBuilderImpl> ParentLayout = GetParentLayoutImpl();
+	if (ParentLayout.IsValid())
 	{
-		return DetailLayoutBuilder.Pin()->GetInstancedPropertyTypeLayoutMap();
+		return ParentLayout->GetInstancedPropertyTypeLayoutMap();
 	}
 
 	return FCustomPropertyTypeLayoutMap();
diff --git a/Engine/Source/Editor/PropertyEditor/Private/DetailCategoryBuilderImpl.h b/Engine/Source/Editor/PropertyEditor/Private/DetailCategoryBuilderImpl.h
--- a/Engine/Source/Editor/PropertyEditor/Private/DetailCategoryBuilderImpl.h
+++ b/Engine/Source/Editor/PropertyEditor/Private/DetailCategoryBuilderImpl.h
@@ -173,6 +173,7 @@ public:
 	virtual IDetailPropertyRow* AddExternalStructure(TSharedPtr<FStructOnScope> StructData, EPropertyLocation::Type Location = EPropertyLocation::Default) override;
 	virtual IDetailPropertyRow* AddExternalStructureProperty(TSharedPtr<FStructOnScope> StructData, FName PropertyName, EPropertyLocation::Type Location = EPropertyLocation::Default, const FAddPropertyParams& Params = FAddPropertyParams()) override;
 	virtual TArray<TSharedPtr<IPropertyHandle>> AddAllExternalStructureProperties(TSharedRef<FStructOnScope> StructData, EPropertyLocation::Type Location = EPropertyLocation::Default, TArray<IDetailPropertyRow*>* OutPropertiesRow = nullptr) override;
+	virtual bool IsParentLayoutValid() const override { return DetailLayoutBuilder.IsValid(); }
 	virtual IDetailLayoutBuilder& GetParentLayout() const override { return *DetailLayoutBuilder.Pin(); }
 	virtual FDetailWidgetRow& AddCustomRow(const FText& FilterString, bool bForAdvanced = false) override;
 	virtual void AddCustomBuilder(TSharedRef<IDetailCustomNodeBuilder> InCustomBuilder, bool bForAdvanced = false) override;
@@ -211,11 +212,6 @@ public:
 
 	FCustomPropertyTypeLayoutMap GetCustomPropertyTypeLayoutMap() const;
 
-	/**
-	 * @return true if the parent layout is valid or has been destroyed by a refresh.
-	 */
-	bool IsParentLayoutValid() const { return DetailLayoutBuilder.IsValid(); }
-
 	/**
 	 * @return The name of the category
 	 */
@@ -224,7 +220,7 @@ public:
 	/**
 	 * @return The parent detail layout builder for this category
 	 */
-	FDetailLayoutBuilderImpl& GetParentLayoutImpl() const { return *DetailLayoutBuilder.Pin(); }
+	TSharedPtr<FDetailLayoutBuilderImpl> GetParentLayoutImpl() const { return DetailLayoutBuilder.Pin(); }
 
 	/**
 	 * Generates the children for this category
diff --git a/Engine/Source/Editor/PropertyEditor/Private/DetailItemNode.cpp b/Engine/Source/Editor/PropertyEditor/Private/DetailItemNode.cpp
--- a/Engine/Source/Editor/PropertyEditor/Private/DetailItemNode.cpp
+++ b/Engine/Source/Editor/PropertyEditor/Private/DetailItemNode.cpp
@@ -89,9 +89,13 @@ EDetailNodeType FDetailItemNode::GetNodeType() const
 TSharedPtr<IPropertyHandle> FDetailItemNode::CreatePropertyHandle() const
 {
 	TSharedPtr<FDetailCategoryImpl> ParentCategoryPtr = ParentCategory.Pin();
-	if (Customization.HasPropertyNode() && ParentCategoryPtr && ParentCategoryPtr->IsParentLayoutValid())
+	if (Customization.HasPropertyNode() && ParentCategoryPtr.IsValid())
 	{
-		return ParentCategoryPtr->GetParentLayoutImpl().GetPropertyHandle(Customization.GetPropertyNode());
+		TSharedPtr<FDetailLayoutBuilderImpl> ParentLayout = ParentCategoryPtr->GetParentLayoutImpl();
+		if (ParentLayout.IsValid())
+		{
+			return ParentLayout->GetPropertyHandle(Customization.GetPropertyNode());
+		}
 	}
 	else if (Customization.HasCustomWidget())
 	{
@@ -387,7 +391,13 @@ void FDetailItemNode::GenerateChildren( bool bUpdateFilteredNodes )
 	Children.Empty();
 
 	TSharedPtr<FDetailCategoryImpl> ParentCategoryPinned = ParentCategory.Pin();
-	if (ParentCategoryPinned.IsValid() == false || !ParentCategoryPinned->IsParentLayoutValid())
+	if (!ParentCategoryPinned.IsValid())
+	{
+		return;
+	}
+
+	TSharedPtr<FDetailLayoutBuilderImpl> ParentLayout = ParentCategoryPinned->GetParentLayoutImpl();
+	if (!ParentLayout.IsValid())
 	{
 		return;
 	}
@@ -398,7 +408,7 @@ void FDetailItemNode::GenerateChildren( bool bUpdateFilteredNodes )
 		TSharedPtr<FComplexPropertyNode> OldChildExternalRootPropertyNode = OldChild->GetExternalRootPropertyNode();
 		if (OldChildExternalRootPropertyNode.IsValid())
 		{
-			ParentCategoryPinned->GetParentLayoutImpl().RemoveExternalRootPropertyNode(OldChildExternalRootPropertyNode.ToSharedRef());
+			ParentLayout->RemoveExternalRootPropertyNode(OldChildExternalRootPropertyNode.ToSharedRef());
 		}
 	}
 
@@ -720,11 +730,16 @@ EVisibility FDetailItemNode::ComputeItemVisibility() const
 		if (NewVisibility != EVisibility::Collapsed)
 		{
 			TSharedPtr<FDetailCategoryImpl> ParentCategoryPtr = GetParentCategory();
-			if (ParentCategoryPtr.IsValid() && ParentCategoryPtr->IsParentLayoutValid())
+			if (ParentCategoryPtr.IsValid())
 			{
-				if (!ParentCategoryPtr->GetParentLayout().IsPropertyVisible(CreatePropertyHandle().ToSharedRef()))
+				TSharedPtr<FDetailLayoutBuilderImpl> ParentLayout = ParentCategoryPtr->GetParentLayoutImpl();
+				if (ParentLayout.IsValid())
 				{
-					NewVisibility = EVisibility::Collapsed;
+					TSharedPtr<IPropertyHandle> PropertyHandle = ParentLayout->GetPropertyHandle(Customization.GetPropertyNode());
+					if (!ParentLayout->IsPropertyVisible(PropertyHandle.ToSharedRef()))
+					{
+						NewVisibility = EVisibility::Collapsed;
+					}
 				}
 			}
 		}
@@ -751,29 +766,17 @@ EVisibility FDetailItemNode::ComputeItemVisibility() const
 		}
 	}
 
-	// check the details view's IsCustomRowVisible delegate if this isn't a property row
-	if (NewVisibility != EVisibility::Collapsed && 
-		GetDetailsView() != nullptr && 
-		!Customization.HasPropertyNode())
+	const IDetailsViewPrivate* DetailsView = GetDetailsView();
+	if (DetailsView != nullptr)
 	{
-		const FName CategoryName = GetParentCategory()->GetCategoryName();
-		FName RowName;
-		if (Customization.HasCustomWidget())
+		// check the details view's IsCustomRowVisible delegate if this isn't a property row
+		// properties are handled by the IsPropertyVisible delegate
+		if (NewVisibility != EVisibility::Collapsed && !Customization.HasPropertyNode())
 		{
-			RowName = Customization.WidgetDecl->GetRowName();
-		}
-		else if (Customization.HasCustomBuilder())
-		{
-			RowName = Customization.CustomBuilderRow->GetCustomBuilderName();
-		}
-		else if (Customization.HasGroup())
-		{
-			RowName = Customization.DetailGroup->GetGroupName();
-		}
-
-		if (!GetDetailsView()->IsCustomRowVisible(RowName, CategoryName))
-		{
-			NewVisibility = EVisibility::Collapsed;
+			if (!DetailsView->IsCustomRowVisible(Customization.GetName(), GetParentCategory()->GetCategoryName()))
+			{
+				NewVisibility = EVisibility::Collapsed;
+			}
 		}
 	}
 
diff --git a/Engine/Source/Editor/PropertyEditor/Private/DetailPropertyRow.cpp b/Engine/Source/Editor/PropertyEditor/Private/DetailPropertyRow.cpp
--- a/Engine/Source/Editor/PropertyEditor/Private/DetailPropertyRow.cpp
+++ b/Engine/Source/Editor/PropertyEditor/Private/DetailPropertyRow.cpp
@@ -27,13 +27,13 @@ FDetailPropertyRow::FDetailPropertyRow(TSharedPtr<FPropertyNode> InPropertyNode,
 	, bForceAutoExpansion( false )
 	, bCachedCustomTypeInterface(false)
 {
-	PropertyHandle = InParentCategory->GetParentLayoutImpl().GetPropertyHandle(PropertyNode);
+	PropertyHandle = InParentCategory->GetParentLayoutImpl()->GetPropertyHandle(PropertyNode);
 
 	if (PropertyNode.IsValid())
 	{
 		TSharedRef<FPropertyNode> PropertyNodeRef = PropertyNode.ToSharedRef();
 
-		const TSharedRef<IPropertyUtilities> Utilities = InParentCategory->GetParentLayoutImpl().GetPropertyUtilities();
+		const TSharedRef<IPropertyUtilities> Utilities = InParentCategory->GetParentLayoutImpl()->GetPropertyUtilities();
 
 		if (PropertyNode->AsCategoryNode() == nullptr)
 		{
@@ -320,8 +320,8 @@ void FDetailPropertyRow::OnGenerateChildren( FDetailNodeList& OutChildren )
 	{
 		// This is a sub-category.  Populate from SubCategory builder
 		TSharedRef<FDetailCategoryImpl> ParentCategoryRef = ParentCategory.Pin().ToSharedRef();
-		FDetailLayoutBuilderImpl& LayoutBuilder = ParentCategoryRef->GetParentLayoutImpl();
-		TSharedPtr<FDetailCategoryImpl> MyCategory = LayoutBuilder.GetSubCategoryImpl(PropertyNode->AsCategoryNode()->GetCategoryName());
+		TSharedPtr<FDetailLayoutBuilderImpl> LayoutBuilder = ParentCategoryRef->GetParentLayoutImpl();
+		TSharedPtr<FDetailCategoryImpl> MyCategory = LayoutBuilder->GetSubCategoryImpl(PropertyNode->AsCategoryNode()->GetCategoryName());
 		if(MyCategory.IsValid())
 		{
 			MyCategory->GenerateLayout();
@@ -498,7 +498,7 @@ TSharedPtr<IPropertyTypeCustomization> FDetailPropertyRow::GetPropertyCustomizat
 	if (!PropertyEditorHelpers::IsStaticArray(*InPropertyNode))
 	{
 		FProperty* Property = InPropertyNode->GetProperty();
-		TSharedPtr<IPropertyHandle> PropHandle = InParentCategory->GetParentLayoutImpl().GetPropertyHandle(InPropertyNode);
+		TSharedPtr<IPropertyHandle> PropHandle = InParentCategory->GetParentLayoutImpl()->GetPropertyHandle(InPropertyNode);
 
 		static FName NAME_PropertyEditor("PropertyEditor");
 		FPropertyEditorModule& PropertyEditorModule = FModuleManager::GetModuleChecked<FPropertyEditorModule>(NAME_PropertyEditor);
@@ -537,7 +537,7 @@ void FDetailPropertyRow::MakeExternalPropertyRowCustomization(TSharedPtr<FStruct
 
 	RootPropertyNode->InitNode(InitParams);
 
-	ParentCategory->GetParentLayoutImpl().AddExternalRootPropertyNode(RootPropertyNode);
+	ParentCategory->GetParentLayoutImpl()->AddExternalRootPropertyNode(RootPropertyNode);
 
 	if (PropertyName != NAME_None)
 	{
@@ -609,7 +609,7 @@ void FDetailPropertyRow::MakeExternalPropertyRowCustomization(const TArray<UObje
 
 	RootPropertyNode->InitNode(InitParams);
 
-	ParentCategory->GetParentLayoutImpl().AddExternalRootPropertyNode(RootPropertyNode);
+	ParentCategory->GetParentLayoutImpl()->AddExternalRootPropertyNode(RootPropertyNode);
 
 	if (PropertyName != NAME_None)
 	{
@@ -783,10 +783,8 @@ void FDetailPropertyRow::MakeNameOrKeyWidget( FDetailWidgetRow& Row, const TShar
 		}
 		else
 		{
-			const TSharedRef<IPropertyUtilities> PropertyUtilities = ParentCategory.Pin()->GetParentLayoutImpl().GetPropertyUtilities();
-
 			NameWidget =
-				SNew(SPropertyValueWidget, PropertyKeyEditor, PropertyUtilities)
+				SNew(SPropertyValueWidget, PropertyKeyEditor, ParentCategory.Pin()->GetParentLayoutImpl()->GetPropertyUtilities())
 				.IsEnabled(IsEnabledAttrib)
 				.ShowPropertyButtons(false);
 		}
diff --git a/Engine/Source/Editor/PropertyEditor/Private/SDetailSingleItemRow.cpp b/Engine/Source/Editor/PropertyEditor/Private/SDetailSingleItemRow.cpp
--- a/Engine/Source/Editor/PropertyEditor/Private/SDetailSingleItemRow.cpp
+++ b/Engine/Source/Editor/PropertyEditor/Private/SDetailSingleItemRow.cpp
@@ -1110,7 +1110,7 @@ void SDetailSingleItemRow::PopulateExtensionWidget()
 			const UClass* ObjectClass = Handle->GetOuterBaseClass();
 			if (Handle->IsValidHandle() && ObjectClass && ExtensionHandler->IsPropertyExtendable(ObjectClass, *Handle))
 			{
-				FDetailLayoutBuilderImpl& DetailLayout = OwnerTreeNodePinned->GetParentCategory()->GetParentLayoutImpl();
+				IDetailLayoutBuilder& DetailLayout = OwnerTreeNodePinned->GetParentCategory()->GetParentLayout();
 				ExtensionHandler->ExtendWidgetRow(WidgetRow, DetailLayout, ObjectClass, Handle);
 			}
 		}
