Added a check to SDetailsViewBase::RestoreExpandedItems to look for a prefix in the items before checking children. In very large node hierarchies, eg. instanced meshes with thousands of instances, this can end up very expensive even though none of the children in the array are even contained in the array.

diff --git a/Engine/Source/Editor/PropertyEditor/Private/SDetailsViewBase.cpp b/Engine/Source/Editor/PropertyEditor/Private/SDetailsViewBase.cpp
--- a/Engine/Source/Editor/PropertyEditor/Private/SDetailsViewBase.cpp
+++ b/Engine/Source/Editor/PropertyEditor/Private/SDetailsViewBase.cpp
@@ -1143,9 +1143,25 @@ static void SetExpandedItems(TSharedPtr<FPropertyNode> InPropertyNode, const TSe
 		InPropertyNode->SetNodeFlags(EPropertyNodeFlags::Expanded, false);
 	}
 
-	for (int32 NodeIndex = 0; NodeIndex < InPropertyNode->GetNumChildNodes(); ++NodeIndex)
+	bool bAnyPrefix = false;
+	for (const FString& Item : InExpandedItems)
 	{
-		SetExpandedItems(InPropertyNode->GetChildNode(NodeIndex), InExpandedItems, bCollapseRest);
+		// check if this path is a prefix to some other path
+		// if it's not, we can early out and skip checking every child node, which is a win for very large child arrays
+		if (Item.Len() > Path.Len() && Item.StartsWith(Path) &&
+			(Item[Path.Len()] == '[' || Item[Path.Len()] == '.'))
+		{
+			bAnyPrefix = true;
+			break;
+		}
+	}
+
+	if (bAnyPrefix)
+	{
+		for (int32 NodeIndex = 0; NodeIndex < InPropertyNode->GetNumChildNodes(); ++NodeIndex)
+		{
+			SetExpandedItems(InPropertyNode->GetChildNode(NodeIndex), InExpandedItems, bCollapseRest);
+		}
 	}
 }
 
