Removed FPropertyNode::ParentNode.

We already had a ParentNodeWeakPtr that we were checking to make sure the pointer was valid - now we just use that value directly instead.

#rnx

diff --git a/Engine/Source/Editor/PropertyEditor/Private/CategoryPropertyNode.cpp b/Engine/Source/Editor/PropertyEditor/Private/CategoryPropertyNode.cpp
--- a/Engine/Source/Editor/PropertyEditor/Private/CategoryPropertyNode.cpp
+++ b/Engine/Source/Editor/PropertyEditor/Private/CategoryPropertyNode.cpp
@@ -117,7 +117,8 @@ bool FCategoryPropertyNode::GetQualifiedName( FString& PathPlusIndex, const bool
 {
 	bool bAddedAnything = false;
 
-	if (ParentNode && StopParent != ParentNode)
+	const TSharedPtr<FPropertyNode> ParentNode = ParentNodeWeakPtr.Pin();
+	if (ParentNode && StopParent != ParentNode.Get())
 	{
 		bAddedAnything = ParentNode->GetQualifiedName(PathPlusIndex, bWithArrayIndex, StopParent, bIgnoreCategories );
 	}
diff --git a/Engine/Source/Editor/PropertyEditor/Private/ItemPropertyNode.cpp b/Engine/Source/Editor/PropertyEditor/Private/ItemPropertyNode.cpp
--- a/Engine/Source/Editor/PropertyEditor/Private/ItemPropertyNode.cpp
+++ b/Engine/Source/Editor/PropertyEditor/Private/ItemPropertyNode.cpp
@@ -11,13 +11,12 @@
 
 #define LOCTEXT_NAMESPACE "ItemPropertyNode"
 
-FItemPropertyNode::FItemPropertyNode(void)
-	: FPropertyNode()
+FItemPropertyNode::FItemPropertyNode()
 {
 	bCanDisplayFavorite = false;
 }
 
-FItemPropertyNode::~FItemPropertyNode(void)
+FItemPropertyNode::~FItemPropertyNode()
 {
 
 }
@@ -25,7 +24,8 @@ FItemPropertyNode::~FItemPropertyNode(void)
 uint8* FItemPropertyNode::GetValueBaseAddress(uint8* StartAddress, bool bIsSparseData) const
 {
 	const FProperty* MyProperty = GetProperty();
-	if( MyProperty && ParentNodeWeakPtr.IsValid())
+	const TSharedPtr<FPropertyNode> ParentNode = ParentNodeWeakPtr.Pin();
+	if (MyProperty && ParentNode)
 	{
 		const FArrayProperty* OuterArrayProp = MyProperty->GetOwner<FArrayProperty>();
 		const FSetProperty* OuterSetProp = MyProperty->GetOwner<FSetProperty>();
@@ -168,7 +168,7 @@ void FItemPropertyNode::InitChildNodes()
 	{
 		void* Array = NULL;
 		FReadAddressList Addresses;
-		if ( GetReadAddress(!!HasNodeFlags(EPropertyNodeFlags::SingleSelectOnly), Addresses ) )
+		if ( GetReadAddress(HasNodeFlags(EPropertyNodeFlags::SingleSelectOnly), Addresses ) )
 		{
 			Array = Addresses.GetAddress(0);
 		}
@@ -473,6 +473,8 @@ FText FItemPropertyNode::GetDisplayName() const
 				{
 					FString PropertyDisplayName;
 					bool bIsBoolProperty = CastField<const FBoolProperty>(PropertyPtr) != NULL;
+
+					const TSharedPtr<FPropertyNode> ParentNode = ParentNodeWeakPtr.Pin();
 					const FStructProperty* ParentStructProperty = CastField<const FStructProperty>(ParentNode->GetProperty());
 					if( ParentStructProperty && ParentStructProperty->Struct->GetFName() == NAME_Rotator )
 					{
@@ -521,6 +523,7 @@ FText FItemPropertyNode::GetDisplayName() const
 			}
 			
 			// Sets and maps do not have a display index.
+			const TSharedPtr<FPropertyNode> ParentNode = ParentNodeWeakPtr.Pin();
 			FProperty* ParentProperty = ParentNode->GetProperty();
 
 			// Also handle UArray's having the ArraySizeEnum entry...
@@ -591,7 +594,7 @@ FText FItemPropertyNode::GetDisplayName() const
 				}
 			}
 			// Maps should have display names that reflect the key and value types
-			else if (PropertyPtr != nullptr && CastField<FMapProperty>(ParentNode->GetProperty()) != nullptr)
+			else if (PropertyPtr != nullptr && CastField<FMapProperty>(ParentProperty) != nullptr)
 			{
 				FText FormatText = GetPropertyKeyNode().IsValid()
 					? LOCTEXT("MapValueDisplayFormat", "Value ({0})")
diff --git a/Engine/Source/Editor/PropertyEditor/Private/ObjectPropertyNode.cpp b/Engine/Source/Editor/PropertyEditor/Private/ObjectPropertyNode.cpp
--- a/Engine/Source/Editor/PropertyEditor/Private/ObjectPropertyNode.cpp
+++ b/Engine/Source/Editor/PropertyEditor/Private/ObjectPropertyNode.cpp
@@ -761,7 +761,7 @@ TSharedPtr<FPropertyNode> FObjectPropertyNode::GenerateSingleChild( FName ChildP
 		return ChildNodes[0];
 	}
 
-	return NULL;
+	return nullptr;
 }
 
 bool FObjectPropertyNode::IsSparseDataStruct(const UScriptStruct* Struct) const
@@ -784,7 +784,8 @@ bool FObjectPropertyNode::IsSparseDataStruct(const UScriptStruct* Struct) const
 bool FObjectPropertyNode::GetQualifiedName(FString& PathPlusIndex, bool bWithArrayIndex, const FPropertyNode* StopParent, bool bIgnoreCategories ) const
 {
 	bool bAddedAnything = false;
-	if (ParentNode && ParentNode != StopParent)
+	const TSharedPtr<FPropertyNode> ParentNode = ParentNodeWeakPtr.Pin();
+	if (ParentNode && StopParent != ParentNode.Get())
 	{
 		bAddedAnything = ParentNode->GetQualifiedName(PathPlusIndex, bWithArrayIndex, StopParent, bIgnoreCategories);
 	}
diff --git a/Engine/Source/Editor/PropertyEditor/Private/ObjectPropertyNode.h b/Engine/Source/Editor/PropertyEditor/Private/ObjectPropertyNode.h
--- a/Engine/Source/Editor/PropertyEditor/Private/ObjectPropertyNode.h
+++ b/Engine/Source/Editor/PropertyEditor/Private/ObjectPropertyNode.h
@@ -140,7 +140,7 @@ public:
 	 */
 	const TSet<FName>& GetHiddenCategories() const { return HiddenCategories; }
 
-	bool IsRootNode() const { return ParentNode == nullptr; }
+	bool IsRootNode() const { return ParentNodeWeakPtr.Pin() == nullptr; }
 
 	/**
 	 * @return True if Struct is one of the sparse data structures used by this object
diff --git a/Engine/Source/Editor/PropertyEditor/Private/PropertyNode.cpp b/Engine/Source/Editor/PropertyEditor/Private/PropertyNode.cpp
--- a/Engine/Source/Editor/PropertyEditor/Private/PropertyNode.cpp
+++ b/Engine/Source/Editor/PropertyEditor/Private/PropertyNode.cpp
@@ -64,9 +64,8 @@ static FObjectPropertyNode* NotifyFindObjectItemParent(FPropertyNode* InNode)
 	return Result;
 }
 
-FPropertyNode::FPropertyNode(void)
-	: ParentNode(NULL)
-	, Property(NULL)
+FPropertyNode::FPropertyNode()
+	: Property(nullptr)
 	, ArrayOffset(0)
 	, ArrayIndex(-1)
 	, MaxChildDepthAllowed(FPropertyNodeConstants::NoDepthRestrictions)
@@ -81,7 +80,7 @@ FPropertyNode::FPropertyNode(void)
 {
 }
 
-FPropertyNode::~FPropertyNode(void)
+FPropertyNode::~FPropertyNode()
 {
 	DestroyTree();
 }
@@ -93,9 +92,8 @@ void FPropertyNode::InitNode(const FPropertyNodeInitParams& InitParams)
 
 	//tree hierarchy
 	check(InitParams.ParentNode.Get() != this);
-	ParentNode = InitParams.ParentNode.Get();
 	ParentNodeWeakPtr = InitParams.ParentNode;
-
+	
 	//Property Data
 	Property = InitParams.Property;
 	ArrayOffset = InitParams.ArrayOffset;
@@ -103,7 +101,8 @@ void FPropertyNode::InitNode(const FPropertyNodeInitParams& InitParams)
 
 	bool bIsSparse = InitParams.IsSparseProperty == FPropertyNodeInitParams::EIsSparseDataProperty::True;
 
-	if (ParentNode && InitParams.IsSparseProperty == FPropertyNodeInitParams::EIsSparseDataProperty::Inherit)
+	TSharedPtr<FPropertyNode> ParentNode = ParentNodeWeakPtr.Pin();
+	if (ParentNode.IsValid() && InitParams.IsSparseProperty == FPropertyNodeInitParams::EIsSparseDataProperty::Inherit)
 	{
 		//default to parents max child depth
 		MaxChildDepthAllowed = ParentNode->MaxChildDepthAllowed;
@@ -281,13 +280,6 @@ void FPropertyNode::RebuildChildren()
 		}
 	}
 
-	//see if they support some kind of edit condition
-	if (Property.IsValid() && Property->GetBoolMetaData(TEXT("FullyExpand")))
-	{
-		bool bExpand = true;
-		bool bRecurse = true;
-	}
-
 	// Children have been rebuilt, clear any pending rebuild requests
 	bRebuildChildrenRequested = false;
 	bChildrenRebuilt = true;
@@ -944,7 +936,7 @@ bool FPropertyNode::IsEditConst() const
 			TSharedRef<FEditPropertyChain> PropertyChain = BuildPropertyChain(Property.Get());
 			
 			// travel up the chain to see if this property's owner struct is EditConst - if it is, so is this property
-			FPropertyNode* CurParent = ParentNode;
+			TSharedPtr<FPropertyNode> CurParent = ParentNodeWeakPtr.Pin();
 			while (CurParent != nullptr)
 			{
 				FStructProperty* StructProperty = CastField<FStructProperty>(CurParent->GetProperty());
@@ -993,7 +985,7 @@ bool FPropertyNode::IsEditConst() const
 					break;
 				}
 				
-				CurParent = CurParent->ParentNode;
+				CurParent = CurParent->ParentNodeWeakPtr.Pin();
 			}
 
 			if (!bIsEditConst)
@@ -1161,7 +1153,8 @@ bool FPropertyNode::IsOnlyVisibleWhenEditConditionMet() const
 bool FPropertyNode::GetQualifiedName( FString& PathPlusIndex, const bool bWithArrayIndex, const FPropertyNode* StopParent, bool bIgnoreCategories ) const
 {
 	bool bAddedAnything = false;
-	if (ParentNodeWeakPtr.IsValid() && StopParent != ParentNode)
+	const TSharedPtr<FPropertyNode> ParentNode = ParentNodeWeakPtr.Pin();
+	if (ParentNode && StopParent != ParentNode.Get())
 	{
 		bAddedAnything = ParentNode->GetQualifiedName(PathPlusIndex, bWithArrayIndex, StopParent, bIgnoreCategories);
 	}
@@ -1195,7 +1188,8 @@ bool FPropertyNode::GetReadAddressUncached( const FPropertyNode& InPropertyNode,
 									bool bObjectForceCompare,
 									bool bArrayPropertiesCanDifferInSize ) const
 {
-	if (ParentNodeWeakPtr.IsValid())
+	const TSharedPtr<FPropertyNode> ParentNode = ParentNodeWeakPtr.Pin();
+	if (ParentNode.IsValid())
 	{
 		return ParentNode->GetReadAddressUncached( InPropertyNode, InRequiresSingleSelection, OutAddresses, bComparePropertyContents, bObjectForceCompare, bArrayPropertiesCanDifferInSize );
 	}
@@ -1205,7 +1199,8 @@ bool FPropertyNode::GetReadAddressUncached( const FPropertyNode& InPropertyNode,
 
 bool FPropertyNode::GetReadAddressUncached( const FPropertyNode& InPropertyNode, FReadAddressListData& OutAddresses ) const
 {
-	if (ParentNodeWeakPtr.IsValid())
+	const TSharedPtr<FPropertyNode> ParentNode = ParentNodeWeakPtr.Pin();
+	if (ParentNode.IsValid())
 	{
 		return ParentNode->GetReadAddressUncached( InPropertyNode, OutAddresses );
 	}
@@ -1325,8 +1320,9 @@ uint8* FPropertyNode::GetValueBaseAddress(uint8* StartAddress, bool bIsSparseDat
 	}
 	else
 	{
-		if (ParentNodeWeakPtr.IsValid())
-	{
+		const TSharedPtr<FPropertyNode> ParentNode = ParentNodeWeakPtr.Pin();
+		if (ParentNode.IsValid())
+		{
 			Result = ParentNode->GetValueAddress(StartAddress, bIsSparseData);
 		}
 	}
@@ -2321,6 +2317,7 @@ void FPropertyNode::FilterNodes( const TArray<FString>& InFilterStrings, const b
 
 		// For containers, check if base class metadata in parent includes 'TitleProperty', add corresponding value to filter names if so.
 		static const FName TitlePropertyFName = FName(TEXT("TitleProperty"));
+		const TSharedPtr<FPropertyNode> ParentNode = ParentNodeWeakPtr.Pin();
 		if (ParentNode && ParentNode->GetProperty())
 		{
 			const FString& TitleProperty = ParentNode->GetProperty()->GetMetaData(TitlePropertyFName);
@@ -2956,7 +2953,18 @@ const TMap<FName, FString>* FPropertyNode::GetInstanceMetaDataMap() const
 
 bool FPropertyNode::ParentOrSelfHasMetaData(const FName& MetaDataKey) const
 {
-	return (Property.IsValid() && Property->HasMetaData(MetaDataKey)) || (ParentNode && ParentNode->ParentOrSelfHasMetaData(MetaDataKey));
+	if (Property.IsValid() && Property->HasMetaData(MetaDataKey))
+	{
+		return true;
+	}
+	
+	const TSharedPtr<FPropertyNode> ParentNode = ParentNodeWeakPtr.Pin();
+	if (ParentNode.IsValid() && ParentNode->ParentOrSelfHasMetaData(MetaDataKey))
+	{
+		return true;
+	}
+
+	return false;
 }
 
 void FPropertyNode::InvalidateCachedState()
@@ -3385,11 +3393,11 @@ void FPropertyNode::PropagatePropertyChange( UObject* ModifiedObject, const TCHA
 			if (DestSimplePropAddr != nullptr)
 			{
 				FProperty* ComplexProperty = Prop;
-				FPropertyNode* ComplexPropertyNode = this;
+				TSharedPtr<FPropertyNode> ComplexPropertyNode = AsShared();
 				if (ParentArrayProp || ParentMapProp || ParentSetProp)
 				{
 					ComplexProperty = ParentProp;
-					ComplexPropertyNode = ParentNode;
+					ComplexPropertyNode = ParentNodeWeakPtr.Pin();
 				}
 				
 				uint8* DestComplexPropAddr = ComplexPropertyNode->GetValueBaseAddressFromObject(ActualObjToChange);
diff --git a/Engine/Source/Editor/PropertyEditor/Private/PropertyNode.h b/Engine/Source/Editor/PropertyEditor/Private/PropertyNode.h
--- a/Engine/Source/Editor/PropertyEditor/Private/PropertyNode.h
+++ b/Engine/Source/Editor/PropertyEditor/Private/PropertyNode.h
@@ -283,8 +283,8 @@ class FPropertyNode : public TSharedFromThis<FPropertyNode>
 {
 public:
 
-	FPropertyNode(void);
-	virtual ~FPropertyNode(void);
+	FPropertyNode();
+	virtual ~FPropertyNode();
 
 	/**
 	 * Init Tree Node internally (used only derived classes to pass through variables that are common to all nodes
@@ -1054,8 +1054,6 @@ protected:
 	 * The node that is the parent of this node or nullptr for the root
 	 */
 	TWeakPtr<FPropertyNode> ParentNodeWeakPtr;
-	//@todo consolidate with ParentNodeWeakPtr, ParentNode is legacy
-	FPropertyNode* ParentNode;
 
 	/**	The property node, if any, that serves as the key value for this node */
 	TSharedPtr<FPropertyNode> PropertyKeyNode;
