Removed inline forward declarations from PropertyEditor.h

diff --git a/Engine/Source/Editor/PropertyEditor/Private/Presentation/PropertyEditor/PropertyEditor.h b/Engine/Source/Editor/PropertyEditor/Private/Presentation/PropertyEditor/PropertyEditor.h
--- a/Engine/Source/Editor/PropertyEditor/Private/Presentation/PropertyEditor/PropertyEditor.h
+++ b/Engine/Source/Editor/PropertyEditor/Private/Presentation/PropertyEditor/PropertyEditor.h
@@ -8,11 +8,14 @@
 
 class FEditConditionExpression;
 class FEditConditionContext;
+class FPropertyNode;
+class IPropertyHandle;
+class IPropertyUtilities;
 
-class FPropertyEditor : public TSharedFromThis< FPropertyEditor >	
+class FPropertyEditor : public TSharedFromThis<FPropertyEditor>	
 {
 public:
-	static TSharedRef< FPropertyEditor > Create( const TSharedRef< class FPropertyNode >& InPropertyNode, const TSharedRef<class IPropertyUtilities >& InPropertyUtilities );
+	static TSharedRef<FPropertyEditor> Create(const TSharedRef<FPropertyNode>& InPropertyNode, const TSharedRef<IPropertyUtilities>& InPropertyUtilities);
 
 	/** @return The String containing the value of the property */
 	FString GetValueAsString() const;
@@ -90,15 +93,15 @@ public:
 	void EditConfigHierarchy();
 
 	/**	In an ideal world we wouldn't expose these */
-	TSharedRef< FPropertyNode > GetPropertyNode() const;
+	TSharedRef< FPropertyNode> GetPropertyNode() const;
 	const FProperty* GetProperty() const;
-	TSharedRef< IPropertyHandle > GetPropertyHandle() const;
+	TSharedRef< IPropertyHandle> GetPropertyHandle() const;
 
-	static void SyncToObjectsInNode( const TWeakPtr< FPropertyNode >& WeakPropertyNode );
+	static void SyncToObjectsInNode( const TWeakPtr<FPropertyNode>& WeakPropertyNode );
 
 	static const FString MultipleValuesDisplayName;
 private:
-	FPropertyEditor( const TSharedRef< class FPropertyNode >& InPropertyNode, const TSharedRef<class IPropertyUtilities >& InPropertyUtilities );
+	FPropertyEditor( const TSharedRef<FPropertyNode>& InPropertyNode, const TSharedRef< IPropertyUtilities>& InPropertyUtilities );
 
 	void OnUseSelected();
 	void OnAddItem();
@@ -113,11 +116,11 @@ private:
 private:
 
 	/** Property handle for actually reading/writing the value of a property */
-	TSharedPtr< class IPropertyHandle > PropertyHandle;
+	TSharedPtr<IPropertyHandle> PropertyHandle;
 
 	/** pointer to the property node. */
-	TSharedRef< class FPropertyNode > PropertyNode;
+	TSharedRef<FPropertyNode> PropertyNode;
 	 
 	/** The property view where this widget resides */
-	TSharedRef< class IPropertyUtilities > PropertyUtilities;
+	TSharedRef<IPropertyUtilities> PropertyUtilities;
 };
