A details row's IsResetToDefaultVisible delegate is now respected even when no property is assigned to the row.

Technically, this is user error as you're meant to pass the delegate that doesn't take a property handle, but there's not much harm in calling it anyway.

diff --git a/Engine/Source/Editor/PropertyEditor/Private/SDetailSingleItemRow.cpp b/Engine/Source/Editor/PropertyEditor/Private/SDetailSingleItemRow.cpp
--- a/Engine/Source/Editor/PropertyEditor/Private/SDetailSingleItemRow.cpp
+++ b/Engine/Source/Editor/PropertyEditor/Private/SDetailSingleItemRow.cpp
@@ -294,6 +294,12 @@ void SDetailSingleItemRow::UpdateResetToDefault()
 	bCachedResetToDefaultVisible = false;
 
 	TSharedPtr<IPropertyHandle> PropertyHandle = GetPropertyHandle();
+	if (WidgetRow.CustomResetToDefault.IsSet())
+	{
+		bCachedResetToDefaultVisible = WidgetRow.CustomResetToDefault.GetValue().IsResetToDefaultVisible(PropertyHandle);
+		return;
+	}
+
 	if (PropertyHandle.IsValid())
 	{
 		if (PropertyHandle->HasMetaData("NoResetToDefault") || PropertyHandle->GetInstanceMetaData("NoResetToDefault"))
@@ -301,14 +307,7 @@ void SDetailSingleItemRow::UpdateResetToDefault()
 			bCachedResetToDefaultVisible = false;
 			return;
 		}
-	}
 
-	if (WidgetRow.CustomResetToDefault.IsSet())
-	{
-		bCachedResetToDefaultVisible = WidgetRow.CustomResetToDefault.GetValue().IsResetToDefaultVisible(PropertyHandle);
-	}
-	else if (PropertyHandle.IsValid())
-	{
 		bCachedResetToDefaultVisible = PropertyHandle->CanResetToDefault();
 	}
 }
@@ -393,7 +392,7 @@ void SDetailSingleItemRow::Construct( const FArguments& InArgs, FDetailLayoutCus
 			TAttribute<bool> IsEnabledAttribute = TAttribute<bool>::CreateLambda(
 				[PropertyEnabledAttribute, RowIsEnabledAttribute, RowEditConditionAttribute]()
 				{
-					return PropertyEnabledAttribute.Get() && RowIsEnabledAttribute.Get(true) && RowEditConditionAttribute.Get(true);
+					return PropertyEnabledAttribute.Get(true) && RowIsEnabledAttribute.Get(true) && RowEditConditionAttribute.Get(true);
 				});
 
 			TAttribute<bool> RowIsValueEnabledAttribute = WidgetRow.IsValueEnabledAttr;
diff --git a/Engine/Source/Editor/PropertyEditor/Public/IDetailPropertyRow.h b/Engine/Source/Editor/PropertyEditor/Public/IDetailPropertyRow.h
--- a/Engine/Source/Editor/PropertyEditor/Public/IDetailPropertyRow.h
+++ b/Engine/Source/Editor/PropertyEditor/Public/IDetailPropertyRow.h
@@ -101,7 +101,7 @@ public:
 			return false;
 		}
 
-		if (IsVisiblePropertyDelegate.IsBound() && Property.IsValid())
+		if (IsVisiblePropertyDelegate.IsBound())
 		{
 			return IsVisiblePropertyDelegate.Execute(Property);
 		}
