Customized details rows are now placed before default rows once again. This is most noticeable in the ordering of certain classes, eg. Intensity and Color in the Light category are now placed first.

This was an unintentional change that was introduced at some point during the development of UE5.

diff --git a/Engine/Source/Editor/PropertyEditor/Private/DetailCategoryBuilderImpl.cpp b/Engine/Source/Editor/PropertyEditor/Private/DetailCategoryBuilderImpl.cpp
--- a/Engine/Source/Editor/PropertyEditor/Private/DetailCategoryBuilderImpl.cpp
+++ b/Engine/Source/Editor/PropertyEditor/Private/DetailCategoryBuilderImpl.cpp
@@ -21,15 +21,36 @@
 #include "StructurePropertyNode.h"
 #include "Styling/StyleColors.h"
  
+static void AddLayoutToList(const FDetailLayoutCustomization& Layout, TArray<FDetailLayoutCustomization>& List)
+{
+	if (Layout.bCustom)
+	{
+		int32 Index = 0;
+		for (; Index < List.Num(); ++Index)
+		{
+			if (!List[Index].bCustom)
+			{
+				break;
+			}
+		}
+
+		List.Insert(Layout, Index);
+	}
+	else
+	{
+		List.Add(Layout);
+	}
+}
+
 void FDetailLayout::AddLayout(const FDetailLayoutCustomization& Layout)
 {
 	if (Layout.bAdvanced)
 	{
-		AdvancedLayouts.Add(Layout);
+		AddLayoutToList(Layout, AdvancedLayouts);
 	}
 	else
 	{
-		SimpleLayouts.Add(Layout);
+		AddLayoutToList(Layout, SimpleLayouts);
 	}
 }
 
