Editor - FStructurePropertyNode::InitChildNodes() initialized the child nodes of a struct, in the proper order, with the proper visibility...etc.  But then some subsiquent code after the InitNode(...) code would run for a structure in FPropertyHandleBase::AddChildStructure, we'd re-iterate the properties of the struct adding them in, but no in the proper order that was already determined in initChildNodes.  So, was messing up any kind of DisplayPriority/DisplayAfter rules for a struct.

diff --git a/Engine/Source/Editor/PropertyEditor/Private/PropertyHandleImpl.cpp b/Engine/Source/Editor/PropertyEditor/Private/PropertyHandleImpl.cpp
--- a/Engine/Source/Editor/PropertyEditor/Private/PropertyHandleImpl.cpp
+++ b/Engine/Source/Editor/PropertyEditor/Private/PropertyHandleImpl.cpp
@@ -3139,34 +3139,16 @@ TArray<TSharedPtr<IPropertyHandle>> FPropertyHandleBase::AddChildStructure( TSha
 
 	StructPropertyNode->InitNode(RootInitParams);
 
-	const bool bShouldShowHiddenProperties = !!PropertyNode->HasNodeFlags(EPropertyNodeFlags::ShouldShowHiddenProperties);
-	const bool bShouldShowDisableEditOnInstance = !!PropertyNode->HasNodeFlags(EPropertyNodeFlags::ShouldShowDisableEditOnInstance);
-
-	for (TFieldIterator<FProperty> It(InStruct->GetStruct()); It; ++It)
+	// Loop over the structs child nodes and add property handles for each of them.
+	for (int32 ChildNodeIndex = 0; ChildNodeIndex < StructPropertyNode->GetNumChildNodes(); ChildNodeIndex++)
 	{
-		FProperty* StructMember = *It;
-
-		if (PropertyEditorHelpers::ShouldBeVisible(*StructPropertyNode.Get(), StructMember))
-		{
-			TSharedRef<FItemPropertyNode> NewItemNode(new FItemPropertyNode);
-
-			FPropertyNodeInitParams InitParams;
-			InitParams.ParentNode = StructPropertyNode;
-			InitParams.Property = StructMember;
-			InitParams.ArrayOffset = 0;
-			InitParams.ArrayIndex = INDEX_NONE;
-			InitParams.bAllowChildren = true;
-			InitParams.bForceHiddenPropertyVisibility = bShouldShowHiddenProperties;
-			InitParams.bCreateDisableEditOnInstanceNodes = bShouldShowDisableEditOnInstance;
-			InitParams.bCreateCategoryNodes = false;
-
-			NewItemNode->InitNode(InitParams);
-			StructPropertyNode->AddChildNode(NewItemNode);
-
-			PropertyHandles.Add(PropertyEditorHelpers::GetPropertyHandle(NewItemNode, Implementation->GetNotifyHook(), Implementation->GetPropertyUtilities()));
-		}
+		PropertyHandles.Add(PropertyEditorHelpers::GetPropertyHandle(
+			StructPropertyNode->GetChildNode(ChildNodeIndex).ToSharedRef(),
+			Implementation->GetNotifyHook(),
+			Implementation->GetPropertyUtilities()
+		));
 	}
-
+	
 	PropertyNode->AddChildNode(StructPropertyNode);
 
 	return PropertyHandles;
