Create struct header widget with FPropertyEditorModule::CreateSingleProperty/SSingleProperty

[snapshot CL20677141 - 7/9/2022 5:11:43 PM]

diff --git a/Engine/Source/Editor/PropertyEditor/Private/CustomChildBuilder.cpp b/Engine/Source/Editor/PropertyEditor/Private/CustomChildBuilder.cpp
--- a/Engine/Source/Editor/PropertyEditor/Private/CustomChildBuilder.cpp
+++ b/Engine/Source/Editor/PropertyEditor/Private/CustomChildBuilder.cpp
@@ -6,6 +6,7 @@
 #include "PropertyHandleImpl.h"
 #include "DetailPropertyRow.h"
 #include "ObjectPropertyNode.h"
+#include "SStandaloneCustomizedValueWidget.h"
 
 IDetailChildrenBuilder& FCustomChildrenBuilder::AddCustomBuilder( TSharedRef<class IDetailCustomNodeBuilder> InCustomBuilder )
 {
@@ -117,42 +118,6 @@ IDetailPropertyRow* FCustomChildrenBuilder::AddExternalObjectProperty(const TArr
 	return NewRow.Get();
 }
 
-class SStandaloneCustomStructValue : public SCompoundWidget, public IPropertyTypeCustomizationUtils
-{
-public:
-	SLATE_BEGIN_ARGS( SStandaloneCustomStructValue )
-	{}
-	SLATE_END_ARGS()
-	
-	void Construct( const FArguments& InArgs, TSharedPtr<IPropertyTypeCustomization> InCustomizationInterface, TSharedRef<IPropertyHandle> InStructPropertyHandle, TSharedRef<FDetailCategoryImpl> InParentCategory )
-	{
-		CustomizationInterface = InCustomizationInterface;
-		StructPropertyHandle = InStructPropertyHandle;
-		ParentCategory = InParentCategory;
-		CustomPropertyWidget = MakeShareable(new FDetailWidgetRow);
-
-		CustomizationInterface->CustomizeHeader(InStructPropertyHandle, *CustomPropertyWidget, *this);
-
-		ChildSlot
-		[
-			CustomPropertyWidget->ValueWidget.Widget
-		];
-	}
-
-	virtual TSharedPtr<FAssetThumbnailPool> GetThumbnailPool() const override
-	{
-		TSharedPtr<FDetailCategoryImpl> ParentCategoryPinned = ParentCategory.Pin();
-		return ParentCategoryPinned.IsValid() ? ParentCategoryPinned->GetParentLayout().GetThumbnailPool() : NULL;
-	}
-
-private:
-	TWeakPtr<FDetailCategoryImpl> ParentCategory;
-	TSharedPtr<IPropertyTypeCustomization> CustomizationInterface;
-	TSharedPtr<IPropertyHandle> StructPropertyHandle;
-	TSharedPtr<FDetailWidgetRow> CustomPropertyWidget;
-};
-
-
 TSharedRef<SWidget> FCustomChildrenBuilder::GenerateStructValueWidget( TSharedRef<IPropertyHandle> StructPropertyHandle )
 {
 	FStructProperty* StructProperty = CastFieldChecked<FStructProperty>( StructPropertyHandle->GetProperty() );
@@ -166,7 +131,7 @@ TSharedRef<SWidget> FCustomChildrenBuilder::GenerateStructValueWidget( TSharedRe
 	{
 		TSharedRef<IPropertyTypeCustomization> CustomStructInterface = LayoutCallback.GetCustomizationInstance();
 
-		return SNew( SStandaloneCustomStructValue, CustomStructInterface, StructPropertyHandle, ParentCategory.Pin().ToSharedRef() );
+		return SNew( SStandaloneCustomizedValueWidget, CustomStructInterface, StructPropertyHandle).ParentCategory(ParentCategory.Pin().ToSharedRef());
 	}
 	else
 	{
diff --git a/Engine/Source/Editor/PropertyEditor/Private/SSingleProperty.cpp b/Engine/Source/Editor/PropertyEditor/Private/SSingleProperty.cpp
--- a/Engine/Source/Editor/PropertyEditor/Private/SSingleProperty.cpp
+++ b/Engine/Source/Editor/PropertyEditor/Private/SSingleProperty.cpp
@@ -1,19 +1,23 @@
 // Copyright Epic Games, Inc. All Rights Reserved.
 
 #include "SSingleProperty.h"
-#include "UObject/UnrealType.h"
+
 #include "AssetThumbnail.h"
+#include "DetailPropertyRow.h"
 #include "IPropertyUtilities.h"
-#include "PropertyNode.h"
-#include "Widgets/Text/STextBlock.h"
-#include "EngineGlobals.h"
-#include "Engine/Engine.h"
-#include "Presentation/PropertyEditor/PropertyEditor.h"
+#include "Modules/ModuleManager.h"
 #include "ObjectPropertyNode.h"
 #include "PropertyEditorHelpers.h"
+#include "PropertyEditorModule.h"
+#include "PropertyNode.h"
 #include "SResetToDefaultPropertyEditor.h"
+#include "SStandaloneCustomizedValueWidget.h"
+#include "Engine/Engine.h"
+#include "Presentation/PropertyEditor/PropertyEditor.h"
 #include "ThumbnailRendering/ThumbnailManager.h"
+#include "UObject/UnrealType.h"
 #include "Widgets/Colors/SColorPicker.h"
+#include "Widgets/Text/STextBlock.h"
 
 
 class FSinglePropertyUtilities : public IPropertyUtilities
@@ -148,12 +152,14 @@ void SSingleProperty::SetObject( UObject* InObject )
 	ValueNode = RootPropertyNode->GenerateSingleChild( PropertyName );
 
 	bool bIsAcceptableProperty = false;
+	FProperty* Property = nullptr; 
 	// valid criteria for standalone properties 
 	if( ValueNode.IsValid() )
 	{
-		FProperty* Property = ValueNode->GetProperty();
-	//TODO MaterialLayers: Remove below commenting
+		//TODO MaterialLayers: Remove below commenting
 		bIsAcceptableProperty = true;
+		Property = ValueNode->GetProperty();
+		check(Property);
 		// not an array property (dynamic or static)
 		//bIsAcceptableProperty &= !( Property->IsA( FArrayProperty::StaticClass() ) || (Property->ArrayDim > 1 && ValueNode->GetArrayIndex() == INDEX_NONE) );
 		// not a struct property unless its a built in type like a vector
@@ -186,13 +192,36 @@ void SSingleProperty::SetObject( UObject* InObject )
 			];
 		}
 
-		HorizontalBox->AddSlot()
-		.Padding( 4.0f, 0.0f)
-		.FillWidth(1.0f)
-		.VAlign( VAlign_Center )
-		[
-			SNew( SPropertyValueWidget, PropertyEditor, PropertyUtilities.ToSharedRef() )
-		];
+		// For structs and other properties with a customized header
+		FPropertyEditorModule& PropertyEditorModule = FModuleManager::GetModuleChecked<FPropertyEditorModule>("PropertyEditor");
+
+		if (const FPropertyTypeLayoutCallback LayoutCallback =
+				PropertyEditorModule.GetPropertyTypeCustomization(
+					Property, *PropertyHandle, FCustomPropertyTypeLayoutMap()
+				); LayoutCallback.IsValid()
+			)
+		{
+			TSharedRef<IPropertyTypeCustomization> CustomizationInstance = LayoutCallback.GetCustomizationInstance();
+			
+			HorizontalBox->AddSlot()
+			.Padding( 4.0f, 0.0f)
+			.FillWidth(1.0f)
+			.VAlign( VAlign_Center )
+			[
+				SNew( SStandaloneCustomizedValueWidget, CustomizationInstance, PropertyHandle.ToSharedRef())
+			];
+		}
+		else // For properties without customization
+		{
+
+			HorizontalBox->AddSlot()
+			.Padding( 4.0f, 0.0f)
+			.FillWidth(1.0f)
+			.VAlign( VAlign_Center )
+			[
+				SNew( SPropertyValueWidget, PropertyEditor, PropertyUtilities.ToSharedRef() )
+			];			
+		}
 
 		if (!PropertyEditor->GetPropertyHandle()->HasMetaData(TEXT("NoResetToDefault")))
 		{
diff --git a/Engine/Source/Editor/PropertyEditor/Public/SStandaloneCustomizedValueWidget.h b/Engine/Source/Editor/PropertyEditor/Public/SStandaloneCustomizedValueWidget.h
new file mode 100644
--- /dev/null
+++ b/Engine/Source/Editor/PropertyEditor/Public/SStandaloneCustomizedValueWidget.h
@@ -0,0 +1,51 @@
+// Copyright Epic Games, Inc. All Rights Reserved.
+
+#pragma once
+
+#include "DetailGroup.h"
+#include "DetailWidgetRow.h"
+#include "Widgets/SCompoundWidget.h"
+
+/**
+ * Generates the header widget for a customized struct or other type.
+ * This widget is generally used in the property editor to display a struct as a single row, like with an FColor.
+ * Properties passed in that do not have a header customization will return a null widget.
+ */
+class PROPERTYEDITOR_API SStandaloneCustomizedValueWidget : public SCompoundWidget, public IPropertyTypeCustomizationUtils
+{
+public:
+	SLATE_BEGIN_ARGS( SStandaloneCustomizedValueWidget )
+	{}
+		/** Optional Parent Detail Category, useful to access Thumbnail Pool. */
+		SLATE_ARGUMENT( TSharedPtr<FDetailCategoryImpl>, ParentCategory)
+	SLATE_END_ARGS()
+	
+	void Construct( const FArguments& InArgs,
+		TSharedPtr<IPropertyTypeCustomization> InCustomizationInterface, TSharedRef<IPropertyHandle> InPropertyHandle)
+	{
+		ParentCategory = InArgs._ParentCategory;
+		
+		CustomizationInterface = InCustomizationInterface;
+		PropertyHandle = InPropertyHandle;
+		CustomPropertyWidget = MakeShareable(new FDetailWidgetRow);
+
+		CustomizationInterface->CustomizeHeader(InPropertyHandle, *CustomPropertyWidget, *this);
+
+		ChildSlot
+		[
+			CustomPropertyWidget->ValueWidget.Widget
+		];
+	}
+
+	virtual TSharedPtr<FAssetThumbnailPool> GetThumbnailPool() const override
+	{
+		TSharedPtr<FDetailCategoryImpl> ParentCategoryPinned = ParentCategory.Pin();
+		return ParentCategoryPinned.IsValid() ? ParentCategoryPinned->GetParentLayout().GetThumbnailPool() : NULL;
+	}
+
+private:
+	TWeakPtr<FDetailCategoryImpl> ParentCategory;
+	TSharedPtr<IPropertyTypeCustomization> CustomizationInterface;
+	TSharedPtr<IPropertyHandle> PropertyHandle;
+	TSharedPtr<FDetailWidgetRow> CustomPropertyWidget;
+};
