World Partition: select actor in viewport button is now available for loaded actors and will pin referenced actors for unloaded soft object paths to actors.

#rnx

diff --git a/Engine/Source/Editor/PropertyEditor/Private/UserInterface/PropertyEditor/SPropertyEditorAsset.cpp b/Engine/Source/Editor/PropertyEditor/Private/UserInterface/PropertyEditor/SPropertyEditorAsset.cpp
--- a/Engine/Source/Editor/PropertyEditor/Private/UserInterface/PropertyEditor/SPropertyEditorAsset.cpp
+++ b/Engine/Source/Editor/PropertyEditor/Private/UserInterface/PropertyEditor/SPropertyEditorAsset.cpp
@@ -30,6 +30,7 @@
 #include "UnrealEdGlobals.h"
 #include "Framework/Notifications/NotificationManager.h"
 #include "Widgets/Notifications/SNotificationList.h"
+#include "WorldPartition/WorldPartition.h"
 #include "FileHelpers.h"
 #include "Presentation/PropertyEditor/PropertyEditor.h"
 
@@ -604,31 +605,7 @@ void SPropertyEditorAsset::Construct(const FArguments& InArgs, const TSharedPtr<
 		];
 	}
 
-	bool bDisplayBrowse = InArgs._DisplayBrowse;
-	if (bDisplayBrowse && bIsActor)
-	{
-		FObjectOrAssetData Value;
-		GetValue(Value);
-
-		UWorld* World = Value.Object ? CastChecked<AActor>(Value.Object)->GetWorld() : nullptr;
-
-		if (!World)
-		{			
-			FSoftObjectPath MapObjectPath = FSoftObjectPath(Value.ObjectPath.GetAssetPathName(), FString());
-
-			if (UObject* MapObject = MapObjectPath.ResolveObject())
-			{
-				World = Cast<UWorld>(MapObject);
-			}
-		}
-
-		if (World && World->IsPartitionedWorld())
-		{
-			bDisplayBrowse = false;
-		}
-	}
-
-	if( bDisplayBrowse )
+	if( InArgs._DisplayBrowse )
 	{
 		ButtonBox->AddSlot()
 		.Padding( 2.0f, 0.0f )
@@ -922,7 +899,7 @@ FText SPropertyEditorAsset::OnGetToolTip() const
 			}
 			else if (State == EActorReferenceState::Exists)
 			{
-				ToolTipText = FText::Format(LOCTEXT("ExistsActorReference", "Unloaded reference to Actor ID '{Actor}', use World Partition editor to load its corresponding cells"), Args);
+				ToolTipText = FText::Format(LOCTEXT("ExistsActorReference", "Unloaded reference to Actor ID '{Actor}', use Browse to pin actor"), Args);
 			}
 			else if (State == EActorReferenceState::Unknown)
 			{
@@ -1192,12 +1169,33 @@ void SPropertyEditorAsset::OnBrowse()
 	if (bIsActor)
 	{
 		TSharedPtr<IPropertyHandle> PropertyHandleToUse = GetMostSpecificPropertyHandle();
-		if (PropertyHandleToUse && Value.Object)
+		if (PropertyHandleToUse)
 		{
-			// This code only works on loaded objects
-			if (TSharedPtr<FPropertyNode> PropertyNodeToSync = PropertyHandleToUse->GetPropertyNode())
+			// Try to resolve a potentially unloaded object
+			if (!Value.Object)
+			{
+				FSoftObjectPath MapObjectPath = FSoftObjectPath(Value.ObjectPath.GetAssetPathName(), FString());
+
+				if (UObject* MapObject = MapObjectPath.ResolveObject())
+				{
+					if (UWorld* World = Cast<UWorld>(MapObject); World && World->IsPartitionedWorld())
+					{
+						if (const FWorldPartitionActorDesc* ActorDesc = World->GetWorldPartition()->GetActorDesc(Value.ObjectPath))
+						{
+							World->GetWorldPartition()->PinActors({ ActorDesc->GetGuid() });
+							GetValue(Value);
+						}
+					}
+				}
+			}
+
+			if (Value.Object)
 			{
-				FPropertyEditor::SyncToObjectsInNode(PropertyNodeToSync);
+				// This code only works on loaded objects
+				if (TSharedPtr<FPropertyNode> PropertyNodeToSync = PropertyHandleToUse->GetPropertyNode())
+				{
+					FPropertyEditor::SyncToObjectsInNode(PropertyNodeToSync);
+				}
 			}
 		}
 	}
@@ -1535,4 +1533,4 @@ UClass* SPropertyEditorAsset::GetObjectPropertyClass(const FProperty* Property)
 	return Class;
 }
 
-#undef LOCTEXT_NAMESPACE
+#undef LOCTEXT_NAMESPACE
\ No newline at end of file
