Fixed -NoPCH -DisableUnity compile issues.

diff --git a/Engine/Source/Developer/TraceServices/Private/Analyzers/PlatformFileTraceAnalysis.cpp b/Engine/Source/Developer/TraceServices/Private/Analyzers/PlatformFileTraceAnalysis.cpp
--- a/Engine/Source/Developer/TraceServices/Private/Analyzers/PlatformFileTraceAnalysis.cpp
+++ b/Engine/Source/Developer/TraceServices/Private/Analyzers/PlatformFileTraceAnalysis.cpp
@@ -6,6 +6,7 @@
 #include "Common/Utils.h"
 #include "HAL/LowLevelMemTracker.h"
 #include "Model/FileActivity.h"
+#include "TraceServices/Utils.h"
 
 #define DEBUG_PLATFORMFILETRACE 0
 
diff --git a/Engine/Source/Developer/TraceServices/Private/Model/SbTree.cpp b/Engine/Source/Developer/TraceServices/Private/Model/SbTree.cpp
--- a/Engine/Source/Developer/TraceServices/Private/Model/SbTree.cpp
+++ b/Engine/Source/Developer/TraceServices/Private/Model/SbTree.cpp
@@ -4,6 +4,7 @@
 
 #include "Misc/FileHelper.h"
 #include "TraceServices/Containers/Allocators.h"
+#include "Model/AllocationsProvider.h"
 
 #include <limits>
 
diff --git a/Engine/Source/Editor/AnimGraph/Public/IAnimBlueprintCompilationBracketContext.h b/Engine/Source/Editor/AnimGraph/Public/IAnimBlueprintCompilationBracketContext.h
--- a/Engine/Source/Editor/AnimGraph/Public/IAnimBlueprintCompilationBracketContext.h
+++ b/Engine/Source/Editor/AnimGraph/Public/IAnimBlueprintCompilationBracketContext.h
@@ -5,6 +5,7 @@
 #include "CoreMinimal.h"
 
 class FCompilerResultsLog;
+class UAnimGraphNode_Base;
 
 // Interface passed to start/end compilation delegates
 class ANIMGRAPH_API IAnimBlueprintCompilationBracketContext
diff --git a/Engine/Source/Editor/BlueprintGraph/Classes/BlueprintNodeBinder.h b/Engine/Source/Editor/BlueprintGraph/Classes/BlueprintNodeBinder.h
--- a/Engine/Source/Editor/BlueprintGraph/Classes/BlueprintNodeBinder.h
+++ b/Engine/Source/Editor/BlueprintGraph/Classes/BlueprintNodeBinder.h
@@ -3,6 +3,7 @@
 #pragma once
 
 #include "CoreMinimal.h"
+#include "Templates/Casts.h"
 #include "UObject/Object.h"
 #include "UObject/WeakObjectPtr.h"
 #include "UObject/WeakFieldPtr.h"
diff --git a/Engine/Source/Editor/CurveEditor/Public/CurveEditorSnapMetrics.h b/Engine/Source/Editor/CurveEditor/Public/CurveEditorSnapMetrics.h
--- a/Engine/Source/Editor/CurveEditor/Public/CurveEditorSnapMetrics.h
+++ b/Engine/Source/Editor/CurveEditor/Public/CurveEditorSnapMetrics.h
@@ -4,6 +4,7 @@
 
 #include "Misc/FrameRate.h"
 #include "Math/Axis.h"
+#include "Math/Vector2D.h"
 #include "Runtime/Core/Public/Algo/MinElement.h"
 
 struct FCurveSnapMetrics
diff --git a/Engine/Source/Editor/EnvironmentLightingViewer/Private/SEnvironmentLightingViewer.cpp b/Engine/Source/Editor/EnvironmentLightingViewer/Private/SEnvironmentLightingViewer.cpp
--- a/Engine/Source/Editor/EnvironmentLightingViewer/Private/SEnvironmentLightingViewer.cpp
+++ b/Engine/Source/Editor/EnvironmentLightingViewer/Private/SEnvironmentLightingViewer.cpp
@@ -16,6 +16,7 @@
 #include "SlateOptMacros.h"
 
 #include "Widgets/Input/SButton.h"
+#include "Widgets/Input/SCheckBox.h"
 #include "Widgets/Input/SComboButton.h"
 #include "Widgets/Layout/SBorder.h"
 #include "Widgets/Layout/SScrollBorder.h"
diff --git a/Engine/Source/Editor/MovieSceneTools/Private/MovieSceneToolsModule.cpp b/Engine/Source/Editor/MovieSceneTools/Private/MovieSceneToolsModule.cpp
--- a/Engine/Source/Editor/MovieSceneTools/Private/MovieSceneToolsModule.cpp
+++ b/Engine/Source/Editor/MovieSceneTools/Private/MovieSceneToolsModule.cpp
@@ -100,6 +100,7 @@
 #include "LevelSequenceAnimSequenceLink.h"
 #include "AnimSequenceLevelSequenceLink.h"
 #include "EditorAnimUtils.h"
+#include "Animation/AnimSequence.h"
 
 #define LOCTEXT_NAMESPACE "FMovieSceneToolsModule"
 
diff --git a/Engine/Source/Editor/OverlayEditor/Public/IOverlayEditorModule.h b/Engine/Source/Editor/OverlayEditor/Public/IOverlayEditorModule.h
--- a/Engine/Source/Editor/OverlayEditor/Public/IOverlayEditorModule.h
+++ b/Engine/Source/Editor/OverlayEditor/Public/IOverlayEditorModule.h
@@ -2,6 +2,8 @@
 
 #pragma once
 
+#include "Modules/ModuleInterface.h"
+
 /**
  * Interface for the OverlayEditor module
  */
diff --git a/Engine/Source/Editor/Persona/Private/SAnimationBlendSpaceGridWidget.cpp b/Engine/Source/Editor/Persona/Private/SAnimationBlendSpaceGridWidget.cpp
--- a/Engine/Source/Editor/Persona/Private/SAnimationBlendSpaceGridWidget.cpp
+++ b/Engine/Source/Editor/Persona/Private/SAnimationBlendSpaceGridWidget.cpp
@@ -16,6 +16,7 @@
 #include "Framework/Application/SlateApplication.h"
 #include "Widgets/Images/SImage.h"
 #include "Widgets/Input/SNumericEntryBox.h"
+#include "Widgets/SToolTip.h"
 
 #include "IDetailsView.h"
 #include "UObject/StructOnScope.h"
diff --git a/Engine/Source/Editor/PropertyEditor/Private/Presentation/PropertyEditor/PropertyEditor.cpp b/Engine/Source/Editor/PropertyEditor/Private/Presentation/PropertyEditor/PropertyEditor.cpp
--- a/Engine/Source/Editor/PropertyEditor/Private/Presentation/PropertyEditor/PropertyEditor.cpp
+++ b/Engine/Source/Editor/PropertyEditor/Private/Presentation/PropertyEditor/PropertyEditor.cpp
@@ -23,6 +23,7 @@
 #include "ActorTreeItem.h"
 #include "PropertyHandleImpl.h"
 #include "IPropertyUtilities.h"
+#include "ScopedTransaction.h"
 
 #define LOCTEXT_NAMESPACE "PropertyEditor"
 
diff --git a/Engine/Source/Editor/PropertyEditor/Private/UserInterface/PropertyEditor/SPropertyEditorAsset.cpp b/Engine/Source/Editor/PropertyEditor/Private/UserInterface/PropertyEditor/SPropertyEditorAsset.cpp
--- a/Engine/Source/Editor/PropertyEditor/Private/UserInterface/PropertyEditor/SPropertyEditorAsset.cpp
+++ b/Engine/Source/Editor/PropertyEditor/Private/UserInterface/PropertyEditor/SPropertyEditorAsset.cpp
@@ -33,6 +33,7 @@
 #include "WorldPartition/WorldPartition.h"
 #include "FileHelpers.h"
 #include "Presentation/PropertyEditor/PropertyEditor.h"
+#include "AssetThumbnail.h"
 
 #define LOCTEXT_NAMESPACE "PropertyEditor"
 
diff --git a/Engine/Source/Editor/UnrealEd/Private/Factories/TIFFLoader.cpp b/Engine/Source/Editor/UnrealEd/Private/Factories/TIFFLoader.cpp
--- a/Engine/Source/Editor/UnrealEd/Private/Factories/TIFFLoader.cpp
+++ b/Engine/Source/Editor/UnrealEd/Private/Factories/TIFFLoader.cpp
@@ -6,6 +6,8 @@ DEFINE_LOG_CATEGORY_STATIC(LogTIFFLoader, Log, All);
 
 #if WITH_FREEIMAGE_LIB
 
+#include "Misc/Paths.h"
+
 #if PLATFORM_WINDOWS
 #include "Windows/AllowWindowsPlatformTypes.h"
 #endif // PLATFORM_WINDOWS
