Added a set of OnPropertyValueChanged delegates for property handles that include the propertychanged event.
This is particularly useful for details panel external structures/objects, as these don't have the "PostEditChange" calls bubble up to the actual owning object.
This enables us to detect whether an interactive change or a direct set change has taken place on an external property, so we can manually call NotifyPostChange with the correct data

diff --git a/Engine/Source/Editor/PropertyEditor/Private/PropertyHandleImpl.cpp b/Engine/Source/Editor/PropertyEditor/Private/PropertyHandleImpl.cpp
--- a/Engine/Source/Editor/PropertyEditor/Private/PropertyHandleImpl.cpp
+++ b/Engine/Source/Editor/PropertyEditor/Private/PropertyHandleImpl.cpp
@@ -534,6 +534,14 @@ void FPropertyValueImpl::SetOnPropertyValueChanged( const FSimpleDelegate& InOnP
 	}
 }
 
+void FPropertyValueImpl::SetOnPropertyValueChangedWithData(const TDelegate<void(const FPropertyChangedEvent&)>& InOnPropertyValueChanged)
+{
+	if( PropertyNode.IsValid() )
+	{
+		PropertyNode.Pin()->OnPropertyValueChangedWithData().Add(InOnPropertyValueChanged);
+	}
+}
+
 void FPropertyValueImpl::SetOnChildPropertyValueChanged( const FSimpleDelegate& InOnChildPropertyValueChanged )
 {
 	if( PropertyNode.IsValid() )
@@ -542,6 +550,14 @@ void FPropertyValueImpl::SetOnChildPropertyValueChanged( const FSimpleDelegate&
 	}
 }
 
+void FPropertyValueImpl::SetOnChildPropertyValueChangedWithData(const TDelegate<void(const FPropertyChangedEvent&)>& InOnChildPropertyValueChanged)
+{
+	if( PropertyNode.IsValid() )
+	{
+		PropertyNode.Pin()->OnChildPropertyValueChangedWithData().Add( InOnChildPropertyValueChanged );
+	}
+}
+
 void FPropertyValueImpl::SetOnPropertyValuePreChange(const FSimpleDelegate& InOnPropertyValuePreChange)
 {
 	if (PropertyNode.IsValid())
@@ -2452,11 +2468,21 @@ void FPropertyHandleBase::SetOnPropertyValueChanged( const FSimpleDelegate& InOn
 	Implementation->SetOnPropertyValueChanged(InOnPropertyValueChanged);
 }
 
+void FPropertyHandleBase::SetOnPropertyValueChangedWithData(const TDelegate<void(const FPropertyChangedEvent&)>& InOnPropertyValueChanged)
+{
+	Implementation->SetOnPropertyValueChangedWithData(InOnPropertyValueChanged);
+}
+
 void FPropertyHandleBase::SetOnChildPropertyValueChanged( const FSimpleDelegate& InOnChildPropertyValueChanged )
 {
 	Implementation->SetOnChildPropertyValueChanged( InOnChildPropertyValueChanged );
 }
 
+void FPropertyHandleBase::SetOnChildPropertyValueChangedWithData(const TDelegate<void(const FPropertyChangedEvent&)>& InOnChildPropertyValueChanged)
+{
+	Implementation->SetOnChildPropertyValueChangedWithData(InOnChildPropertyValueChanged);
+}
+
 void FPropertyHandleBase::SetOnPropertyValuePreChange(const FSimpleDelegate& InOnPropertyValuePreChange)
 {
 	Implementation->SetOnPropertyValuePreChange(InOnPropertyValuePreChange);
diff --git a/Engine/Source/Editor/PropertyEditor/Private/PropertyHandleImpl.h b/Engine/Source/Editor/PropertyEditor/Private/PropertyHandleImpl.h
--- a/Engine/Source/Editor/PropertyEditor/Private/PropertyHandleImpl.h
+++ b/Engine/Source/Editor/PropertyEditor/Private/PropertyHandleImpl.h
@@ -139,12 +139,12 @@ public:
 	 * Sets a delegate to call when the property value changes
 	 */
 	void SetOnPropertyValueChanged( const FSimpleDelegate& InOnPropertyValueChanged );
-	
+	void SetOnPropertyValueChangedWithData(const TDelegate<void(const FPropertyChangedEvent&)>& InOnPropertyValueChanged);
 	/**
 	 * Sets a delegate to call when the propery value of a child changes
 	 */
 	void SetOnChildPropertyValueChanged( const FSimpleDelegate& InOnChildPropertyValueChanged );
-
+	void SetOnChildPropertyValueChangedWithData(const TDelegate<void(const FPropertyChangedEvent&)>& InOnChildPropertyValueChanged);
 	/**
 	 * Sets a delegate to call when the property value is about to change
 	 */
@@ -421,7 +421,9 @@ public:
 	virtual bool IsEditConst() const override;
 	virtual bool IsEditable() const override;
 	virtual void SetOnPropertyValueChanged( const FSimpleDelegate& InOnPropertyValueChanged ) override;
+	virtual void SetOnPropertyValueChangedWithData(const TDelegate<void(const FPropertyChangedEvent&)>& InOnPropertyValueChanged) override;
 	virtual void SetOnChildPropertyValueChanged( const FSimpleDelegate& InOnPropertyValueChanged ) override;
+	virtual void SetOnChildPropertyValueChangedWithData( const TDelegate<void(const FPropertyChangedEvent&)>& InOnChildPropertyValueChanged ) override;
 	virtual void SetOnPropertyValuePreChange(const FSimpleDelegate& InOnPropertyValuePreChange) override;
 	virtual void SetOnChildPropertyValuePreChange(const FSimpleDelegate& InOnPropertyValuePreChange) override;
 	virtual void SetOnPropertyResetToDefault(const FSimpleDelegate& InOnPropertyResetToDefault) override;
diff --git a/Engine/Source/Editor/PropertyEditor/Private/PropertyNode.cpp b/Engine/Source/Editor/PropertyEditor/Private/PropertyNode.cpp
--- a/Engine/Source/Editor/PropertyEditor/Private/PropertyNode.cpp
+++ b/Engine/Source/Editor/PropertyEditor/Private/PropertyNode.cpp
@@ -2682,6 +2682,7 @@ void FPropertyNode::NotifyPostChange( FPropertyChangedEvent& InPropertyChangedEv
 	}
 
 	// Broadcast the change to any listeners
+	BroadcastPropertyChangedDelegates(InPropertyChangedEvent);
 	BroadcastPropertyChangedDelegates();
 
 	// Call through to the property window's notify hook.
@@ -2746,7 +2747,23 @@ void FPropertyNode::BroadcastPropertyChangedDelegates()
 
 		LocalParentNode = LocalParentNode->GetParentNode();
 	}
+}
+
+void FPropertyNode::BroadcastPropertyChangedDelegates(const FPropertyChangedEvent& Event)
+{
+	PropertyValueChangedDelegate.Broadcast(Event);
+
+	// Walk through the parents and broadcast
+	FPropertyNode* LocalParentNode = GetParentNode();
+	while( LocalParentNode )
+	{
+		if( LocalParentNode->OnChildPropertyValueChangedWithData().IsBound() )
+		{
+			LocalParentNode->OnChildPropertyValueChangedWithData().Broadcast(Event);
+		}
 
+		LocalParentNode = LocalParentNode->GetParentNode();
+	}
 }
 
 void FPropertyNode::BroadcastPropertyPreChangeDelegates()
diff --git a/Engine/Source/Editor/PropertyEditor/Private/PropertyNode.h b/Engine/Source/Editor/PropertyEditor/Private/PropertyNode.h
--- a/Engine/Source/Editor/PropertyEditor/Private/PropertyNode.h
+++ b/Engine/Source/Editor/PropertyEditor/Private/PropertyNode.h
@@ -646,10 +646,14 @@ public:
 
 	/** Broadcasts when a property value changes */
 	DECLARE_EVENT(FPropertyNode, FPropertyValueChangedEvent);
+	/** Broadcasts when a property value changes, but additionally includes the property changed event.*/
+	DECLARE_MULTICAST_DELEGATE_OneParam(FPropertyValueChangedWithData, const FPropertyChangedEvent&)
 	FPropertyValueChangedEvent& OnPropertyValueChanged() { return PropertyValueChangedEvent; }
-
+	FPropertyValueChangedWithData& OnPropertyValueChangedWithData() { return PropertyValueChangedDelegate; }
+	
 	/** Broadcasts when a child of this property changes */
 	FPropertyValueChangedEvent& OnChildPropertyValueChanged() { return ChildPropertyValueChangedEvent; }
+	FPropertyValueChangedWithData& OnChildPropertyValueChangedWithData() { return ChildPropertyValueChangedDelegate; }
 
 	/** Broadcasts when a property value changes */
 	DECLARE_EVENT(FPropertyNode, FPropertyValuePreChangeEvent);
@@ -1019,6 +1023,12 @@ protected:
 	 */
 	void BroadcastPropertyChangedDelegates();
 
+	/**
+	 * Helper function for derived members to be able to 
+	 * broadcast property changed notifications including property changed event data
+	 */
+	void BroadcastPropertyChangedDelegates(const FPropertyChangedEvent& Event);
+
 
 	/**
 	* Helper function for derived members to be able to
@@ -1078,9 +1088,13 @@ protected:
 
 	/** Called when this node's property value has changed (called during NotifyPostChange) */
 	FPropertyValueChangedEvent PropertyValueChangedEvent;
+	/** Called when this node's property value has changed with the property changed event data as payload (called during NotifyPostChange) */
+	FPropertyValueChangedWithData PropertyValueChangedDelegate;
 	
 	/** Called when a child's property value has changed */
 	FPropertyValueChangedEvent ChildPropertyValueChangedEvent;
+	/** Called when a child's property value has changed with the property changed event data as payload */
+	FPropertyValueChangedWithData ChildPropertyValueChangedDelegate;
 
 	/** Called when the property is reset to default */
 	FPropertyResetToDefaultEvent PropertyResetToDefaultEvent;
diff --git a/Engine/Source/Editor/PropertyEditor/Public/PropertyHandle.h b/Engine/Source/Editor/PropertyEditor/Public/PropertyHandle.h
--- a/Engine/Source/Editor/PropertyEditor/Public/PropertyHandle.h
+++ b/Engine/Source/Editor/PropertyEditor/Public/PropertyHandle.h
@@ -222,12 +222,26 @@ public:
 	 */
 	virtual void SetOnPropertyValueChanged( const FSimpleDelegate& InOnPropertyValueChanged ) = 0;
 	
+	/**
+	 * Sets a delegate to call when the value of the property is changed including the property changed event as parameter
+	 * 
+	 * @param InOnPropertyValueChanged The delegate to call. Needs to be of type void with a const FPropertyChangedEvent& parameter
+	 */
+	virtual void SetOnPropertyValueChangedWithData( const TDelegate<void(const FPropertyChangedEvent&)>& InOnPropertyValueChanged ) = 0;
+
 	/**
 	 * Sets a delegate to call when the value of the property of a child is changed
 	 * 
 	 * @param InOnChildPropertyValueChanged	The delegate to call
 	 */
 	virtual void SetOnChildPropertyValueChanged( const FSimpleDelegate& InOnChildPropertyValueChanged ) = 0;
+	
+	/**
+	 * Sets a delegate to call when the value of the property of a child is changed including the property changed event as parameter
+	 * 
+	 * @param InOnPropertyValueChanged The delegate to call. Needs to be of type void with a const FPropertyChangedEvent& parameter
+	 */
+	virtual void SetOnChildPropertyValueChangedWithData( const TDelegate<void(const FPropertyChangedEvent&)>& InOnPropertyValueChanged ) = 0;
 
 	/**
 	 * Sets a delegate to call when the value of the property is about to be changed
