PropertyEditor: Fixed widgets added with OnAddMaterialItemViewExtraBottomWidget getting clipped incorrectly.
#rnx

diff --git a/Engine/Source/Editor/PropertyEditor/Private/MaterialList.cpp b/Engine/Source/Editor/PropertyEditor/Private/MaterialList.cpp
--- a/Engine/Source/Editor/PropertyEditor/Private/MaterialList.cpp
+++ b/Engine/Source/Editor/PropertyEditor/Private/MaterialList.cpp
@@ -157,6 +157,7 @@
 		for (TSharedPtr<SWidget> Widget : Widgets)
 		{
 			WidgetContainer->AddSlot()
+				.AutoHeight()
 			[
 				Widget.ToSharedRef()
 			];
