DetailsView: Show a modified badge when any view option that filters down the results is active

diff --git a/Engine/Source/Editor/PropertyEditor/Private/SDetailsView.cpp b/Engine/Source/Editor/PropertyEditor/Private/SDetailsView.cpp
--- a/Engine/Source/Editor/PropertyEditor/Private/SDetailsView.cpp
+++ b/Engine/Source/Editor/PropertyEditor/Private/SDetailsView.cpp
@@ -28,6 +28,7 @@
 #include "Widgets/Input/SSearchBox.h"
 #include "Widgets/Input/SSegmentedControl.h"
 #include "Widgets/Layout/SWrapBox.h"
+#include "Widgets/Images/SLayeredImage.h"
 
 #define LOCTEXT_NAMESPACE "SDetailsView"
 
@@ -321,6 +322,13 @@ void SDetailsView::Construct(const FArguments& InArgs, const FDetailsViewArgs& I
 
 	if (DetailsViewArgs.bShowOptions)
 	{
+		TSharedPtr<SLayeredImage> FilterImage = SNew(SLayeredImage)
+		 .Image(FAppStyle::Get().GetBrush("DetailsView.ViewOptions"))
+		 .ColorAndOpacity(FSlateColor::UseForeground());
+
+		// Badge the filter icon if there are filters active
+		FilterImage->AddLayer(TAttribute<const FSlateBrush*>(this, &SDetailsView::GetViewOptionsBadgeIcon));
+		
 		FilterRowHBox->AddSlot()
 			.Padding(0)
 			.HAlign(HAlign_Right)
@@ -339,9 +347,7 @@ void SDetailsView::Construct(const FArguments& InArgs, const FDetailsViewArgs& I
 				]
 				.ButtonContent()
 				[
-					SNew(SImage)
-					.ColorAndOpacity(FSlateColor::UseForeground())
-					.Image(FAppStyle::Get().GetBrush("DetailsView.ViewOptions"))
+					FilterImage.ToSharedRef()
 				]
 			];
 	}
@@ -1416,4 +1422,15 @@ TMap<FName, FText> SDetailsView::GetAllSections() const
 	return MoveTemp(AllSections);
 }
 
+const FSlateBrush* SDetailsView::GetViewOptionsBadgeIcon() const
+{
+	// Badge the icon if any view option that narrows down the results is checked
+	bool bHasBadge = (DetailsViewArgs.bShowModifiedPropertiesOption && IsShowOnlyModifiedChecked() )
+					|| (DetailsViewArgs.bShowDifferingPropertiesOption && IsShowOnlyAllowedChecked() )
+					|| (DetailsViewArgs.bShowKeyablePropertiesOption && IsShowKeyableChecked() )
+					|| (DetailsViewArgs.bShowAnimatedPropertiesOption && IsShowAnimatedChecked() );
+
+	return bHasBadge ? FAppStyle::Get().GetBrush("Icons.BadgeModified") : nullptr;
+}
+
 #undef LOCTEXT_NAMESPACE
diff --git a/Engine/Source/Editor/PropertyEditor/Private/SDetailsView.h b/Engine/Source/Editor/PropertyEditor/Private/SDetailsView.h
--- a/Engine/Source/Editor/PropertyEditor/Private/SDetailsView.h
+++ b/Engine/Source/Editor/PropertyEditor/Private/SDetailsView.h
@@ -190,6 +190,9 @@ private:
 	/** Get the currently selected section. */
 	ECheckBoxState IsSectionChecked(FName Section) const;
 
+	/** Get the badge for the view options icon (if there are any options active */
+	const FSlateBrush* GetViewOptionsBadgeIcon() const;
+
 private:
 	/** The filter for objects viewed by this details panel */
 	TSharedPtr<FDetailsViewObjectFilter> ObjectFilter;
