Fix an editor issue where expanding a container of non-nullable objects would append a null element, now we create a new object to assign to the new element

diff --git a/Engine/Source/Editor/PropertyEditor/Private/PropertyHandleImpl.cpp b/Engine/Source/Editor/PropertyEditor/Private/PropertyHandleImpl.cpp
--- a/Engine/Source/Editor/PropertyEditor/Private/PropertyHandleImpl.cpp
+++ b/Engine/Source/Editor/PropertyEditor/Private/PropertyHandleImpl.cpp
@@ -1027,18 +1027,54 @@ void FPropertyValueImpl::AddChild()
 						{
 							FScriptArrayHelper	ArrayHelper(Array, Addr);
 							Index = ArrayHelper.AddValue();
+
+							// check whether the inner type is flagged as a non-nullable. if so, create it.
+							FObjectProperty* InnerObjectProperty = CastField<FObjectProperty>(Array->Inner);
+							if (InnerObjectProperty && InnerObjectProperty->HasAnyPropertyFlags(CPF_NonNullable))
+							{
+								UObject* NewItem = NewObject<UObject>(Obj, InnerObjectProperty->PropertyClass);
+								InnerObjectProperty->SetObjectPropertyValue(ArrayHelper.GetRawPtr(Index), NewItem);
+							}
 						}
 						else if (Set)
 						{
 							FScriptSetHelper	SetHelper(Set, Addr);
 							Index = SetHelper.AddDefaultValue_Invalid_NeedsRehash();
-							SetHelper.Rehash();
 
+							// check whether the element type is flagged as a non-nullable. if so, create it.
+							FObjectProperty* ElementObjectProperty = CastField<FObjectProperty>(Set->ElementProp);
+							if (ElementObjectProperty && ElementObjectProperty->HasAnyPropertyFlags(CPF_NonNullable))
+							{
+								UObject* NewItem = NewObject<UObject>(Obj, ElementObjectProperty->PropertyClass);
+								ElementObjectProperty->SetObjectPropertyValue(SetHelper.GetElementPtr(Index), NewItem);
+							}
+
+							SetHelper.Rehash();
 						}
 						else if (Map)
 						{
 							FScriptMapHelper	MapHelper(Map, Addr);
 							Index = MapHelper.AddDefaultValue_Invalid_NeedsRehash();
+
+							// check whether the key or value type is flagged as a non-nullable. if so, create it.
+							{
+								FObjectProperty* KeyObjectProperty = CastField<FObjectProperty>(Map->KeyProp);
+								if (KeyObjectProperty && KeyObjectProperty->HasAnyPropertyFlags(CPF_NonNullable))
+								{
+									UObject* NewItem = NewObject<UObject>(Obj, KeyObjectProperty->PropertyClass);
+									KeyObjectProperty->SetObjectPropertyValue(MapHelper.GetKeyPtr(Index), NewItem);
+								}
+							}
+
+							{
+								FObjectProperty* ValueObjectProperty = CastField<FObjectProperty>(Map->ValueProp);
+								if (ValueObjectProperty && ValueObjectProperty->HasAnyPropertyFlags(CPF_NonNullable))
+								{
+									UObject* NewItem = NewObject<UObject>(Obj, ValueObjectProperty->PropertyClass);
+									ValueObjectProperty->SetObjectPropertyValue(MapHelper.GetValuePtr(Index), NewItem);
+								}
+							}
+
 							MapHelper.Rehash();
 							bAddedMapEntry = true;
 						}
