Limit sparse data properties to the common base class when multiple objects with potentially unrelated types are selected into the property editor.

diff --git a/Engine/Source/Editor/PropertyEditor/Private/ObjectPropertyNode.cpp b/Engine/Source/Editor/PropertyEditor/Private/ObjectPropertyNode.cpp
--- a/Engine/Source/Editor/PropertyEditor/Private/ObjectPropertyNode.cpp
+++ b/Engine/Source/Editor/PropertyEditor/Private/ObjectPropertyNode.cpp
@@ -581,43 +581,39 @@ void FObjectPropertyNode::InternalInitChildNodes( FName SinglePropertyName )
 	const bool bShouldShowHiddenProperties = !!HasNodeFlags(EPropertyNodeFlags::ShouldShowHiddenProperties);
 	const bool bShouldShowDisableEditOnInstance = !!HasNodeFlags(EPropertyNodeFlags::ShouldShowDisableEditOnInstance);
 
-	for (TFieldIterator<FProperty> It(BaseClass.Get()); It; ++It)
+	if (UClass* ResolvedBaseClass = BaseClass.Get())
 	{
-		GetCategoryProperties(ClassesToConsider, *It, bShouldShowDisableEditOnInstance, bShouldShowHiddenProperties, CategoriesFromBlueprints, CategoriesFromProperties, SortedCategories);
-	}
+		for (TFieldIterator<FProperty> It(ResolvedBaseClass); It; ++It)
+		{
+			GetCategoryProperties(ClassesToConsider, *It, bShouldShowDisableEditOnInstance, bShouldShowHiddenProperties, CategoriesFromBlueprints, CategoriesFromProperties, SortedCategories);
+		}
 
-	for (UClass* Class : ClassesToConsider)
-	{
-		if (Class)
+		if (UScriptStruct* SparseClassDataStruct = ResolvedBaseClass->GetSparseClassDataStruct())
 		{
-			UScriptStruct* SparseClassDataStruct = Class->GetSparseClassDataStruct();
-			if (SparseClassDataStruct)
+			SparseClassDataInstances.Add(ResolvedBaseClass, TTuple<UScriptStruct*, void*>(SparseClassDataStruct, ResolvedBaseClass->GetOrCreateSparseClassData()));
+		
+			for (TFieldIterator<FProperty> It(SparseClassDataStruct); It; ++It)
 			{
-				SparseClassDataInstances.Add(Class, TTuple<UScriptStruct*, void*>(SparseClassDataStruct, Class->GetOrCreateSparseClassData()));
-
-				for (TFieldIterator<FProperty> It(SparseClassDataStruct); It; ++It)
-				{
-					GetCategoryProperties(ClassesToConsider, *It, bShouldShowDisableEditOnInstance, bShouldShowHiddenProperties, CategoriesFromBlueprints, CategoriesFromProperties, SortedCategories);
-				}
+				GetCategoryProperties(ClassesToConsider, *It, bShouldShowDisableEditOnInstance, bShouldShowHiddenProperties, CategoriesFromBlueprints, CategoriesFromProperties, SortedCategories);
 			}
 		}
-	}
 
-	// Categories from the Blueprint class may include categories with no associated properties, so remove those ones
-	for ( const FName& CategoryName : CategoriesFromBlueprints )
-	{
-		if ( !CategoriesFromProperties.Contains(CategoryName) )
+		// Categories from the Blueprint class may include categories with no associated properties, so remove those ones
+		for (const FName& CategoryName : CategoriesFromBlueprints)
 		{
-			SortedCategories.Remove(CategoryName);
+			if (!CategoriesFromProperties.Contains(CategoryName))
+			{
+				SortedCategories.Remove(CategoryName);
+			}
 		}
-	}
 
-	// Remove any prioritized categories that don't actually exist
-	for (const FName& PrioritizeCategory : PrioritizeCategories)
-	{
-		if (!CategoriesFromProperties.Contains(PrioritizeCategory))
+		// Remove any prioritized categories that don't actually exist
+		for (const FName& PrioritizeCategory : PrioritizeCategories)
 		{
-			SortedCategories.Remove(PrioritizeCategory);
+			if (!CategoriesFromProperties.Contains(PrioritizeCategory))
+			{
+				SortedCategories.Remove(PrioritizeCategory);
+			}
 		}
 	}
 
