Expanding the Advanced node no longer causes previously collapsed nodes to re-expand.

The Advanced node is really just a filter, and we call RestoreExpandedItems() after filtering nodes. Now we also call SaveExpandedItems() first.

Reported in UDN: https://epicgames.lightning.force.com/lightning/r/Case/5004z00001fDN3RAAW/view

diff --git a/Engine/Source/Editor/PropertyEditor/Private/SDetailsViewBase.cpp b/Engine/Source/Editor/PropertyEditor/Private/SDetailsViewBase.cpp
--- a/Engine/Source/Editor/PropertyEditor/Private/SDetailsViewBase.cpp
+++ b/Engine/Source/Editor/PropertyEditor/Private/SDetailsViewBase.cpp
@@ -98,7 +98,7 @@ void SDetailsViewBase::SetRootExpansionStates(const bool bExpand, const bool bRe
 			Children.Reset();
 			(*Iter)->GetChildren(Children);
 
-			for(TSharedRef<class FDetailTreeNode>& Child : Children)
+			for(TSharedRef<FDetailTreeNode>& Child : Children)
 			{
 				SetNodeExpansionState(Child, bExpand, bRecurse);
 			}
@@ -127,7 +127,7 @@ void SDetailsViewBase::SetNodeExpansionState(TSharedRef<FDetailTreeNode> InTreeN
 		if (bRecursive)
 		{
 			for (int32 ChildIndex = 0; ChildIndex < Children.Num(); ++ChildIndex)
-			{
+					{
 				TSharedRef<FDetailTreeNode> Child = Children[ChildIndex];
 
 				SetNodeExpansionState(Child, bIsItemExpanded, bRecursive);
@@ -778,7 +778,7 @@ TSharedPtr<SWidget> SDetailsViewBase::GetFilterAreaWidget()
 	return DetailsViewArgs.bCustomFilterAreaLocation ? FilterRow : nullptr;
 }
 
-TSharedPtr<class FUICommandList> SDetailsViewBase::GetHostCommandList() const
+TSharedPtr<FUICommandList> SDetailsViewBase::GetHostCommandList() const
 {
 	return DetailsViewArgs.HostCommandList;
 }
@@ -1063,13 +1063,11 @@ void SDetailsViewBase::Tick( const FGeometry& AllottedGeometry, const double InC
 	if (bUpdateFilteredDetails)
 	{
 		UpdateFilteredDetails();
-
-		RestoreAllExpandedItems();
 	}
 
-	for(FDetailLayoutData& LayoutData : DetailLayouts)
+	for (FDetailLayoutData& LayoutData : DetailLayouts)
 	{
-		if(LayoutData.DetailLayout.IsValid())
+		if (LayoutData.DetailLayout.IsValid())
 		{
 			LayoutData.DetailLayout->Tick(InDeltaTime);
 		}
@@ -1370,43 +1368,44 @@ void SDetailsViewBase::HandleNodeAnimationComplete()
 	CurrentlyAnimatingNodePath = nullptr;
 }
 
+void SDetailsViewBase::FilterRootNode(const TSharedPtr<FComplexPropertyNode>& RootNode)
+{
+	if (RootNode.IsValid())
+	{
+		SaveExpandedItems(RootNode.ToSharedRef());
+
+		RootNode->FilterNodes(CurrentFilter.FilterStrings);
+		RootNode->ProcessSeenFlags(true);
+
+		RestoreExpandedItems(RootNode.ToSharedRef());
+	}
+}
+
 void SDetailsViewBase::UpdateFilteredDetails()
 {
 	RootTreeNodes.Reset();
 
 	FDetailNodeList InitialRootNodeList;
-	
 	NumVisibleTopLevelObjectNodes = 0;
 	
-	FRootPropertyNodeList& RootPropertyNodes = GetRootNodes();
-	for(int32 RootNodeIndex = 0; RootNodeIndex < RootPropertyNodes.Num(); ++RootNodeIndex)
+	const FRootPropertyNodeList& RootPropertyNodes = GetRootNodes();
+	for (int32 RootNodeIndex = 0; RootNodeIndex < RootPropertyNodes.Num(); ++RootNodeIndex)
 	{
-		const TSharedPtr<FComplexPropertyNode>& RootPropertyNode = RootPropertyNodes[RootNodeIndex];
-		if(RootPropertyNode.IsValid())
+		if (RootPropertyNodes[RootNodeIndex].IsValid())
 		{
-			RootPropertyNode->FilterNodes(CurrentFilter.FilterStrings);
-			RootPropertyNode->ProcessSeenFlags(true);
-
-			RestoreExpandedItems(RootPropertyNode.ToSharedRef());
+			FilterRootNode(RootPropertyNodes[RootNodeIndex]);
 
 			const TSharedPtr<FDetailLayoutBuilderImpl>& DetailLayout = DetailLayouts[RootNodeIndex].DetailLayout;
-			if(DetailLayout.IsValid())
+			if (DetailLayout.IsValid())
 			{
-				FRootPropertyNodeList& ExternalRootPropertyNodes = DetailLayout->GetExternalRootPropertyNodes();
-				for (const TSharedPtr<FComplexPropertyNode>& ExternalRootNode : ExternalRootPropertyNodes)
+				for (const TSharedPtr<FComplexPropertyNode>& ExternalRootNode : DetailLayout->GetExternalRootPropertyNodes())
 				{
-					if (ExternalRootNode.IsValid())
-					{
-						ExternalRootNode->FilterNodes(CurrentFilter.FilterStrings);
-						ExternalRootNode->ProcessSeenFlags(true);
-
-						RestoreExpandedItems(ExternalRootNode.ToSharedRef());
-					}
+					FilterRootNode(ExternalRootNode);
 				}
 
 				DetailLayout->FilterDetailLayout(CurrentFilter);
 
-				FDetailNodeList& LayoutRoots = DetailLayout->GetFilteredRootTreeNodes();
+				const FDetailNodeList& LayoutRoots = DetailLayout->GetFilteredRootTreeNodes();
 				if (LayoutRoots.Num() > 0)
 				{
 					// A top level object nodes has a non-filtered away root so add one to the total number we have
@@ -1419,9 +1418,9 @@ void SDetailsViewBase::UpdateFilteredDetails()
 	}
 
 	// for multiple top level object we need to do a secondary pass on top level object nodes after we have determined if there is any nodes visible at all.  If there are then we ask the details panel if it wants to show childen
-	for(TSharedRef<class FDetailTreeNode> RootNode : InitialRootNodeList)
+	for (const TSharedRef<FDetailTreeNode>& RootNode : InitialRootNodeList)
 	{
-		if(RootNode->ShouldShowOnlyChildren())
+		if (RootNode->ShouldShowOnlyChildren())
 		{
 			RootNode->GetChildren(RootTreeNodes);
 		}
diff --git a/Engine/Source/Editor/PropertyEditor/Private/SDetailsViewBase.h b/Engine/Source/Editor/PropertyEditor/Private/SDetailsViewBase.h
--- a/Engine/Source/Editor/PropertyEditor/Private/SDetailsViewBase.h
+++ b/Engine/Source/Editor/PropertyEditor/Private/SDetailsViewBase.h
@@ -117,7 +117,7 @@ public:
 	}
 	/** IDetailsViewPrivate interface */
 	virtual void RerunCurrentFilter() override;
-	void SetNodeExpansionState(TSharedRef<FDetailTreeNode> InTreeNode, bool bIsItemExpanded, bool bRecursive) override;
+	void SetNodeExpansionState(TSharedRef<FDetailTreeNode> InTreeNode, bool bExpand, bool bRecursive) override;
 	void SaveCustomExpansionState(const FString& NodePath, bool bIsExpanded) override;
 	bool GetCustomSavedExpansionState(const FString& NodePath) const override;
 	virtual void NotifyFinishedChangingProperties(const FPropertyChangedEvent& PropertyChangedEvent) override;
@@ -273,9 +273,10 @@ protected:
 	/** Called when show animated is clicked */
 	void OnShowAnimatedClicked();
 
-	/**
-	* Updates the details with the passed in filter
-	*/
+	/** Apply the CurrentFilter to the given root node. */
+	void FilterRootNode(const TSharedPtr<FComplexPropertyNode>& RootNode);
+
+	/** Updates the details with the CurrentFilter */
 	void UpdateFilteredDetails();
 
 	/** Called when the filter text changes.  This filters specific property nodes out of view */
@@ -285,7 +286,7 @@ protected:
 	void OnFilterTextCommitted(const FText& InSearchText, ETextCommit::Type InCommitType);
 
 	/** Called when the list of currently differing properties changes */
-	virtual void UpdatePropertyAllowList(const TSet<FPropertyPath> InAllowedProperties) override { CurrentFilter.PropertyAllowList = InAllowedProperties; }
+	virtual void UpdatePropertyAllowList(const TSet<FPropertyPath>& InAllowedProperties) override { CurrentFilter.PropertyAllowList = InAllowedProperties; }
 
 	virtual TSharedPtr<SWidget> GetNameAreaWidget() override;
 	virtual void SetNameAreaCustomContent(TSharedRef<SWidget>& InCustomContent) override;
diff --git a/Engine/Source/Editor/PropertyEditor/Public/IDetailsView.h b/Engine/Source/Editor/PropertyEditor/Public/IDetailsView.h
--- a/Engine/Source/Editor/PropertyEditor/Public/IDetailsView.h
+++ b/Engine/Source/Editor/PropertyEditor/Public/IDetailsView.h
@@ -12,6 +12,8 @@ class AActor;
 class FNotifyHook;
 class FPropertyPath;
 class FTabManager;
+class FUICommandList;
+class IClassViewerFilter;
 class IDetailKeyframeHandler;
 class IDetailPropertyExtensionHandler;
 class IDetailRootObjectCustomization;
@@ -67,7 +69,7 @@ public:
 	 *
 	 * @param InRootObjectCustomization	If null is passed in, the customization will be removed
 	 */
-	virtual void SetRootObjectCustomizationInstance(TSharedPtr<class IDetailRootObjectCustomization> InRootObjectCustomization) = 0;
+	virtual void SetRootObjectCustomizationInstance(TSharedPtr<IDetailRootObjectCustomization> InRootObjectCustomization) = 0;
 
 	/**
 	 * Sets the objects this details view is viewing
@@ -215,7 +217,7 @@ public:
 	/**
 	 * Sets the set of properties that are considered differing, used when filtering out identical properties
 	 */
-	virtual void UpdatePropertyAllowList(const TSet<FPropertyPath> InAllowedProperties) = 0;
+	virtual void UpdatePropertyAllowList(const TSet<FPropertyPath>& InAllowedProperties) = 0;
 
 	/** Returns the name area widget used to display object naming functionality so it can be placed in a custom location.  Note FDetailsViewArgs.bCustomNameAreaLocation must be true */
 	virtual TSharedPtr<SWidget> GetNameAreaWidget() = 0;
@@ -227,7 +229,7 @@ public:
 	virtual TSharedPtr<SWidget> GetFilterAreaWidget() = 0;
 
 	/** Returns the command list of the hosting toolkit (can be nullptr if the widget that contains the details panel didn't route a command list in) */
-	virtual TSharedPtr<class FUICommandList> GetHostCommandList() const = 0;
+	virtual TSharedPtr<FUICommandList> GetHostCommandList() const = 0;
 
 	/** Returns the tab manager of the hosting toolkit (can be nullptr if the details panel is not hosted within a tab) */
 	virtual TSharedPtr<FTabManager> GetHostTabManager() const = 0;
@@ -248,7 +250,7 @@ public:
 	virtual void SetObjectFilter(TSharedPtr<FDetailsViewObjectFilter> InFilter) = 0;
 
 	/** Sets the custom filter(s) to be used when selecting values for class properties in this view. */
-	virtual void SetClassViewerFilters(const TArray<TSharedRef<class IClassViewerFilter>>& InFilters) = 0;
+	virtual void SetClassViewerFilters(const TArray<TSharedRef<IClassViewerFilter>>& InFilters) = 0;
 
 	/** Allows other systems to add a custom filter in the details panel */
 	virtual void SetCustomFilterDelegate(FSimpleDelegate InDelegate) = 0;
