Removed inline forward declarations in SPropertyEditorAsset.h

#rnx

diff --git a/Engine/Source/Editor/PropertyEditor/Private/UserInterface/PropertyEditor/SPropertyEditorAsset.h b/Engine/Source/Editor/PropertyEditor/Private/UserInterface/PropertyEditor/SPropertyEditorAsset.h
--- a/Engine/Source/Editor/PropertyEditor/Private/UserInterface/PropertyEditor/SPropertyEditorAsset.h
+++ b/Engine/Source/Editor/PropertyEditor/Private/UserInterface/PropertyEditor/SPropertyEditorAsset.h
@@ -19,6 +19,11 @@
 #include "Presentation/PropertyEditor/PropertyEditor.h"
 #include "PropertyCustomizationHelpers.h"
 
+class FAssetThumbnail;
+class FPropertyEditor;
+class IPropertyHandle;
+class SBorder;
+class SComboButton;
 class UFactory;
 
 /**
@@ -76,7 +81,7 @@ public:
 
 	SLATE_END_ARGS()
 
-	static bool Supports( const TSharedRef< class FPropertyEditor >& InPropertyEditor );
+	static bool Supports( const TSharedRef<FPropertyEditor>& InPropertyEditor );
 	static bool Supports( const FProperty* NodeProperty );
 
 	/**
@@ -215,7 +220,7 @@ private:
 	 * Delegate for handling selection in the asset browser.
 	 * @param	Object	The chosen asset
 	 */
-	void OnAssetSelected( const struct FAssetData& AssetData );
+	void OnAssetSelected( const FAssetData& AssetData );
 
 	/** 
 	 * Delegate for handling selection in the scene outliner.
@@ -305,7 +310,7 @@ private:
 	bool CanSetBasedOnAssetReferenceFilter( const FAssetData& InAssetData, FText* OutOptionalFailureReason = nullptr ) const;
 
 	/** @return Returns the property handle most relevant */
-	TSharedPtr<class IPropertyHandle> GetMostSpecificPropertyHandle() const;
+	TSharedPtr<IPropertyHandle> GetMostSpecificPropertyHandle() const;
 
 	/**
 	 * Gets the class of the supplied property for use within the PropertyEditorAsset widget. Asserts if the property
@@ -335,13 +340,13 @@ private:
 private:
 
 	/** Main combobutton */
-	TSharedPtr< class SComboButton > AssetComboButton;
+	TSharedPtr<SComboButton> AssetComboButton;
 
 	/** The border surrounding the thumbnail image. */
-	TSharedPtr< class SBorder > ThumbnailBorder;
+	TSharedPtr<SBorder> ThumbnailBorder;
 
 	/** The property editor, if any */
-	TSharedPtr<class FPropertyEditor> PropertyEditor;
+	TSharedPtr<FPropertyEditor> PropertyEditor;
 
 	/** Path to the object being edited instead of accessing the value directly with a property handle */
 	TAttribute<FString> ObjectPath;
@@ -386,10 +391,10 @@ private:
 	FOnShouldFilterAsset OnShouldFilterAsset;
 
 	/** Thumbnail for the asset */
-	TSharedPtr<class FAssetThumbnail> AssetThumbnail;
+	TSharedPtr<FAssetThumbnail> AssetThumbnail;
 
 	/** The property handle, if any */
-	TSharedPtr<class IPropertyHandle> PropertyHandle;
+	TSharedPtr<IPropertyHandle> PropertyHandle;
 
 	/*
 	 * The reference object on which the picker will assign the picked asset, if any.
