Removed inline forward declarations from DetailTreeNode.h

#rnx

diff --git a/Engine/Source/Editor/PropertyEditor/Private/DetailTreeNode.h b/Engine/Source/Editor/PropertyEditor/Private/DetailTreeNode.h
--- a/Engine/Source/Editor/PropertyEditor/Private/DetailTreeNode.h
+++ b/Engine/Source/Editor/PropertyEditor/Private/DetailTreeNode.h
@@ -8,16 +8,19 @@
 #include "PropertyPath.h"
 #include "Misc/Attribute.h"
 
+class FComplexPropertyNode;
 class FDetailCategoryImpl;
+class FDetailTreeNode;
+class FDetailWidgetRow;
+class FPropertyNode;
+class IDetailPropertyRow;
 class IDetailsViewPrivate;
 class ITableRow;
 class STableViewBase;
-class FDetailWidgetRow;
-class FPropertyNode;
-class FComplexPropertyNode;
+
 struct FDetailFilter;
 
-typedef TArray<TSharedRef<class FDetailTreeNode>> FDetailNodeList;
+typedef TArray<TSharedRef<FDetailTreeNode>> FDetailNodeList;
 
 enum class ENodeVisibility : uint8
 {
@@ -36,7 +39,7 @@ public:
 	/** IDetailTreeNode interface */
 	virtual FNodeWidgets CreateNodeWidgets() const;
 	virtual void GetChildren(TArray<TSharedRef<IDetailTreeNode>>& OutChildren);
-	virtual TSharedPtr<class IDetailPropertyRow> GetRow() const override { return nullptr; }
+	virtual TSharedPtr<IDetailPropertyRow> GetRow() const override { return nullptr; }
 	virtual void GetFilterStrings(TArray<FString>& OutFilterStrings) const override { };
 	virtual bool GetInitiallyCollapsed() const override { return false; }
 
