FDetailItemNode::GetChildren() now reserves space for its children immediately.

#rnx

diff --git a/Engine/Source/Editor/PropertyEditor/Private/DetailItemNode.cpp b/Engine/Source/Editor/PropertyEditor/Private/DetailItemNode.cpp
--- a/Engine/Source/Editor/PropertyEditor/Private/DetailItemNode.cpp
+++ b/Engine/Source/Editor/PropertyEditor/Private/DetailItemNode.cpp
@@ -360,10 +360,10 @@ bool FDetailItemNode::GenerateStandaloneWidget(FDetailWidgetRow& OutRow) const
 
 void FDetailItemNode::GetChildren(FDetailNodeList& OutChildren)
 {
-	for( int32 ChildIndex = 0; ChildIndex < Children.Num(); ++ChildIndex )
-	{
-		TSharedRef<FDetailTreeNode>& Child = Children[ChildIndex];
+	OutChildren.Reserve(Children.Num());
 
+	for (const TSharedRef<FDetailTreeNode>& Child : Children)
+	{
 		ENodeVisibility ChildVisibility = Child->GetVisibility();
 
 		// Report the child if the child is visible or we are visible due to filtering and there were no filtered children.  
@@ -430,25 +430,27 @@ void FDetailItemNode::GenerateChildren( bool bUpdateFilteredNodes )
 
 	// Discard generated nodes that don't pass the property allow list, as well as generated categories who no longer contain any children
 	// Searching backwards guarantees that a category's children will be culled before the category itself.
-	for (int32 i = Children.Num() - 1; i >= 0; --i)
+	for (int32 Index = Children.Num() - 1; Index >= 0; --Index)
 	{
-		Children[i]->SetParentNode(AsShared());
-		if (Children[i]->GetNodeType() == EDetailNodeType::Object || Children[i]->GetNodeType() == EDetailNodeType::Item)
+		const TSharedRef<FDetailTreeNode>& Child = Children[Index];
+
+		Child->SetParentNode(AsShared());
+		if (Child->GetNodeType() == EDetailNodeType::Object || Child->GetNodeType() == EDetailNodeType::Item)
 		{
-			if (!FPropertyEditorPermissionList::Get().DoesPropertyPassFilter(Children[i]->GetParentBaseStructure(), Children[i]->GetNodeName()))
+			if (!FPropertyEditorPermissionList::Get().DoesPropertyPassFilter(Child->GetParentBaseStructure(), Child->GetNodeName()))
 			{
-				Children.RemoveAt(i);
+				Children.RemoveAt(Index);
 			}
 		}
-		else if (Children[i]->GetNodeType() == EDetailNodeType::Category)
+		else if (Child->GetNodeType() == EDetailNodeType::Category)
 		{
 			// Nodes default to hidden until the filter runs the first time - categories return no children if they're hidden, so force an empty filter to initialize properly
-			Children[i]->FilterNode(FDetailFilter());
+			Child->FilterNode(FDetailFilter());
 			FDetailNodeList Subchildren;
-			Children[i]->GetChildren(Subchildren);
+			Child->GetChildren(Subchildren);
 			if (Subchildren.Num() == 0)
 			{
-				Children.RemoveAt(i);
+				Children.RemoveAt(Index);
 			}
 		}
 	}
