PR #9121: Let custom row take precedence over default widget for key node (Contributed by [REDACTED_PERSON])

diff --git a/Engine/Source/Editor/PropertyEditor/Private/DetailPropertyRow.cpp b/Engine/Source/Editor/PropertyEditor/Private/DetailPropertyRow.cpp
--- a/Engine/Source/Editor/PropertyEditor/Private/DetailPropertyRow.cpp
+++ b/Engine/Source/Editor/PropertyEditor/Private/DetailPropertyRow.cpp
@@ -781,6 +781,15 @@ void FDetailPropertyRow::MakeNameOrKeyWidget( FDetailWidgetRow& Row, const TShar
 
 			NameWidget = CustomTypeWidget->ValueWidget.Widget;
 		}
+		else if (InCustomRow.IsValid())
+		{
+			NameWidget = 
+				SNew( SBox )
+				.IsEnabled( IsEnabledAttrib )
+				[
+					InCustomRow->NameWidget.Widget
+				];
+		}
 		else
 		{
 			NameWidget =
