Removed unnecessary null check when iterating over properties in FStructurePropertyNode.

#rnx

diff --git a/Engine/Source/Editor/PropertyEditor/Private/StructurePropertyNode.cpp b/Engine/Source/Editor/PropertyEditor/Private/StructurePropertyNode.cpp
--- a/Engine/Source/Editor/PropertyEditor/Private/StructurePropertyNode.cpp
+++ b/Engine/Source/Editor/PropertyEditor/Private/StructurePropertyNode.cpp
@@ -10,20 +10,16 @@ void FStructurePropertyNode::InitChildNodes()
 	const bool bShouldShowHiddenProperties = !!HasNodeFlags(EPropertyNodeFlags::ShouldShowHiddenProperties);
 	const bool bShouldShowDisableEditOnInstance = !!HasNodeFlags(EPropertyNodeFlags::ShouldShowDisableEditOnInstance);
 
-	const UStruct* Struct = StructData.IsValid() ? StructData->GetStruct() : NULL;
+	const UStruct* Struct = StructData.IsValid() ? StructData->GetStruct() : nullptr;
 
 	TArray<FProperty*> StructMembers;
 
 	for (TFieldIterator<FProperty> It(Struct); It; ++It)
 	{
 		FProperty* StructMember = *It;
-
-		if (StructMember)
+		if (PropertyEditorHelpers::ShouldBeVisible(*this, StructMember))
 		{
-			if (PropertyEditorHelpers::ShouldBeVisible(*this, StructMember))
-			{
-				StructMembers.Add(StructMember);
-			}
+			StructMembers.Add(StructMember);
 		}
 	}
 
