Fixed property names with spaces (from BPs) not working as an EditCondition

This adds support for parsing quoted (and escaped) property name tokens, which allows BP properties to be used as an EditCondition as long as their names are quoted (and escaped)

diff --git a/Engine/Source/Editor/PropertyEditor/Private/EditConditionParser.cpp b/Engine/Source/Editor/PropertyEditor/Private/EditConditionParser.cpp
--- a/Engine/Source/Editor/PropertyEditor/Private/EditConditionParser.cpp
+++ b/Engine/Source/Editor/PropertyEditor/Private/EditConditionParser.cpp
@@ -61,29 +61,84 @@ static TOptional<FExpressionError> ConsumeNullPtr(FExpressionTokenConsumer& Cons
 
 static TOptional<FExpressionError> ConsumePropertyName(FExpressionTokenConsumer& Consumer)
 {
+	enum class EParsedStringType : uint8
+	{
+		Unknown,
+		Unquoted,
+		Quoted,
+	};
+
 	FString PropertyName;
 	bool bShouldBeEnum = false;
-			
-	TOptional<FStringToken> StringToken = Consumer.GetStream().ParseToken([&PropertyName, &bShouldBeEnum](TCHAR InC)
+	EParsedStringType ParsedStringType = EParsedStringType::Unknown;
+
+	TCHAR OpeningQuoteChar = TEXT('\0');
+	int32 NumConsecutiveSlashes = 0;
+
+	TOptional<FStringToken> StringToken = Consumer.GetStream().ParseToken([&PropertyName, &bShouldBeEnum, &ParsedStringType, &OpeningQuoteChar, &NumConsecutiveSlashes](TCHAR InC)
 	{
-		for (const TCHAR BreakingChar : PropertyBreakingChars)
+		if (ParsedStringType == EParsedStringType::Unknown)
 		{
-			if (InC == BreakingChar)
+			if (InC == '"' || InC == '\'')
 			{
-				return EParseState::StopBefore;
+				ParsedStringType = EParsedStringType::Quoted;
+
+				OpeningQuoteChar = InC;
+				NumConsecutiveSlashes = 0;
+				return EParseState::Continue;
 			}
+			
+			ParsedStringType = EParsedStringType::Unquoted;
 		}
 
+		check(ParsedStringType != EParsedStringType::Unknown);
+
 		if (InC == ':')
 		{
 			bShouldBeEnum = true;
 		}
 
-		PropertyName.AppendChar(InC);
+		if (ParsedStringType == EParsedStringType::Unquoted)
+		{
+			for (const TCHAR BreakingChar : PropertyBreakingChars)
+			{
+				if (InC == BreakingChar)
+				{
+					return EParseState::StopBefore;
+				}
+			}
+
+			PropertyName.AppendChar(InC);
+		}
+		else
+		{
+			check(ParsedStringType == EParsedStringType::Quoted);
+
+			if (InC == OpeningQuoteChar && NumConsecutiveSlashes % 2 == 0)
+			{
+				return EParseState::StopAfter;
+			}
+
+			PropertyName.AppendChar(InC);
+
+			if (InC == '\\')
+			{
+				NumConsecutiveSlashes++;
+			}
+			else
+			{
+				NumConsecutiveSlashes = 0;
+			}
+		}
 
 		return EParseState::Continue;
 	});
 
+	if (ParsedStringType == EParsedStringType::Quoted)
+	{
+		PropertyName.ReplaceEscapedCharWithCharInline();
+	}
+
 	if (StringToken.IsSet())
 	{
 		if (bShouldBeEnum)
