Fixed includetool issues.

diff --git a/Engine/Source/Editor/DeviceProfileEditor/Private/DetailsPanel/SDeviceProfileDetailsPanel.h b/Engine/Source/Editor/DeviceProfileEditor/Private/DetailsPanel/SDeviceProfileDetailsPanel.h
--- a/Engine/Source/Editor/DeviceProfileEditor/Private/DetailsPanel/SDeviceProfileDetailsPanel.h
+++ b/Engine/Source/Editor/DeviceProfileEditor/Private/DetailsPanel/SDeviceProfileDetailsPanel.h
@@ -2,7 +2,6 @@
 
 #pragma once
 
-#include "SDeviceProfileDetailsPanel.h"
 #include "Templates/SharedPointer.h"
 #include "UObject/WeakObjectPtr.h"
 #include "UObject/WeakObjectPtrTemplates.h"
diff --git a/Engine/Source/Editor/InputBindingEditor/Private/Widgets/SChordEditor.h b/Engine/Source/Editor/InputBindingEditor/Private/Widgets/SChordEditor.h
--- a/Engine/Source/Editor/InputBindingEditor/Private/Widgets/SChordEditor.h
+++ b/Engine/Source/Editor/InputBindingEditor/Private/Widgets/SChordEditor.h
@@ -7,7 +7,6 @@
 #include "Framework/Commands/UICommandInfo.h"
 #include "Input/Reply.h"
 #include "Internationalization/Text.h"
-#include "SChordEditor.h"
 #include "Templates/SharedPointer.h"
 #include "Widgets/DeclarativeSyntaxSupport.h"
 #include "Widgets/Input/SEditableText.h"
diff --git a/Engine/Source/Editor/PropertyEditor/Private/UserInterface/PropertyEditor/SPropertyComboBox.h b/Engine/Source/Editor/PropertyEditor/Private/UserInterface/PropertyEditor/SPropertyComboBox.h
--- a/Engine/Source/Editor/PropertyEditor/Private/UserInterface/PropertyEditor/SPropertyComboBox.h
+++ b/Engine/Source/Editor/PropertyEditor/Private/UserInterface/PropertyEditor/SPropertyComboBox.h
@@ -8,7 +8,6 @@
 #include "Input/Reply.h"
 #include "Layout/Visibility.h"
 #include "Misc/Attribute.h"
-#include "SPropertyComboBox.h"
 #include "SSearchableComboBox.h"
 #include "Styling/AppStyle.h"
 #include "Templates/SharedPointer.h"
diff --git a/Engine/Source/Editor/PropertyEditor/Private/UserInterface/PropertyEditor/SPropertyEditorArrayItem.h b/Engine/Source/Editor/PropertyEditor/Private/UserInterface/PropertyEditor/SPropertyEditorArrayItem.h
--- a/Engine/Source/Editor/PropertyEditor/Private/UserInterface/PropertyEditor/SPropertyEditorArrayItem.h
+++ b/Engine/Source/Editor/PropertyEditor/Private/UserInterface/PropertyEditor/SPropertyEditorArrayItem.h
@@ -7,7 +7,6 @@
 #include "Internationalization/Text.h"
 #include "Misc/Attribute.h"
 #include "PropertyEditorModule.h"
-#include "SPropertyEditorArrayItem.h"
 #include "Styling/AppStyle.h"
 #include "Templates/SharedPointer.h"
 #include "Templates/TypeHash.h"
diff --git a/Engine/Source/Editor/PropertyEditor/Private/UserInterface/PropertyEditor/SPropertyEditorDateTime.h b/Engine/Source/Editor/PropertyEditor/Private/UserInterface/PropertyEditor/SPropertyEditorDateTime.h
--- a/Engine/Source/Editor/PropertyEditor/Private/UserInterface/PropertyEditor/SPropertyEditorDateTime.h
+++ b/Engine/Source/Editor/PropertyEditor/Private/UserInterface/PropertyEditor/SPropertyEditorDateTime.h
@@ -3,7 +3,6 @@
 
 #include "Fonts/SlateFontInfo.h"
 #include "Misc/Attribute.h"
-#include "SPropertyEditorDateTime.h"
 #include "Styling/AppStyle.h"
 #include "Templates/SharedPointer.h"
 #include "Templates/TypeHash.h"
diff --git a/Engine/Source/Editor/PropertyEditor/Private/UserInterface/PropertyEditor/SPropertyEditorInteractiveActorPicker.h b/Engine/Source/Editor/PropertyEditor/Private/UserInterface/PropertyEditor/SPropertyEditorInteractiveActorPicker.h
--- a/Engine/Source/Editor/PropertyEditor/Private/UserInterface/PropertyEditor/SPropertyEditorInteractiveActorPicker.h
+++ b/Engine/Source/Editor/PropertyEditor/Private/UserInterface/PropertyEditor/SPropertyEditorInteractiveActorPicker.h
@@ -3,7 +3,6 @@
 
 #include "ActorPickerMode.h"
 #include "Input/Reply.h"
-#include "SPropertyEditorInteractiveActorPicker.h"
 #include "Widgets/DeclarativeSyntaxSupport.h"
 #include "Widgets/Input/SButton.h"
 
diff --git a/Engine/Source/Editor/PropertyEditor/Private/UserInterface/PropertyEditor/SPropertyEditorMap.h b/Engine/Source/Editor/PropertyEditor/Private/UserInterface/PropertyEditor/SPropertyEditorMap.h
--- a/Engine/Source/Editor/PropertyEditor/Private/UserInterface/PropertyEditor/SPropertyEditorMap.h
+++ b/Engine/Source/Editor/PropertyEditor/Private/UserInterface/PropertyEditor/SPropertyEditorMap.h
@@ -4,7 +4,6 @@
 #include "Fonts/SlateFontInfo.h"
 #include "Internationalization/Text.h"
 #include "Misc/Attribute.h"
-#include "SPropertyEditorMap.h"
 #include "Styling/AppStyle.h"
 #include "Templates/SharedPointer.h"
 #include "Templates/TypeHash.h"
diff --git a/Engine/Source/Editor/PropertyEditor/Private/UserInterface/PropertyEditor/SPropertyEditorSceneDepthPicker.h b/Engine/Source/Editor/PropertyEditor/Private/UserInterface/PropertyEditor/SPropertyEditorSceneDepthPicker.h
--- a/Engine/Source/Editor/PropertyEditor/Private/UserInterface/PropertyEditor/SPropertyEditorSceneDepthPicker.h
+++ b/Engine/Source/Editor/PropertyEditor/Private/UserInterface/PropertyEditor/SPropertyEditorSceneDepthPicker.h
@@ -2,7 +2,6 @@
 #pragma once
 
 #include "Input/Reply.h"
-#include "SPropertyEditorSceneDepthPicker.h"
 #include "SceneDepthPickerMode.h"
 #include "Widgets/DeclarativeSyntaxSupport.h"
 #include "Widgets/Input/SButton.h"
diff --git a/Engine/Source/Editor/PropertyEditor/Private/UserInterface/PropertyEditor/SPropertyEditorSet.h b/Engine/Source/Editor/PropertyEditor/Private/UserInterface/PropertyEditor/SPropertyEditorSet.h
--- a/Engine/Source/Editor/PropertyEditor/Private/UserInterface/PropertyEditor/SPropertyEditorSet.h
+++ b/Engine/Source/Editor/PropertyEditor/Private/UserInterface/PropertyEditor/SPropertyEditorSet.h
@@ -4,7 +4,6 @@
 #include "Fonts/SlateFontInfo.h"
 #include "Internationalization/Text.h"
 #include "Misc/Attribute.h"
-#include "SPropertyEditorSet.h"
 #include "Styling/AppStyle.h"
 #include "Templates/SharedPointer.h"
 #include "Templates/TypeHash.h"
diff --git a/Engine/Source/Editor/PropertyEditor/Private/UserInterface/PropertyEditor/SPropertyEditorText.h b/Engine/Source/Editor/PropertyEditor/Private/UserInterface/PropertyEditor/SPropertyEditorText.h
--- a/Engine/Source/Editor/PropertyEditor/Private/UserInterface/PropertyEditor/SPropertyEditorText.h
+++ b/Engine/Source/Editor/PropertyEditor/Private/UserInterface/PropertyEditor/SPropertyEditorText.h
@@ -4,7 +4,6 @@
 #include "Fonts/SlateFontInfo.h"
 #include "Input/Reply.h"
 #include "Misc/Attribute.h"
-#include "SPropertyEditorText.h"
 #include "Styling/AppStyle.h"
 #include "Templates/SharedPointer.h"
 #include "Templates/TypeHash.h"
diff --git a/Engine/Source/Editor/PropertyEditor/Private/UserInterface/PropertyTable/SPropertyTableCell.h b/Engine/Source/Editor/PropertyEditor/Private/UserInterface/PropertyTable/SPropertyTableCell.h
--- a/Engine/Source/Editor/PropertyEditor/Private/UserInterface/PropertyTable/SPropertyTableCell.h
+++ b/Engine/Source/Editor/PropertyEditor/Private/UserInterface/PropertyTable/SPropertyTableCell.h
@@ -3,7 +3,6 @@
 
 #include "HAL/PlatformMath.h"
 #include "Input/Reply.h"
-#include "SPropertyTableCell.h"
 #include "Templates/SharedPointer.h"
 #include "Types/SlateEnums.h"
 #include "UObject/NameTypes.h"
diff --git a/Engine/Source/Editor/UMGEditor/Private/Designer/SDesignerToolBar.h b/Engine/Source/Editor/UMGEditor/Private/Designer/SDesignerToolBar.h
--- a/Engine/Source/Editor/UMGEditor/Private/Designer/SDesignerToolBar.h
+++ b/Engine/Source/Editor/UMGEditor/Private/Designer/SDesignerToolBar.h
@@ -6,7 +6,6 @@
 #include "Containers/UnrealString.h"
 #include "HAL/Platform.h"
 #include "Internationalization/Text.h"
-#include "SDesignerToolBar.h"
 #include "SViewportToolBar.h"
 #include "Templates/SharedPointer.h"
 #include "UObject/NameTypes.h"
diff --git a/Engine/Source/Editor/UMGEditor/Private/Designer/SDisappearingBar.h b/Engine/Source/Editor/UMGEditor/Private/Designer/SDisappearingBar.h
--- a/Engine/Source/Editor/UMGEditor/Private/Designer/SDisappearingBar.h
+++ b/Engine/Source/Editor/UMGEditor/Private/Designer/SDisappearingBar.h
@@ -6,7 +6,6 @@
 #include "Layout/Visibility.h"
 #include "Math/Color.h"
 #include "Misc/Attribute.h"
-#include "SDisappearingBar.h"
 #include "Templates/SharedPointer.h"
 #include "Widgets/DeclarativeSyntaxSupport.h"
 #include "Widgets/SCompoundWidget.h"
diff --git a/Engine/Source/Editor/UMGEditor/Private/Designer/SRuler.h b/Engine/Source/Editor/UMGEditor/Private/Designer/SRuler.h
--- a/Engine/Source/Editor/UMGEditor/Private/Designer/SRuler.h
+++ b/Engine/Source/Editor/UMGEditor/Private/Designer/SRuler.h
@@ -7,7 +7,6 @@
 #include "Layout/Clipping.h"
 #include "Math/Vector2D.h"
 #include "Misc/Optional.h"
-#include "SRuler.h"
 #include "Types/SlateEnums.h"
 #include "Types/WidgetMouseEventsDelegate.h"
 #include "Widgets/DeclarativeSyntaxSupport.h"
diff --git a/Engine/Source/Editor/UMGEditor/Private/Designer/SZoomPan.h b/Engine/Source/Editor/UMGEditor/Private/Designer/SZoomPan.h
--- a/Engine/Source/Editor/UMGEditor/Private/Designer/SZoomPan.h
+++ b/Engine/Source/Editor/UMGEditor/Private/Designer/SZoomPan.h
@@ -5,7 +5,6 @@
 #include "HAL/Platform.h"
 #include "Math/Vector2D.h"
 #include "Misc/Attribute.h"
-#include "SZoomPan.h"
 #include "Templates/SharedPointer.h"
 #include "Widgets/DeclarativeSyntaxSupport.h"
 #include "Widgets/SCompoundWidget.h"
diff --git a/Engine/Source/Editor/WorldPartitionEditor/Private/WorldPartition/HLOD/SWorldPartitionBuildHLODsDialog.h b/Engine/Source/Editor/WorldPartitionEditor/Private/WorldPartition/HLOD/SWorldPartitionBuildHLODsDialog.h
--- a/Engine/Source/Editor/WorldPartitionEditor/Private/WorldPartition/HLOD/SWorldPartitionBuildHLODsDialog.h
+++ b/Engine/Source/Editor/WorldPartitionEditor/Private/WorldPartition/HLOD/SWorldPartitionBuildHLODsDialog.h
@@ -7,7 +7,6 @@
 #include "Templates/SharedPointer.h"
 #include "Widgets/DeclarativeSyntaxSupport.h"
 #include "Widgets/SCompoundWidget.h"
-#include "WorldPartition/HLOD/SWorldPartitionBuildHLODsDialog.h"
 
 class SWindow;
 
