FEditConditionContext::GetIntegerValue can fall back on the underlying property if it is an enum

diff --git a/Engine/Source/Editor/PropertyEditor/Private/EditConditionContext.cpp b/Engine/Source/Editor/PropertyEditor/Private/EditConditionContext.cpp
--- a/Engine/Source/Editor/PropertyEditor/Private/EditConditionContext.cpp
+++ b/Engine/Source/Editor/PropertyEditor/Private/EditConditionContext.cpp
@@ -186,7 +186,18 @@ TOptional<int64> FEditConditionContext::GetIntegerValue(const FString& PropertyN
 		return TOptional<int64>();
 	}
 
-	const FNumericProperty* NumericProperty = FindTypedField<FNumericProperty>(PinnedNode, PropertyName);
+	const FProperty* Property = FindTypedField<FProperty>(PinnedNode, PropertyName);
+	const FNumericProperty* NumericProperty = CastField<FNumericProperty>(Property);
+
+	if (NumericProperty == nullptr)
+	{
+		// Retry with an enum and its underlying property
+		if (const FEnumProperty* EnumProperty = CastField<FEnumProperty>(Property))
+		{
+			NumericProperty = EnumProperty->GetUnderlyingProperty();
+		}
+	}
+
 	if (NumericProperty == nullptr || !NumericProperty->IsInteger())
 	{
 		return TOptional<int64>();
@@ -207,7 +218,7 @@ TOptional<int64> FEditConditionContext::GetIntegerValue(const FString& PropertyN
 	TOptional<int64> Result;
 	for (int32 Index = 0; Index < ComplexParentNode->GetInstancesNum(); ++Index)
 	{
-		const uint8* ValuePtr = GetPropertyValuePtr(NumericProperty, PinnedNode, ParentNode, ComplexParentNode, Index);
+		const uint8* ValuePtr = GetPropertyValuePtr(Property, PinnedNode, ParentNode, ComplexParentNode, Index);
 		if (ValuePtr == nullptr)
 		{
 			return TOptional<int64>();
