PR #9234: Fix property display for ReadOnlyKeys (Contributed by [REDACTED_PERSON])

diff --git a/Engine/Source/Editor/PropertyEditor/Private/DetailPropertyRow.cpp b/Engine/Source/Editor/PropertyEditor/Private/DetailPropertyRow.cpp
--- a/Engine/Source/Editor/PropertyEditor/Private/DetailPropertyRow.cpp
+++ b/Engine/Source/Editor/PropertyEditor/Private/DetailPropertyRow.cpp
@@ -754,7 +754,7 @@ void FDetailPropertyRow::MakeNameOrKeyWidget( FDetailWidgetRow& Row, const TShar
 	EHorizontalAlignment HorizontalAlignment = HAlign_Fill;
 
 	// We will only use key widgets for non-struct keys
-	const bool bHasKeyNode = PropertyKeyEditor.IsValid() && !PropertyHandle->HasMetaData(TEXT("ReadOnlyKeys"));
+	const bool bHasKeyNode = PropertyKeyEditor.IsValid();
 
 	if (!bHasKeyNode && InCustomRow.IsValid())
 	{
@@ -772,6 +772,11 @@ void FDetailPropertyRow::MakeNameOrKeyWidget( FDetailWidgetRow& Row, const TShar
 	// Key nodes take precedence over custom rows
 	if (bHasKeyNode)
 	{
+		if (PropertyHandle->HasMetaData(TEXT("ReadOnlyKeys")))
+		{
+			PropertyKeyEditor->GetPropertyNode()->SetNodeFlags(EPropertyNodeFlags::IsReadOnly, true);
+		}
+
 		// Does this key have a custom type, use it
 		if (CachedKeyCustomTypeInterface)
 		{
