FNames containing asset paths are deprecated. FSoftObjectPath or FTopLevelAssetPath should be used instead.
Fixups for property editors.

diff --git a/Engine/Source/Editor/EditorWidgets/Private/STextPropertyEditableTextBox.cpp b/Engine/Source/Editor/EditorWidgets/Private/STextPropertyEditableTextBox.cpp
--- a/Engine/Source/Editor/EditorWidgets/Private/STextPropertyEditableTextBox.cpp
+++ b/Engine/Source/Editor/EditorWidgets/Private/STextPropertyEditableTextBox.cpp
@@ -379,9 +379,10 @@ void STextPropertyEditableStringTableReference::UpdateStringTableComboOptions()
 
 		for (const FAssetData& StringTableAsset : StringTableAssets)
 		{
+			FName StringTableId = *StringTableAsset.GetObjectPathString();
 			// Only allow string tables assets that have entries to be visible otherwise unexpected behavior happens for the user
 			bool HasEntries = false;
-			FStringTableConstPtr StringTable = FStringTableRegistry::Get().FindStringTable(StringTableAsset.ObjectPath);
+			FStringTableConstPtr StringTable = FStringTableRegistry::Get().FindStringTable(StringTableId);
 			if (StringTable.IsValid())
 			{
 				StringTable->EnumerateSourceStrings([&](const FString& InKey, const FString& InSourceString) -> bool
@@ -402,9 +403,9 @@ void STextPropertyEditableStringTableReference::UpdateStringTableComboOptions()
 			}
 
 			TSharedRef<FAvailableStringTable> AvailableStringTableEntry = MakeShared<FAvailableStringTable>();
-			AvailableStringTableEntry->TableId = StringTableAsset.ObjectPath;
+			AvailableStringTableEntry->TableId = StringTableId;
 			AvailableStringTableEntry->DisplayName = FText::FromName(StringTableAsset.AssetName);
-			if (StringTableAsset.ObjectPath == CurrentTableId)
+			if (StringTableId == CurrentTableId)
 			{
 				SelectedStringTableComboEntry = AvailableStringTableEntry;
 			}
diff --git a/Engine/Source/Editor/PropertyEditor/Private/PropertyHandleImpl.cpp b/Engine/Source/Editor/PropertyEditor/Private/PropertyHandleImpl.cpp
--- a/Engine/Source/Editor/PropertyEditor/Private/PropertyHandleImpl.cpp
+++ b/Engine/Source/Editor/PropertyEditor/Private/PropertyHandleImpl.cpp
@@ -3862,7 +3862,7 @@ FPropertyAccess::Result FPropertyHandleObject::SetValue(const FAssetData& NewVal
 			NewValue.GetAsset();
 		}
 
-		FString ObjectPathName = NewValue.IsValid() ? FString::Printf(TEXT("%s'%s'"), *NewValue.GetClass()->GetName(), *NewValue.ObjectPath.ToString()) : TEXT("None");
+		FString ObjectPathName = NewValue.IsValid() ? FString::Printf(TEXT("%s'%s'"), *NewValue.GetClass()->GetName(), *NewValue.GetObjectPathString()) : TEXT("None");
 		return SetValueFromFormattedString(ObjectPathName, Flags, bSkipResolve);
 	}
 
diff --git a/Engine/Source/Editor/PropertyEditor/Private/UserInterface/PropertyEditor/SPropertyEditorArray.h b/Engine/Source/Editor/PropertyEditor/Private/UserInterface/PropertyEditor/SPropertyEditorArray.h
--- a/Engine/Source/Editor/PropertyEditor/Private/UserInterface/PropertyEditor/SPropertyEditorArray.h
+++ b/Engine/Source/Editor/PropertyEditor/Private/UserInterface/PropertyEditor/SPropertyEditorArray.h
@@ -112,7 +112,7 @@ private:
 					// if the type matches
 					if (AssetData.IsInstanceOf(ObjectProperty->PropertyClass))
 					{
-						PropertyEditor->AddGivenItem(AssetData.ObjectPath.ToString());
+						PropertyEditor->AddGivenItem(AssetData.GetObjectPathString());
 					}
 				}
 				// Let this bubble up to the rest of the row
diff --git a/Engine/Source/Editor/PropertyEditor/Private/UserInterface/PropertyEditor/SPropertyEditorAsset.cpp b/Engine/Source/Editor/PropertyEditor/Private/UserInterface/PropertyEditor/SPropertyEditorAsset.cpp
--- a/Engine/Source/Editor/PropertyEditor/Private/UserInterface/PropertyEditor/SPropertyEditorAsset.cpp
+++ b/Engine/Source/Editor/PropertyEditor/Private/UserInterface/PropertyEditor/SPropertyEditorAsset.cpp
@@ -701,7 +701,7 @@ SPropertyEditorAsset::EActorReferenceState SPropertyEditorAsset::GetActorReferen
 		else
 		{
 			// Get a path pointing to the owning map
-			FSoftObjectPath MapObjectPath = FSoftObjectPath(Value.ObjectPath.GetAssetPathName(), FString());
+			FSoftObjectPath MapObjectPath = Value.ObjectPath.GetWithoutSubPath();
 
 			if (UObject* MapObject = MapObjectPath.ResolveObject())
 			{
@@ -987,12 +987,12 @@ FPropertyAccess::Result SPropertyEditorAsset::GetValue( FObjectOrAssetData& OutV
 
 				if (SoftObjectPath.IsAsset())
 				{
-					if (!CachedAssetData.IsValid() || CachedAssetData.ObjectPath.ToString() != CurrentObjectPath)
+					if (!CachedAssetData.IsValid() || CachedAssetData.GetObjectPathString() != CurrentObjectPath)
 					{
 						static FName AssetRegistryName("AssetRegistry");
 
 						FAssetRegistryModule& AssetRegistryModule = FModuleManager::Get().LoadModuleChecked<FAssetRegistryModule>(AssetRegistryName);
-						CachedAssetData = AssetRegistryModule.Get().GetAssetByObjectPath(*CurrentObjectPath);
+						CachedAssetData = AssetRegistryModule.Get().GetAssetByObjectPath(FSoftObjectPath(CurrentObjectPath));
 					}
 
 					Result = FPropertyAccess::Success;
@@ -1064,13 +1064,13 @@ FPropertyAccess::Result SPropertyEditorAsset::GetValue( FObjectOrAssetData& OutV
 
 			if (SoftObjectPath.IsAsset())
 			{
-				const FString CurrentObjectPath = SoftObjectPath.ToString();
-				if (CurrentObjectPath != TEXT("None") && (!CachedAssetData.IsValid() || CachedAssetData.ObjectPath.ToString() != CurrentObjectPath))
+				const FSoftObjectPath CurrentObjectPath = SoftObjectPath;
+				if (CurrentObjectPath.IsValid() && (!CachedAssetData.IsValid() || CachedAssetData.GetSoftObjectPath() != CurrentObjectPath))
 				{
 					static FName AssetRegistryName("AssetRegistry");
 
 					FAssetRegistryModule& AssetRegistryModule = FModuleManager::Get().LoadModuleChecked<FAssetRegistryModule>(AssetRegistryName);
-					CachedAssetData = AssetRegistryModule.Get().GetAssetByObjectPath(*CurrentObjectPath);
+					CachedAssetData = AssetRegistryModule.Get().GetAssetByObjectPath(CurrentObjectPath);
 				}
 
 				OutValue = FObjectOrAssetData(CachedAssetData);
@@ -1175,7 +1175,7 @@ void SPropertyEditorAsset::OnBrowse()
 			// Try to resolve a potentially unloaded object
 			if (!Value.Object)
 			{
-				FSoftObjectPath MapObjectPath = FSoftObjectPath(Value.ObjectPath.GetAssetPathName(), FString());
+				FSoftObjectPath MapObjectPath = Value.ObjectPath.GetWithoutSubPath();
 
 				if (UObject* MapObject = MapObjectPath.ResolveObject())
 				{
@@ -1383,7 +1383,7 @@ bool SPropertyEditorAsset::CanPaste()
 		else
 		{
 			FAssetRegistryModule& AssetRegistryModule = FModuleManager::Get().LoadModuleChecked<FAssetRegistryModule>(TEXT("AssetRegistry"));
-			bCanPaste = PossibleObjectPath.Len() < NAME_SIZE && AssetRegistryModule.Get().GetAssetByObjectPath( *PossibleObjectPath ).IsValid();
+			bCanPaste = PossibleObjectPath.Len() < NAME_SIZE && AssetRegistryModule.Get().GetAssetByObjectPath( FSoftObjectPath(PossibleObjectPath) ).IsValid();
 		}
 	}
 
diff --git a/Engine/Source/Editor/PropertyEditor/Private/UserInterface/PropertyEditor/SPropertyEditorClass.cpp b/Engine/Source/Editor/PropertyEditor/Private/UserInterface/PropertyEditor/SPropertyEditorClass.cpp
--- a/Engine/Source/Editor/PropertyEditor/Private/UserInterface/PropertyEditor/SPropertyEditorClass.cpp
+++ b/Engine/Source/Editor/PropertyEditor/Private/UserInterface/PropertyEditor/SPropertyEditorClass.cpp
@@ -380,7 +380,7 @@ static UObject* LoadDragDropObject(TSharedPtr<FAssetDragDropOp> UnloadedClassOp)
 	// Find the class/blueprint path
 	if (UnloadedClassOp->HasAssets())
 	{
-		AssetPath = UnloadedClassOp->GetAssets()[0].ObjectPath.ToString();
+		AssetPath = UnloadedClassOp->GetAssets()[0].GetObjectPathString();
 	}
 	else if (UnloadedClassOp->HasAssetPaths())
 	{
@@ -467,7 +467,7 @@ FReply SPropertyEditorClass::OnDrop(const FGeometry& MyGeometry, const FDragDrop
 		// Find the class/blueprint path
 		if (UnloadedClassOp->HasAssets())
 		{
-			AssetPath = UnloadedClassOp->GetAssets()[0].ObjectPath.ToString();
+			AssetPath = UnloadedClassOp->GetAssets()[0].GetObjectPathString();
 		}
 		else if (UnloadedClassOp->HasAssetPaths())
 		{
diff --git a/Engine/Source/Editor/PropertyEditor/Private/UserInterface/PropertyEditor/SPropertyEditorStruct.cpp b/Engine/Source/Editor/PropertyEditor/Private/UserInterface/PropertyEditor/SPropertyEditorStruct.cpp
--- a/Engine/Source/Editor/PropertyEditor/Private/UserInterface/PropertyEditor/SPropertyEditorStruct.cpp
+++ b/Engine/Source/Editor/PropertyEditor/Private/UserInterface/PropertyEditor/SPropertyEditorStruct.cpp
@@ -246,7 +246,7 @@ void SPropertyEditorStruct::OnDragEnter(const FGeometry& MyGeometry, const FDrag
 		// Find the struct path
 		if (UnloadedStructOp->HasAssets())
 		{
-			AssetPath = UnloadedStructOp->GetAssets()[0].ObjectPath.ToString();
+			AssetPath = UnloadedStructOp->GetAssets()[0].GetObjectPathString();
 		}
 		else if (UnloadedStructOp->HasAssetPaths())
 		{
@@ -296,7 +296,7 @@ FReply SPropertyEditorStruct::OnDrop(const FGeometry& MyGeometry, const FDragDro
 		// Find the struct path
 		if (UnloadedStructOp->HasAssets())
 		{
-			AssetPath = UnloadedStructOp->GetAssets()[0].ObjectPath.ToString();
+			AssetPath = UnloadedStructOp->GetAssets()[0].GetObjectPathString();
 		}
 		else if (UnloadedStructOp->HasAssetPaths())
 		{
diff --git a/Engine/Source/Editor/PropertyEditor/Private/UserInterface/PropertyEditor/SPropertyMenuActorPicker.cpp b/Engine/Source/Editor/PropertyEditor/Private/UserInterface/PropertyEditor/SPropertyMenuActorPicker.cpp
--- a/Engine/Source/Editor/PropertyEditor/Private/UserInterface/PropertyEditor/SPropertyMenuActorPicker.cpp
+++ b/Engine/Source/Editor/PropertyEditor/Private/UserInterface/PropertyEditor/SPropertyMenuActorPicker.cpp
@@ -175,7 +175,7 @@ bool SPropertyMenuActorPicker::CanPaste()
 	else
 	{
 		FAssetRegistryModule& AssetRegistryModule = FModuleManager::Get().LoadModuleChecked<FAssetRegistryModule>(TEXT("AssetRegistry"));
-		bCanPaste = PossibleObjectPath.Len() < NAME_SIZE && AssetRegistryModule.Get().GetAssetByObjectPath( *PossibleObjectPath ).IsValid();
+		bCanPaste = PossibleObjectPath.Len() < NAME_SIZE && AssetRegistryModule.Get().GetAssetByObjectPath( FSoftObjectPath(PossibleObjectPath) ).IsValid();
 	}
 
 	return bCanPaste;
diff --git a/Engine/Source/Editor/PropertyEditor/Private/UserInterface/PropertyEditor/SPropertyMenuAssetPicker.cpp b/Engine/Source/Editor/PropertyEditor/Private/UserInterface/PropertyEditor/SPropertyMenuAssetPicker.cpp
--- a/Engine/Source/Editor/PropertyEditor/Private/UserInterface/PropertyEditor/SPropertyMenuAssetPicker.cpp
+++ b/Engine/Source/Editor/PropertyEditor/Private/UserInterface/PropertyEditor/SPropertyMenuAssetPicker.cpp
@@ -295,7 +295,7 @@ bool SPropertyMenuAssetPicker::CanPaste()
 	else
 	{
 		FAssetRegistryModule& AssetRegistryModule = FModuleManager::Get().LoadModuleChecked<FAssetRegistryModule>(TEXT("AssetRegistry"));
-		bCanPaste = PossibleObjectPath.Len() < NAME_SIZE && AssetRegistryModule.Get().GetAssetByObjectPath( *PossibleObjectPath ).IsValid();
+		bCanPaste = PossibleObjectPath.Len() < NAME_SIZE && AssetRegistryModule.Get().GetAssetByObjectPath( FSoftObjectPath(PossibleObjectPath) ).IsValid();
 	}
 
 	return bCanPaste;
