Fixed purging of objects only working if an object was actually pointing to something but is no longer valid. This fixes a use case in Niagara where an external object is added to the details panel, which is "None" by default. This would then be considered a "purged object" and force a reconstruction every frame.

diff --git a/Engine/Source/Editor/PropertyEditor/Private/ObjectPropertyNode.cpp b/Engine/Source/Editor/PropertyEditor/Private/ObjectPropertyNode.cpp
--- a/Engine/Source/Editor/PropertyEditor/Private/ObjectPropertyNode.cpp
+++ b/Engine/Source/Editor/PropertyEditor/Private/ObjectPropertyNode.cpp
@@ -96,18 +96,13 @@ bool FObjectPropertyNode::PurgeKilledObjects()
 	bool bDidPurgeObjects = false;
 
 	// Purge any objects that are marked pending kill from the object list
-	for (int32 Index = 0; Index < Objects.Num(); )
+	for (auto It = Objects.CreateIterator(); It; ++It)
 	{
-		TWeakObjectPtr<UObject> Object = Objects[Index];
-
-		if ( !Object.IsValid() )
+		const TWeakObjectPtr<UObject>& Object = *It;
+		if (Object.IsStale())
 		{
-			Objects.RemoveAt(Index, 1);
 			bDidPurgeObjects = true;
-		}
-		else
-		{
-			++Index;
+			It.RemoveCurrent();
 		}
 	}
 
