Reverted PR #9121, which was causing key properties with customized value editors (eg. vectors, transforms) to display incorrectly.

#rnx

diff --git a/Engine/Source/Editor/PropertyEditor/Private/DetailPropertyRow.cpp b/Engine/Source/Editor/PropertyEditor/Private/DetailPropertyRow.cpp
--- a/Engine/Source/Editor/PropertyEditor/Private/DetailPropertyRow.cpp
+++ b/Engine/Source/Editor/PropertyEditor/Private/DetailPropertyRow.cpp
@@ -786,15 +786,6 @@ void FDetailPropertyRow::MakeNameOrKeyWidget( FDetailWidgetRow& Row, const TShar
 
 			NameWidget = CustomTypeWidget->ValueWidget.Widget;
 		}
-		else if (InCustomRow.IsValid())
-		{
-			NameWidget = 
-				SNew( SBox )
-				.IsEnabled( IsEnabledAttrib )
-				[
-					InCustomRow->NameWidget.Widget
-				];
-		}
 		else
 		{
 			NameWidget =
