Show Hidden Properties While Playing setting is now respected in the details view.

Previously, the checkbox value wasn't even listened to during the same startup, but rather the config value at startup was used, which made it very deceptive and sticky. This feature probably needs some thought in the future as well, because it can stall incredibly long.

diff --git a/Engine/Source/Editor/PropertyEditor/Private/SDetailsView.cpp b/Engine/Source/Editor/PropertyEditor/Private/SDetailsView.cpp
--- a/Engine/Source/Editor/PropertyEditor/Private/SDetailsView.cpp
+++ b/Engine/Source/Editor/PropertyEditor/Private/SDetailsView.cpp
@@ -917,7 +917,8 @@ void SDetailsView::PostSetObject(const TArray<FDetailsViewObjectRoot>& Roots)
 
 	// Are we editing PIE objects?  If the bShowHiddenPropertiesWhilePlaying setting is enabled, we may want to
 	// show all of the properties that would normally be hidden for objects that are part of the PIE world.
-	bool bAnyPIEObjects = false;
+	bool bShowHiddenPropertiesWhilePlaying = false;
+	if (IsShowHiddenPropertiesWhilePlayingChecked())
 	{
 		for( int32 RootNodeIndex = 0; RootNodeIndex < RootPropertyNodes.Num(); ++RootNodeIndex )
 		{
@@ -930,7 +931,7 @@ void SDetailsView::PostSetObject(const TArray<FDetailsViewObjectRoot>& Roots)
 					UObject* Object = RootPropertyNode->GetUObject( ObjectIndex );
 					if( Object->GetOutermost()->HasAnyPackageFlags( PKG_PlayInEditor ) )
 					{
-						bAnyPIEObjects = true;
+						bShowHiddenPropertiesWhilePlaying = true;
 						break;
 					}
 				}
@@ -946,7 +947,7 @@ void SDetailsView::PostSetObject(const TArray<FDetailsViewObjectRoot>& Roots)
 	InitParams.bAllowChildren = true;
 	InitParams.bForceHiddenPropertyVisibility = 
 		FPropertySettings::Get().ShowHiddenProperties() || 
-		( GetDefault<UEditorStyleSettings>()->bShowHiddenPropertiesWhilePlaying && bAnyPIEObjects ) ||
+		bShowHiddenPropertiesWhilePlaying ||
 		DetailsViewArgs.bForceHiddenPropertyVisibility;
 
 	switch ( DetailsViewArgs.DefaultsOnlyVisibility )
@@ -1020,7 +1021,7 @@ bool SDetailsView::IsShowHiddenPropertiesWhilePlayingChecked() const
 		return ViewConfig->bShowHiddenPropertiesWhilePlaying;
 	}
 
-	return false;
+	return GetDefault<UEditorStyleSettings>()->bShowHiddenPropertiesWhilePlaying;
 }
 
 void SDetailsView::OnShowHiddenPropertiesWhilePlayingClicked()
