Added a list of types to ignore when recursively expanding details view nodes.

Requested on UDN: https://epicgames.lightning.force.com/lightning/r/Case/5004z00001ecpUKAAY/view

diff --git a/Engine/Source/Editor/PropertyEditor/Private/SDetailsViewBase.cpp b/Engine/Source/Editor/PropertyEditor/Private/SDetailsViewBase.cpp
--- a/Engine/Source/Editor/PropertyEditor/Private/SDetailsViewBase.cpp
+++ b/Engine/Source/Editor/PropertyEditor/Private/SDetailsViewBase.cpp
@@ -113,24 +113,38 @@ void SDetailsViewBase::SetRootExpansionStates(const bool bExpand, const bool bRe
 	}
 }
 
-void SDetailsViewBase::SetNodeExpansionState(TSharedRef<FDetailTreeNode> InTreeNode, bool bIsItemExpanded, bool bRecursive)
+void SDetailsViewBase::SetNodeExpansionState(TSharedRef<FDetailTreeNode> InTreeNode, bool bExpand, bool bRecursive)
 {
 	TArray< TSharedRef<FDetailTreeNode> > Children;
 	InTreeNode->GetChildren(Children);
 
 	if (Children.Num())
 	{
-		RequestItemExpanded(InTreeNode, bIsItemExpanded);
+		RequestItemExpanded(InTreeNode, bExpand);
 		const bool bShouldSaveState = true;
-		InTreeNode->OnItemExpansionChanged(bIsItemExpanded, bShouldSaveState);
+		InTreeNode->OnItemExpansionChanged(bExpand, bShouldSaveState);
 
 		if (bRecursive)
 		{
-			for (int32 ChildIndex = 0; ChildIndex < Children.Num(); ++ChildIndex)
+			TArray<FString> TypesToIgnoreForRecursiveExpansion;
+			GConfig->GetArray(TEXT("DetailPropertyExpansion"), TEXT("TypesToIgnoreForRecursiveExpansion"), TypesToIgnoreForRecursiveExpansion, GEditorPerProjectIni);
+
+			for (const TSharedRef<FDetailTreeNode>& Child : Children)
+			{
+				if (bExpand)
+				{
+					const FProperty* ChildProperty = Child->GetPropertyNode().IsValid() ? Child->GetPropertyNode()->GetProperty() : nullptr;
+					if (ChildProperty != nullptr)
 					{
-				TSharedRef<FDetailTreeNode> Child = Children[ChildIndex];
+						if (TypesToIgnoreForRecursiveExpansion.Contains(ChildProperty->GetClass()->GetName()))
+						{
+							continue;
+						}
+					}
+				}
+
 
-				SetNodeExpansionState(Child, bIsItemExpanded, bRecursive);
+				SetNodeExpansionState(Child, bExpand, bRecursive);
 			}
 		}
 	}
