Fixed a gap in the header of nested categories in the details view, notably in a blueprint's Class Defaults.

diff --git a/Engine/Source/Editor/PropertyEditor/Private/SDetailCategoryTableRow.cpp b/Engine/Source/Editor/PropertyEditor/Private/SDetailCategoryTableRow.cpp
--- a/Engine/Source/Editor/PropertyEditor/Private/SDetailCategoryTableRow.cpp
+++ b/Engine/Source/Editor/PropertyEditor/Private/SDetailCategoryTableRow.cpp
@@ -76,8 +76,6 @@ void SDetailCategoryTableRow::Construct(const FArguments& InArgs, TSharedRef<FDe
 		[
 			SNew(SHorizontalBox)
 			+ SHorizontalBox::Slot()
-			.HAlign(HAlign_Fill)
-			.Padding(0)
 			[
 				SNew(SBorder)
 				.BorderImage(this, &SDetailCategoryTableRow::GetBackgroundImage)
@@ -86,7 +84,6 @@ void SDetailCategoryTableRow::Construct(const FArguments& InArgs, TSharedRef<FDe
 				[
 					SNew(SBox)
 					.MinDesiredHeight(PropertyEditorConstants::PropertyRowHeight)
-					.VAlign(VAlign_Center)
 					[
 						HeaderBox
 					]
@@ -94,8 +91,6 @@ void SDetailCategoryTableRow::Construct(const FArguments& InArgs, TSharedRef<FDe
 			]
 			+ SHorizontalBox::Slot()
 			.HAlign(HAlign_Right)
-			.VAlign(VAlign_Fill)
-			.Padding(0)
 			.AutoWidth()
 			[
 				SNew(SBorder)
