Adding a property module setting for overriding property matrix availability

diff --git a/Engine/Plugins/Editor/ContentBrowser/ContentBrowserAssetDataSource/Source/ContentBrowserAssetDataSource/Private/AssetFileContextMenu.cpp b/Engine/Plugins/Editor/ContentBrowser/ContentBrowserAssetDataSource/Source/ContentBrowserAssetDataSource/Private/AssetFileContextMenu.cpp
--- a/Engine/Plugins/Editor/ContentBrowser/ContentBrowserAssetDataSource/Source/ContentBrowserAssetDataSource/Private/AssetFileContextMenu.cpp
+++ b/Engine/Plugins/Editor/ContentBrowser/ContentBrowserAssetDataSource/Source/ContentBrowserAssetDataSource/Private/AssetFileContextMenu.cpp
@@ -468,7 +468,7 @@ void FAssetFileContextMenu::MakeAssetActionsSubMenu(UToolMenu* Menu)
 		}
 
 		// Property Matrix
-		bool bCanUsePropertyMatrix = true;
+		bool bCanUsePropertyMatrix = FModuleManager::LoadModuleChecked<FPropertyEditorModule>( "PropertyEditor" ).GetCanUsePropertyMatrix();
 		// Materials can't be bulk edited currently as they require very special handling because of their dependencies with the rendering thread, and we'd have to hack the property matrix too much.
 		for (auto& Asset : SelectedAssets)
 		{
diff --git a/Engine/Source/Editor/PropertyEditor/Private/SDetailsView.cpp b/Engine/Source/Editor/PropertyEditor/Private/SDetailsView.cpp
--- a/Engine/Source/Editor/PropertyEditor/Private/SDetailsView.cpp
+++ b/Engine/Source/Editor/PropertyEditor/Private/SDetailsView.cpp
@@ -276,8 +276,7 @@ void SDetailsView::Construct(const FArguments& InArgs, const FDetailsViewArgs& I
 			.OnTextCommitted(this, &SDetailsView::OnFilterTextCommitted)
 			.AddMetaData<FTagMetaData>(TEXT("Details.Search"))
 		];
-
-	if (DetailsViewArgs.bShowPropertyMatrixButton)
+	
 	{
 		FilterRowHBox->AddSlot()
 			.Padding(0)
@@ -290,6 +289,7 @@ void SDetailsView::Construct(const FArguments& InArgs, const FDetailsViewArgs& I
 				.ButtonStyle( FAppStyle::Get(), "SimpleButton" )
 				.OnClicked(this, &SDetailsView::OnOpenRawPropertyEditorClicked)
 				.IsEnabled(this, &SDetailsView::CanOpenRawPropertyEditor)
+				.Visibility(this, &SDetailsView::CanShowRawPropertyEditorButton, DetailsViewArgs.bShowPropertyMatrixButton)
 				.ToolTipText(LOCTEXT("RawPropertyEditorButtonLabel", "Open Selection in Property Matrix"))
 				[
 					SNew(SImage)
@@ -451,6 +451,12 @@ bool SDetailsView::CanOpenRawPropertyEditor() const
 	return SelectedObjects.Num() > 0 && IsPropertyEditingEnabled();
 }
 
+EVisibility SDetailsView::CanShowRawPropertyEditorButton(const bool bAllowedByDetailsViewArgs) const
+{
+	const bool bShowPropertyMatrixOverride =  FModuleManager::LoadModuleChecked<FPropertyEditorModule>( "PropertyEditor" ).GetCanUsePropertyMatrix();
+	return (DetailsViewArgs.bShowPropertyMatrixButton && bShowPropertyMatrixOverride) ? EVisibility::Visible : EVisibility::Collapsed;
+}
+
 FReply SDetailsView::OnOpenRawPropertyEditorClicked()
 {
 	FPropertyEditorModule& PropertyEditorModule = FModuleManager::LoadModuleChecked<FPropertyEditorModule>( "PropertyEditor" );
diff --git a/Engine/Source/Editor/PropertyEditor/Private/SDetailsView.h b/Engine/Source/Editor/PropertyEditor/Private/SDetailsView.h
--- a/Engine/Source/Editor/PropertyEditor/Private/SDetailsView.h
+++ b/Engine/Source/Editor/PropertyEditor/Private/SDetailsView.h
@@ -152,6 +152,9 @@ private:
 	/** Whether the property matrix button should be enabled */
 	bool CanOpenRawPropertyEditor() const;
 
+	/** Whether the property matrix button should be shown */
+	EVisibility CanShowRawPropertyEditorButton(const bool bAllowedByDetailsViewArgs) const;
+
 	/**
 	 * Called to open the raw property editor (property matrix)                                                              
 	 */
diff --git a/Engine/Source/Editor/PropertyEditor/Public/PropertyEditorModule.h b/Engine/Source/Editor/PropertyEditor/Public/PropertyEditorModule.h
--- a/Engine/Source/Editor/PropertyEditor/Public/PropertyEditorModule.h
+++ b/Engine/Source/Editor/PropertyEditor/Public/PropertyEditorModule.h
@@ -404,6 +404,12 @@ public:
 	/** Get the global row extension generators. */
 	FOnGenerateGlobalRowExtension& GetGlobalRowExtensionDelegate() { return OnGenerateGlobalRowExtension; }
 
+	const bool GetCanUsePropertyMatrix() const { return bCanUsePropertyMatrixOverride; } 
+	void SetCanUsePropertyMatrix(const bool bInCanUsePropertyMatrix) 
+	{ 
+		bCanUsePropertyMatrixOverride = bInCanUsePropertyMatrix; 
+	} 
+
 private:
 
 	/**
@@ -447,4 +453,6 @@ private:
 	UStruct* StructOnScopePropertyOwner;
 	/** Delegate called to extend the name column widget on a property row. */
 	FOnGenerateGlobalRowExtension OnGenerateGlobalRowExtension;
+	/** Override for if the Property Matrix is availabe */
+	bool bCanUsePropertyMatrixOverride = true;
 };
