Property Editor:

-Fix assert resolving null to shared ref with vector properties

diff --git a/Engine/Source/Editor/PropertyEditor/Private/PropertyHandleImpl.cpp b/Engine/Source/Editor/PropertyEditor/Private/PropertyHandleImpl.cpp
--- a/Engine/Source/Editor/PropertyEditor/Private/PropertyHandleImpl.cpp
+++ b/Engine/Source/Editor/PropertyEditor/Private/PropertyHandleImpl.cpp
@@ -4257,21 +4257,29 @@ bool FPropertyHandleVector::Supports( TSharedRef<FPropertyNode> PropertyNode )
 FPropertyHandleVector::FPropertyHandleVector( TSharedRef<class FPropertyNode> PropertyNode, class FNotifyHook* NotifyHook, TSharedPtr<IPropertyUtilities> PropertyUtilities )
 	: FPropertyHandleBase( PropertyNode, NotifyHook, PropertyUtilities ) 
 {
-	const bool bRecurse = false;
-	// A vector is a struct property that has 3 children.  We get/set the values from the children
-	VectorComponents.Add( MakeShareable( new FPropertyHandleMixed( Implementation->GetChildNode("X", bRecurse).ToSharedRef(), NotifyHook, PropertyUtilities ) ) );
-
-	VectorComponents.Add( MakeShareable( new FPropertyHandleMixed( Implementation->GetChildNode("Y", bRecurse).ToSharedRef(), NotifyHook, PropertyUtilities ) ) );
-
-	if( Implementation->GetNumChildren() > 2 )
+	if( Implementation->GetNumChildren() > 0 )
 	{
-		// at least a 3 component vector
-		VectorComponents.Add( MakeShareable( new FPropertyHandleMixed( Implementation->GetChildNode("Z",bRecurse).ToSharedRef(), NotifyHook, PropertyUtilities ) ) );
-	}
-	if( Implementation->GetNumChildren() > 3 )
-	{
-		// a 4 component vector
-		VectorComponents.Add( MakeShareable( new FPropertyHandleMixed( Implementation->GetChildNode("W",bRecurse).ToSharedRef(), NotifyHook, PropertyUtilities ) ) );
+		const bool bRecurse = false;
+		// A vector is a struct property that has multiple children.  We get/set the values from the children
+		VectorComponents.Add( MakeShareable( new FPropertyHandleMixed( Implementation->GetChildNode("X", bRecurse).ToSharedRef(), NotifyHook, PropertyUtilities ) ) );
+
+		if( Implementation->GetNumChildren() > 1 )
+		{
+			// at least a 2 component vector
+			VectorComponents.Add( MakeShareable( new FPropertyHandleMixed( Implementation->GetChildNode("Y", bRecurse).ToSharedRef(), NotifyHook, PropertyUtilities ) ) );
+		}
+		
+		if( Implementation->GetNumChildren() > 2 )
+		{
+			// at least a 3 component vector
+			VectorComponents.Add( MakeShareable( new FPropertyHandleMixed( Implementation->GetChildNode("Z",bRecurse).ToSharedRef(), NotifyHook, PropertyUtilities ) ) );
+		}
+		
+		if( Implementation->GetNumChildren() > 3 )
+		{
+			// a 4 component vector
+			VectorComponents.Add( MakeShareable( new FPropertyHandleMixed( Implementation->GetChildNode("W",bRecurse).ToSharedRef(), NotifyHook, PropertyUtilities ) ) );
+		}
 	}
 }
 
