Removed inline forward declarations.

#rnx

diff --git a/Engine/Source/Editor/PropertyEditor/Private/SDetailsViewBase.h b/Engine/Source/Editor/PropertyEditor/Private/SDetailsViewBase.h
--- a/Engine/Source/Editor/PropertyEditor/Private/SDetailsViewBase.h
+++ b/Engine/Source/Editor/PropertyEditor/Private/SDetailsViewBase.h
@@ -26,18 +26,24 @@
 #include "Widgets/Views/STableViewBase.h"
 #include "Widgets/Views/STreeView.h"
 
+class FAssetThumbnailPool;
 class FDetailCategoryImpl;
 class FDetailLayoutBuilderImpl;
+class FDetailTreeNode;
 class FNotifyHook;
+class FPropertyEditor;
+class FUICommandList;
+class IClassViewerFilter;
 class IDetailCustomization;
 class IDetailKeyframeHandler;
 class IDetailPropertyExtensionHandler;
-class SDetailNameArea;
 class IPropertyGenerationUtilities;
+class IPropertyUtilities;
+class SDetailNameArea;
 class SSearchBox;
 struct FDetailsViewObjectRoot;
 
-typedef STreeView< TSharedRef<class FDetailTreeNode> > SDetailTree;
+typedef STreeView< TSharedRef<FDetailTreeNode> > SDetailTree;
 
 class SDetailsViewBase : public IDetailsViewPrivate
 {
@@ -94,9 +100,9 @@ public:
 	virtual void SetIsPropertyEditingEnabledDelegate(FIsPropertyEditingEnabled IsPropertyEditingEnabled) override;
 	virtual FIsPropertyEditingEnabled& GetIsPropertyEditingEnabledDelegate() override { return IsPropertyEditingEnabledDelegate; }
 	virtual bool IsPropertyEditingEnabled() const override;
-	virtual void SetKeyframeHandler(TSharedPtr<class IDetailKeyframeHandler> InKeyframeHandler) override;
+	virtual void SetKeyframeHandler(TSharedPtr<IDetailKeyframeHandler> InKeyframeHandler) override;
 	virtual TSharedPtr<IDetailKeyframeHandler> GetKeyframeHandler() const override { return KeyframeHandler; }
-	virtual void SetExtensionHandler(TSharedPtr<class IDetailPropertyExtensionHandler> InExtensionHandler) override;
+	virtual void SetExtensionHandler(TSharedPtr<IDetailPropertyExtensionHandler> InExtensionHandler) override;
 	virtual TSharedPtr<IDetailPropertyExtensionHandler> GetExtensionHandler() const override { return ExtensionHandler; }
 	virtual bool IsPropertyVisible(const struct FPropertyAndParent& PropertyAndParent) const override;
 	virtual bool IsPropertyReadOnly(const struct FPropertyAndParent& PropertyAndParent) const override;
@@ -122,10 +128,10 @@ public:
 	bool GetCustomSavedExpansionState(const FString& NodePath) const override;
 	virtual void NotifyFinishedChangingProperties(const FPropertyChangedEvent& PropertyChangedEvent) override;
 	void RefreshTree() override;
-	TSharedPtr<class FAssetThumbnailPool> GetThumbnailPool() const override;
-	virtual const TArray<TSharedRef<class IClassViewerFilter>>& GetClassViewerFilters() const override;
+	TSharedPtr<FAssetThumbnailPool> GetThumbnailPool() const override;
+	virtual const TArray<TSharedRef<IClassViewerFilter>>& GetClassViewerFilters() const override;
 	TSharedPtr<IPropertyUtilities> GetPropertyUtilities() override;
-	void CreateColorPickerWindow(const TSharedRef< class FPropertyEditor >& PropertyEditor, bool bUseAlpha) override;
+	void CreateColorPickerWindow(const TSharedRef<FPropertyEditor>& PropertyEditor, bool bUseAlpha) override;
 	virtual void UpdateSinglePropertyMap(TSharedPtr<FComplexPropertyNode> InRootPropertyNode, FDetailLayoutData& LayoutData, bool bIsExternal) override;
 	virtual FNotifyHook* GetNotifyHook() const override { return DetailsViewArgs.NotifyHook; }
 	virtual const FCustomPropertyTypeLayoutMap& GetCustomPropertyTypeLayoutMap() const { return InstancedTypeToLayoutMap; }
@@ -291,7 +297,7 @@ protected:
 	virtual TSharedPtr<SWidget> GetNameAreaWidget() override;
 	virtual void SetNameAreaCustomContent(TSharedRef<SWidget>& InCustomContent) override;
 	virtual TSharedPtr<SWidget> GetFilterAreaWidget() override;
-	virtual TSharedPtr<class FUICommandList> GetHostCommandList() const override;
+	virtual TSharedPtr<FUICommandList> GetHostCommandList() const override;
 	virtual TSharedPtr<FTabManager> GetHostTabManager() const override;
 	virtual void SetHostTabManager(TSharedPtr<FTabManager> InTabManager) override;
 
@@ -395,11 +401,11 @@ protected:
 	/** Whether or not this instance of the details view opened a color picker and it is not closed yet */
 	bool bHasOpenColorPicker;
 	/** Settings for this view */
-	TSharedPtr<class IPropertyUtilities> PropertyUtilities;
+	TSharedPtr<IPropertyUtilities> PropertyUtilities;
 	/** Gets internal utilities for generating property layouts. */
 	TSharedPtr<IPropertyGenerationUtilities> PropertyGenerationUtilities;
 	/** The name area which is not recreated when selection changes */
-	TSharedPtr<class SDetailNameArea> NameArea;
+	TSharedPtr<SDetailNameArea> NameArea;
 	/** The current filter */
 	FDetailFilter CurrentFilter;
 	/** Delegate called to get generic details not specific to an object being viewed */
@@ -448,7 +454,7 @@ protected:
 	mutable TSharedPtr<FEditConditionParser> EditConditionParser;
 	
 	/** Optional custom filter(s) to be applied when selecting values for class properties */
-	TArray<TSharedRef<class IClassViewerFilter>> ClassViewerFilters;
+	TArray<TSharedRef<IClassViewerFilter>> ClassViewerFilters;
 
 	/** The EnsureDataIsValid function can be skipped with this member, if set. Useful if your implementation doesn't require this kind of validation each Tick. */
 	FOnValidateDetailsViewPropertyNodes CustomValidatePropertyNodesFunction;
diff --git a/Engine/Source/Editor/PropertyEditor/Public/PropertyHandle.h b/Engine/Source/Editor/PropertyEditor/Public/PropertyHandle.h
--- a/Engine/Source/Editor/PropertyEditor/Public/PropertyHandle.h
+++ b/Engine/Source/Editor/PropertyEditor/Public/PropertyHandle.h
@@ -9,6 +9,8 @@
 #include "UObject/PropertyPortFlags.h"
 
 struct FAssetData;
+class FPropertyRestriction;
+class FResetToDefaultOverride;
 class IPropertyHandleArray;
 class IPropertyHandleMap;
 class IPropertyHandleSet;
@@ -482,17 +484,17 @@ public:
 	 *
 	 * @return the handle as an array if it is an array (static or dynamic)
 	 */
-	virtual TSharedPtr<class IPropertyHandleArray> AsArray() = 0;
+	virtual TSharedPtr<IPropertyHandleArray> AsArray() = 0;
 
 	/**
 	 * @return This handle as a set if possible
 	 */
-	virtual TSharedPtr<class IPropertyHandleSet> AsSet() = 0;
+	virtual TSharedPtr<IPropertyHandleSet> AsSet() = 0;
 
 	/**
 	 * @return This handle as a map if possible
 	 */
-	virtual TSharedPtr<class IPropertyHandleMap> AsMap() = 0;
+	virtual TSharedPtr<IPropertyHandleMap> AsMap() = 0;
 
 	/**
 	 * @return The display name of the property
@@ -598,7 +600,7 @@ public:
 	 * Adds a restriction to the possible values for this property.
 	 * @param Restriction	The restriction being added to this property.
 	 */
-	virtual void AddRestriction( TSharedRef<const class FPropertyRestriction> Restriction ) = 0;
+	virtual void AddRestriction( TSharedRef<const FPropertyRestriction> Restriction ) = 0;
 
 	/**
 	* Tests if a value is restricted for this property
@@ -677,7 +679,7 @@ public:
 	/**
 	 * Sets an override for this property's reset to default behavior
 	 */
-	virtual void ExecuteCustomResetToDefault(const class FResetToDefaultOverride& OnCustomResetToDefault) = 0;
+	virtual void ExecuteCustomResetToDefault(const FResetToDefaultOverride& OnCustomResetToDefault) = 0;
 
 	/**
 	 * Gets the category FName that a property is in at the default location defined by the class the property is in
