[BugFix] Editor Crashes when selecting an Addition/Removal change in Defaults section of the Blueprint Diff window
 - this fix additionally has a side effect of fixing a bug that caused parent items in details diffs to not be selected when clicked

diff --git a/Engine/Source/Editor/Kismet/Private/DiffUtils.cpp b/Engine/Source/Editor/Kismet/Private/DiffUtils.cpp
--- a/Engine/Source/Editor/Kismet/Private/DiffUtils.cpp
+++ b/Engine/Source/Editor/Kismet/Private/DiffUtils.cpp
@@ -197,6 +197,12 @@ FPropertyPath FPropertySoftPath::ResolvePath(const UObject* Object) const
 		FName PropertyIdentifier = PropertyChain[I].PropertyName;
 		FProperty* ResolvedProperty = UEDiffUtils_Private::Resolve(ContainerStruct, PropertyIdentifier);
 
+		// If the property didn't exist inside the container, return an invalid property
+		if (!ResolvedProperty)
+		{
+			return FPropertyPath();
+		}
+		
 		FPropertyInfo Info(ResolvedProperty, INDEX_NONE);
 		Ret.AddProperty(Info);
 
diff --git a/Engine/Source/Editor/PropertyEditor/Private/SDetailsViewBase.cpp b/Engine/Source/Editor/PropertyEditor/Private/SDetailsViewBase.cpp
--- a/Engine/Source/Editor/PropertyEditor/Private/SDetailsViewBase.cpp
+++ b/Engine/Source/Editor/PropertyEditor/Private/SDetailsViewBase.cpp
@@ -240,10 +240,11 @@ static void GetPropertyPathToTreeNodes(const TArray< TSharedRef<FDetailTreeNode>
 {
 	for (const TSharedRef<FDetailTreeNode>& TreeNode : RootTreeNodes)
 	{
-		if (TreeNode->IsLeaf())
+		FPropertyPath Path = TreeNode->GetPropertyPath();
+		if (Path.IsValid())
 		{
-			TArray<TSharedRef<FDetailTreeNode>>& Nodes = OutMap.FindOrAdd(TreeNode->GetPropertyPath().ToString());
-			Nodes.Add(TreeNode);
+			TArray<TSharedRef<FDetailTreeNode>>& Nodes = OutMap.FindOrAdd(Path.ToString());
+            Nodes.Add(TreeNode);
 		}
 
 		// Need to check children even if we're a leaf, because all DetailItemNodes are leaves, even if they may have sub-children
@@ -271,6 +272,11 @@ static TArray< TSharedRef< FDetailTreeNode > > FindBestFitTreeNodeFromProperty(c
 
 void SDetailsViewBase::HighlightProperty(const FPropertyPath& Property)
 {
+	if (!Property.IsValid())
+	{
+		return;
+	}
+	
 	TSharedPtr<FDetailTreeNode> PrevHighlightedNodePtr = CurrentlyHighlightedNode.Pin();
 	if (PrevHighlightedNodePtr.IsValid())
 	{
diff --git a/Engine/Source/Editor/PropertyEditor/Public/PropertyPath.h b/Engine/Source/Editor/PropertyEditor/Public/PropertyPath.h
--- a/Engine/Source/Editor/PropertyEditor/Public/PropertyPath.h
+++ b/Engine/Source/Editor/PropertyEditor/Public/PropertyPath.h
@@ -56,6 +56,11 @@ public:
 		return NewPath;
 	}
 
+	bool IsValid() const 
+	{ 
+		return !Properties.IsEmpty(); 
+	}
+
 	int32 GetNumProperties() const 
 	{ 
 		return Properties.Num(); 
