[BugFix] Details View Diff doesn't scroll to all places displayed in the diff
[BugFix] Details View Diff doesn't open folded categories that contains diffs

diff --git a/Engine/Source/Editor/PropertyEditor/Private/SDetailsViewBase.cpp b/Engine/Source/Editor/PropertyEditor/Private/SDetailsViewBase.cpp
--- a/Engine/Source/Editor/PropertyEditor/Private/SDetailsViewBase.cpp
+++ b/Engine/Source/Editor/PropertyEditor/Private/SDetailsViewBase.cpp
@@ -204,47 +204,15 @@ TArray< FPropertyPath > SDetailsViewBase::GetPropertiesInOrderDisplayed() const
 	return Ret;
 }
 
-// @return populates OutNodes with the leaf node corresponding to property as the first entry in the list (e.g. [leaf, parent, grandparent]):
-static void FindTreeNodeFromPropertyRecursive( const TArray< TSharedRef<FDetailTreeNode> >& Nodes, const FPropertyPath& Property, TArray< TSharedPtr< FDetailTreeNode > >& OutNodes )
-{
-	if (Property == FPropertyPath())
-	{
-		return;
-	}
-
-	for (const TSharedRef<FDetailTreeNode>& TreeNode : Nodes)
-	{
-		if (TreeNode->IsLeaf())
-		{
-			FPropertyPath tmp = TreeNode->GetPropertyPath();
-			if( Property == tmp )
-			{
-				OutNodes.Push(TreeNode);
-				return;
-			}
-		}
-
-		// Need to check children even if we're a leaf, because all DetailItemNodes are leaves, even if they may have sub-children
-		TArray< TSharedRef<FDetailTreeNode> > Children;
-		TreeNode->GetChildren(Children);
-		FindTreeNodeFromPropertyRecursive(Children, Property, OutNodes);
-		if (OutNodes.Num() > 0)
-		{
-			OutNodes.Push(TreeNode);
-			return;
-		}
-	}
-}
-
-static void GetPropertyPathToTreeNodes(const TArray< TSharedRef<FDetailTreeNode> >& RootTreeNodes, TMap<FString, TArray<TSharedRef<FDetailTreeNode>>> &OutMap)
+// construct a map for fast FPropertyPath.ToString() -> FDetailTreeNode lookup
+static void GetPropertyPathToTreeNodes(const TArray<TSharedRef<FDetailTreeNode>>& RootTreeNodes, TMap<FString, TSharedRef<FDetailTreeNode>> &OutMap)
 {
 	for (const TSharedRef<FDetailTreeNode>& TreeNode : RootTreeNodes)
 	{
 		FPropertyPath Path = TreeNode->GetPropertyPath();
 		if (Path.IsValid())
 		{
-			TArray<TSharedRef<FDetailTreeNode>>& Nodes = OutMap.FindOrAdd(Path.ToString());
-            Nodes.Add(TreeNode);
+			OutMap.Add(Path.ToString(), TreeNode);
 		}
 
 		// Need to check children even if we're a leaf, because all DetailItemNodes are leaves, even if they may have sub-children
@@ -254,20 +222,22 @@ static void GetPropertyPathToTreeNodes(const TArray< TSharedRef<FDetailTreeNode>
 	}
 }
 
-static TArray< TSharedRef< FDetailTreeNode > > FindBestFitTreeNodeFromProperty(const TArray< TSharedRef<FDetailTreeNode> >& RootTreeNodes, FPropertyPath Property)
+// Find an FDetailTreeNode in the tree that most closely matches the provided property.
+// If there isn't an exact match, the provided property will be trimmed until we find a match
+static TSharedPtr<FDetailTreeNode> FindBestFitTreeNodeFromProperty(const TArray<TSharedRef<FDetailTreeNode>>& RootTreeNodes, FPropertyPath Property)
 {
-	TMap<FString, TArray<TSharedRef<FDetailTreeNode>>> PropertyPathToTreeNodes;
+	TMap<FString, TSharedRef<FDetailTreeNode>> PropertyPathToTreeNodes;
 	GetPropertyPathToTreeNodes(RootTreeNodes, PropertyPathToTreeNodes);
 	
-	TArray< TSharedRef< FDetailTreeNode > >* TreeNodeChain = PropertyPathToTreeNodes.Find(Property.ToString());
+	TSharedRef< FDetailTreeNode >* TreeNode = PropertyPathToTreeNodes.Find(Property.ToString());
 
 	// if we couldn't find the exact property, trim the path to get the next best option
-	while (!TreeNodeChain && Property.GetNumProperties() > 1)
+	while (!TreeNode && Property.GetNumProperties() > 1)
 	{
 		Property = *Property.TrimPath(1);
-		TreeNodeChain = PropertyPathToTreeNodes.Find(Property.ToString());
+		TreeNode = PropertyPathToTreeNodes.Find(Property.ToString());
 	}
-	return TreeNodeChain? MoveTemp(*TreeNodeChain) : TArray< TSharedRef< FDetailTreeNode > >{};
+	return TreeNode? *TreeNode : TSharedPtr<FDetailTreeNode>();
 }
 
 void SDetailsViewBase::HighlightProperty(const FPropertyPath& Property)
@@ -283,25 +253,25 @@ void SDetailsViewBase::HighlightProperty(const FPropertyPath& Property)
 		PrevHighlightedNodePtr->SetIsHighlighted(false);
 	}
 
-	const TArray< TSharedRef< FDetailTreeNode > > TreeNodeChain = FindBestFitTreeNodeFromProperty(RootTreeNodes, Property);
-	
-	TSharedPtr< FDetailTreeNode > FinalNodePtr = nullptr;
-	if (TreeNodeChain.Num() > 0)
+	const TSharedPtr< FDetailTreeNode > TreeNode = FindBestFitTreeNodeFromProperty(RootTreeNodes, Property);
+	if (TreeNode.IsValid())
 	{
-		FinalNodePtr = TreeNodeChain[0];
-		check(FinalNodePtr.IsValid());
-		FinalNodePtr->SetIsHighlighted(true);
+		// highlight the found node
+		TreeNode->SetIsHighlighted(true);
 
-		for (int ParentIndex = 1; ParentIndex < TreeNodeChain.Num(); ++ParentIndex)
+		// make sure all ancestors are expanded so we can see the found node
+		TSharedPtr< FDetailTreeNode > Ancestor = TreeNode;
+		while(Ancestor->GetParentNode().IsValid())
 		{
-			TSharedPtr< FDetailTreeNode > CurrentParent = TreeNodeChain[ParentIndex];
-			check(CurrentParent.IsValid());
-			DetailTree->SetItemExpansion(CurrentParent.ToSharedRef(), true);
+			Ancestor = Ancestor->GetParentNode().Pin();
+			DetailTree->SetItemExpansion(Ancestor.ToSharedRef(), true);
 		}
-
-		DetailTree->RequestScrollIntoView(FinalNodePtr.ToSharedRef());
+		
+		// scroll to the found node
+		DetailTree->RequestScrollIntoView(TreeNode.ToSharedRef());
 	}
-	CurrentlyHighlightedNode = FinalNodePtr;
+	
+	CurrentlyHighlightedNode = TreeNode;
 }
 
 void SDetailsViewBase::ShowAllAdvancedProperties()
