Added a max length limit for string/text properties that can be set via the MaxLength property meta-data

diff --git a/Engine/Source/Editor/PropertyEditor/Private/UserInterface/PropertyEditor/SPropertyEditorText.cpp b/Engine/Source/Editor/PropertyEditor/Private/UserInterface/PropertyEditor/SPropertyEditorText.cpp
--- a/Engine/Source/Editor/PropertyEditor/Private/UserInterface/PropertyEditor/SPropertyEditorText.cpp
+++ b/Engine/Source/Editor/PropertyEditor/Private/UserInterface/PropertyEditor/SPropertyEditorText.cpp
@@ -32,10 +32,14 @@ void SPropertyEditorText::Construct( const FArguments& InArgs, const TSharedRef<
 {
 	PropertyEditor = InPropertyEditor;
 
-	bIsFNameProperty = InPropertyEditor->PropertyIsA(FNameProperty::StaticClass());
-	bIsMultiLine = InPropertyEditor->GetPropertyHandle()->GetBoolMetaData("MultiLine");
+	static const FName NAME_MaxLength = "MaxLength";
+	static const FName NAME_MultiLine = "MultiLine";
+	static const FName NAME_PasswordField = "PasswordField";
 
-	const bool bIsPassword = InPropertyEditor->GetPropertyHandle()->GetBoolMetaData("PasswordField");
+	bIsMultiLine = InPropertyEditor->GetPropertyHandle()->GetBoolMetaData(NAME_MultiLine);
+	MaxLength = InPropertyEditor->PropertyIsA(FNameProperty::StaticClass()) ? NAME_SIZE - 1 : InPropertyEditor->GetPropertyHandle()->GetIntMetaData(NAME_MaxLength);
+
+	const bool bIsPassword = InPropertyEditor->GetPropertyHandle()->GetBoolMetaData(NAME_PasswordField);
 	
 	TSharedPtr<SHorizontalBox> HorizontalBox;
 	if(bIsMultiLine)
@@ -52,7 +56,7 @@ void SPropertyEditorText::Construct( const FArguments& InArgs, const TSharedRef<
 				.SelectAllTextWhenFocused(false)
 				.ClearKeyboardFocusOnCommit(false)
 				.OnTextCommitted(this, &SPropertyEditorText::OnTextCommitted)
-				.OnTextChanged(this, &SPropertyEditorText::OnMultiLineTextChanged)
+				.OnVerifyTextChanged(this, &SPropertyEditorText::OnVerifyTextChanged)
 				.SelectAllTextOnCommit(false)
 				.IsReadOnly(this, &SPropertyEditorText::IsReadOnly)
 				.AutoWrapText(true)
@@ -77,7 +81,7 @@ void SPropertyEditorText::Construct( const FArguments& InArgs, const TSharedRef<
 				.SelectAllTextWhenFocused( true )
 				.ClearKeyboardFocusOnCommit(false)
 				.OnTextCommitted( this, &SPropertyEditorText::OnTextCommitted )
-				.OnTextChanged( this, &SPropertyEditorText::OnSingleLineTextChanged )
+				.OnVerifyTextChanged( this, &SPropertyEditorText::OnVerifyTextChanged )
 				.SelectAllTextOnCommit( true )
 				.IsReadOnly(this, &SPropertyEditorText::IsReadOnly)
 				.IsPassword( bIsPassword )
@@ -144,33 +148,15 @@ void SPropertyEditorText::OnTextCommitted( const FText& NewText, ETextCommit::Ty
 	}
 }
 
-static FText ValidateNameLength( const FText& Text )
+bool SPropertyEditorText::OnVerifyTextChanged(const FText& Text, FText& OutError)
 {
-	if( Text.ToString().Len() >= NAME_SIZE )
+	if (MaxLength > 0 && Text.ToString().Len() > MaxLength)
 	{
-		static FText ErrorString = FText::Format( LOCTEXT("NamePropertySizeTooLongError", "Name properties may only be a maximum of {0} characters"), FText::AsNumber( NAME_SIZE ) );
-		return ErrorString;
+		OutError = FText::Format(LOCTEXT("PropertyTextTooLongError", "This value is too long ({0}/{1} characters)"), Text.ToString().Len(), MaxLength);
+		return false;
 	}
 
-	return FText::GetEmpty();
-}
-
-void SPropertyEditorText::OnMultiLineTextChanged( const FText& NewText )
-{
-	if( bIsFNameProperty )
-	{
-		FText ErrorMessage = ValidateNameLength( NewText );
-		MultiLineWidget->SetError( ErrorMessage );
-	}
-}
-
-void SPropertyEditorText::OnSingleLineTextChanged( const FText& NewText )
-{
-	if( bIsFNameProperty )
-	{
-		FText ErrorMessage = ValidateNameLength( NewText );
-		SingleLineWidget->SetError( ErrorMessage );
-	}
+	return true;
 }
 
 bool SPropertyEditorText::SupportsKeyboardFocus() const
diff --git a/Engine/Source/Editor/PropertyEditor/Private/UserInterface/PropertyEditor/SPropertyEditorText.h b/Engine/Source/Editor/PropertyEditor/Private/UserInterface/PropertyEditor/SPropertyEditorText.h
--- a/Engine/Source/Editor/PropertyEditor/Private/UserInterface/PropertyEditor/SPropertyEditorText.h
+++ b/Engine/Source/Editor/PropertyEditor/Private/UserInterface/PropertyEditor/SPropertyEditorText.h
@@ -42,11 +42,7 @@ private:
 
 	void OnTextCommitted( const FText& NewText, ETextCommit::Type CommitInfo );
 
-	/** Called if the single line widget text changes */
-	void OnSingleLineTextChanged( const FText& NewText );
-
-	/** Called if the multi line widget text changes */
-	void OnMultiLineTextChanged( const FText& NewText );
+	bool OnVerifyTextChanged(const FText& Text, FText& OutError);
 
 	/** @return True if the property can be edited */
 	bool CanEdit() const;
@@ -69,6 +65,6 @@ private:
 	/** Cached flag as we would like multi-line text widgets to be slightly larger */
 	bool bIsMultiLine;
 	
-	/** True if property is an FName property which causes us to run extra size validation checks */
-	bool bIsFNameProperty;
+	/** The maximum length of the value that can be edited, or <=0 for unlimited */
+	int32 MaxLength = 0;
 };
diff --git a/Engine/Source/Runtime/CoreUObject/Public/UObject/ObjectMacros.h b/Engine/Source/Runtime/CoreUObject/Public/UObject/ObjectMacros.h
--- a/Engine/Source/Runtime/CoreUObject/Public/UObject/ObjectMacros.h
+++ b/Engine/Source/Runtime/CoreUObject/Public/UObject/ObjectMacros.h
@@ -1369,6 +1369,9 @@ namespace UM
 		/// [PropertyMetadata] Used for numeric properties. Stipulates that the value must be a multiple of the metadata value.
 		Multiple,
 
+		/// [PropertyMetadata] Used for FString and FText properties.  Indicates the maximum length of the value that can be edited, or <=0 for unlimited.
+		MaxLength,
+
 		/// [PropertyMetadata] Used for FString and FText properties.  Indicates that the edit field should be multi-line, allowing entry of newlines.
 		MultiLine,
 
