Property Editor:
- Fix an issue with Details View not showing any items for a selected object (e.g. a material) with no registered sections when the previous object had visible sections selected.

diff --git a/Engine/Source/Editor/PropertyEditor/Private/SDetailsView.cpp b/Engine/Source/Editor/PropertyEditor/Private/SDetailsView.cpp
--- a/Engine/Source/Editor/PropertyEditor/Private/SDetailsView.cpp
+++ b/Engine/Source/Editor/PropertyEditor/Private/SDetailsView.cpp
@@ -1246,7 +1246,8 @@ void SDetailsView::RebuildSectionSelector()
 	const TMap<FName, FText> AllSections = GetAllSections();
 	if (AllSections.IsEmpty())
 	{
-		// we've selected something that has no sections - rather than show just "All", hide the box
+		// we've selected something that has no sections - rather than show just "All", hide the box and clear visible sections
+		CurrentFilter.VisibleSections.Reset();
 		return;
 	}
 
