Adding support for specifying an actor filter on SPropertyEditor and SObjectPropertyEntryBox. The new OnShouldFilterActor takes a bool(Actor*) delegate that enables additional customized checks for actors to satisfy a field's requirements. The function should return TRUE if the actor DOES meet the requirements. Also fixed a bug where the actor selection tool associated with this widget would not use the same filter as the dropdown menu.

#p4v-preflight-copy 22462605 [preflight 6345c737f04e6f8505b5c21b failed for unrelated warnings]

diff --git a/Engine/Source/Editor/PropertyEditor/Private/PropertyCustomizationHelpers.cpp b/Engine/Source/Editor/PropertyEditor/Private/PropertyCustomizationHelpers.cpp
--- a/Engine/Source/Editor/PropertyEditor/Private/PropertyCustomizationHelpers.cpp
+++ b/Engine/Source/Editor/PropertyEditor/Private/PropertyCustomizationHelpers.cpp
@@ -464,6 +464,7 @@ void SObjectPropertyEntryBox::Construct( const FArguments& InArgs )
 	OnObjectChanged = InArgs._OnObjectChanged;
 	OnShouldSetAsset = InArgs._OnShouldSetAsset;
 	OnIsEnabled = InArgs._OnIsEnabled;
+	OnShouldFilterActor = InArgs._OnShouldFilterActor;
 
 	const TArray<FAssetData>& OwnerAssetDataArray = InArgs._OwnerAssetDataArray;
 
@@ -529,6 +530,7 @@ void SObjectPropertyEntryBox::Construct( const FArguments& InArgs )
 				.OwnerAssetDataArray(OwnerAssetDataArray)
 				.ThumbnailSize(ThumbnailSize)
 				.DisplayCompactSize(InArgs._DisplayCompactSize)
+				.OnShouldFilterActor(InArgs._OnShouldFilterActor)
 				.CustomContentSlot()
 				[
 					InArgs._CustomContentSlot.Widget
diff --git a/Engine/Source/Editor/PropertyEditor/Private/UserInterface/PropertyEditor/SPropertyEditorAsset.cpp b/Engine/Source/Editor/PropertyEditor/Private/UserInterface/PropertyEditor/SPropertyEditorAsset.cpp
--- a/Engine/Source/Editor/PropertyEditor/Private/UserInterface/PropertyEditor/SPropertyEditorAsset.cpp
+++ b/Engine/Source/Editor/PropertyEditor/Private/UserInterface/PropertyEditor/SPropertyEditorAsset.cpp
@@ -315,6 +315,7 @@ void SPropertyEditorAsset::Construct(const FArguments& InArgs, const TSharedPtr<
 	OwnerAssetDataArray = InArgs._OwnerAssetDataArray;
 	OnSetObject = InArgs._OnSetObject;
 	OnShouldFilterAsset = InArgs._OnShouldFilterAsset;
+	OnShouldFilterActor = InArgs._OnShouldFilterActor;
 	ObjectPath = InArgs._ObjectPath;
 
 	FProperty* Property = nullptr;
@@ -624,7 +625,7 @@ void SPropertyEditorAsset::Construct(const FArguments& InArgs, const TSharedPtr<
 
 	if( bIsActor )
 	{
-		TSharedRef<SWidget> ActorPicker = PropertyCustomizationHelpers::MakeInteractiveActorPicker( FOnGetAllowedClasses::CreateSP(this, &SPropertyEditorAsset::OnGetAllowedClasses), FOnShouldFilterActor(), FOnActorSelected::CreateSP( this, &SPropertyEditorAsset::OnActorSelected ) );
+		TSharedRef<SWidget> ActorPicker = PropertyCustomizationHelpers::MakeInteractiveActorPicker( FOnGetAllowedClasses::CreateSP(this, &SPropertyEditorAsset::OnGetAllowedClasses), FOnShouldFilterActor::CreateSP(this, &SPropertyEditorAsset::IsFilteredActor), FOnActorSelected::CreateSP(this, &SPropertyEditorAsset::OnActorSelected));
 		ActorPicker->SetEnabled( IsEnabledAttribute );
 
 		ButtonBox->AddSlot()
@@ -808,6 +809,10 @@ void SPropertyEditorAsset::OnMenuOpenChanged(bool bOpen)
 bool SPropertyEditorAsset::IsFilteredActor( const AActor* const Actor ) const
 {
 	bool IsAllowed = Actor != nullptr && Actor->IsA(ObjectClass) && !Actor->IsChildActor() && IsClassAllowed(Actor->GetClass());
+	if (IsAllowed && OnShouldFilterActor.IsBound())
+	{
+		IsAllowed = OnShouldFilterActor.Execute(Actor);
+	}
 	return IsAllowed;
 }
 
diff --git a/Engine/Source/Editor/PropertyEditor/Private/UserInterface/PropertyEditor/SPropertyEditorAsset.h b/Engine/Source/Editor/PropertyEditor/Private/UserInterface/PropertyEditor/SPropertyEditorAsset.h
--- a/Engine/Source/Editor/PropertyEditor/Private/UserInterface/PropertyEditor/SPropertyEditorAsset.h
+++ b/Engine/Source/Editor/PropertyEditor/Private/UserInterface/PropertyEditor/SPropertyEditorAsset.h
@@ -78,6 +78,7 @@ public:
 		SLATE_NAMED_SLOT(FArguments, CustomContentSlot)
 		SLATE_ARGUMENT(TSharedPtr<IPropertyHandle>, PropertyHandle)
 		SLATE_ARGUMENT(TArray<FAssetData>, OwnerAssetDataArray)
+		SLATE_EVENT(FOnShouldFilterActor, OnShouldFilterActor)
 
 	SLATE_END_ARGS()
 
@@ -396,6 +397,9 @@ private:
 	/** The property handle, if any */
 	TSharedPtr<IPropertyHandle> PropertyHandle;
 
+	/** Delegate for filtering valid actors */
+	FOnShouldFilterActor OnShouldFilterActor;
+
 	/*
 	 * The reference object on which the picker will assign the picked asset, if any.
 	 * The system will test first the PropertyHandle follow by the PropertyEditor and if nothing this array of assets
diff --git a/Engine/Source/Editor/PropertyEditor/Public/PropertyCustomizationHelpers.h b/Engine/Source/Editor/PropertyEditor/Public/PropertyCustomizationHelpers.h
--- a/Engine/Source/Editor/PropertyEditor/Public/PropertyCustomizationHelpers.h
+++ b/Engine/Source/Editor/PropertyEditor/Public/PropertyCustomizationHelpers.h
@@ -210,6 +210,8 @@ public:
 		/** A custom content slot for widgets */ 
 		SLATE_NAMED_SLOT(FArguments, CustomContentSlot)
 		SLATE_ATTRIBUTE(FIntPoint, ThumbnailSizeOverride)
+		/** Called to check if an actor is valid to use */
+		SLATE_EVENT(FOnShouldFilterActor, OnShouldFilterActor)
 	SLATE_END_ARGS()
 
 	PROPERTYEDITOR_API void Construct( const FArguments& InArgs );
@@ -230,6 +232,8 @@ private:
 private:
 	/** Delegate to call to determine whether the asset should be set */
 	FOnShouldSetAsset OnShouldSetAsset;
+	/** Delegate to call to determine whether the actor should be allowed */
+	FOnShouldFilterActor OnShouldFilterActor;
 	/** Delegate to call when the object changes */
 	FOnSetObject OnObjectChanged;
 	/** Delegate to call to check if this widget should be enabled. */
