Added AllowCreate to SObjectPropertyEntryBox / SPropertyEditorAsset to prevent creating new assets from the asset picker : similar to NoCreate metadata on UPROPERTY, but allows on a non-property based customization

diff --git a/Engine/Source/Editor/PropertyEditor/Private/PropertyCustomizationHelpers.cpp b/Engine/Source/Editor/PropertyEditor/Private/PropertyCustomizationHelpers.cpp
--- a/Engine/Source/Editor/PropertyEditor/Private/PropertyCustomizationHelpers.cpp
+++ b/Engine/Source/Editor/PropertyEditor/Private/PropertyCustomizationHelpers.cpp
@@ -523,6 +523,7 @@ void SObjectPropertyEntryBox::Construct( const FArguments& InArgs )
 				.DisplayThumbnail(bDisplayThumbnail)
 				.OnShouldFilterAsset(InArgs._OnShouldFilterAsset)
 				.AllowClear(InArgs._AllowClear)
+				.AllowCreate(InArgs._AllowCreate)
 				.DisplayUseSelected(InArgs._DisplayUseSelected)
 				.DisplayBrowse(InArgs._DisplayBrowse)
 				.EnableContentPicker(InArgs._EnableContentPicker)
diff --git a/Engine/Source/Editor/PropertyEditor/Private/UserInterface/PropertyEditor/SPropertyEditorAsset.cpp b/Engine/Source/Editor/PropertyEditor/Private/UserInterface/PropertyEditor/SPropertyEditorAsset.cpp
--- a/Engine/Source/Editor/PropertyEditor/Private/UserInterface/PropertyEditor/SPropertyEditorAsset.cpp
+++ b/Engine/Source/Editor/PropertyEditor/Private/UserInterface/PropertyEditor/SPropertyEditorAsset.cpp
@@ -330,6 +330,7 @@ void SPropertyEditorAsset::Construct(const FArguments& InArgs, const TSharedPtr<
 
 	ObjectClass = InArgs._Class != nullptr ? InArgs._Class : GetObjectPropertyClass(Property);
 	bAllowClear = InArgs._AllowClear.IsSet() ? InArgs._AllowClear.GetValue() : (Property ? !(Property->PropertyFlags & CPF_NoClear) : true);
+	bAllowCreate = InArgs._AllowCreate.IsSet() ? InArgs._AllowCreate.GetValue() : (Property ? !Property->HasMetaData("NoCreate") : true);
 
 	InitializeAssetDataTags(Property);
 	if (DisallowedAssetDataTags.IsValid() || RequiredAssetDataTags.IsValid())
@@ -360,8 +361,7 @@ void SPropertyEditorAsset::Construct(const FArguments& InArgs, const TSharedPtr<
 
 	bIsActor = ObjectClass->IsChildOf(AActor::StaticClass());
 
-	const bool bAllowCreation = !Property || !Property->HasMetaData("NoCreate");
-	if (bAllowCreation) 
+	if (bAllowCreate)
 	{
 		if (InArgs._NewAssetFactories.IsSet())
 		{
@@ -1539,4 +1539,4 @@ UClass* SPropertyEditorAsset::GetObjectPropertyClass(const FProperty* Property)
 	return Class;
 }
 
-#undef LOCTEXT_NAMESPACE
\ No newline at end of file
+#undef LOCTEXT_NAMESPACE
diff --git a/Engine/Source/Editor/PropertyEditor/Private/UserInterface/PropertyEditor/SPropertyEditorAsset.h b/Engine/Source/Editor/PropertyEditor/Private/UserInterface/PropertyEditor/SPropertyEditorAsset.h
--- a/Engine/Source/Editor/PropertyEditor/Private/UserInterface/PropertyEditor/SPropertyEditorAsset.h
+++ b/Engine/Source/Editor/PropertyEditor/Private/UserInterface/PropertyEditor/SPropertyEditorAsset.h
@@ -63,6 +63,7 @@ public:
 		, _CustomContentSlot()
 	{}
 		SLATE_ARGUMENT(TOptional<bool>, AllowClear)
+		SLATE_ARGUMENT(TOptional<bool>, AllowCreate)
 		SLATE_ARGUMENT(bool, DisplayThumbnail)
 		SLATE_ARGUMENT(bool, DisplayUseSelected)
 		SLATE_ARGUMENT(bool, DisplayBrowse)
@@ -376,6 +377,9 @@ private:
 	/** Whether the asset can be 'None' in this case */
 	bool bAllowClear;
 
+	/** Whether the asset can be created from the asset picker */
+	bool bAllowCreate;
+
 	/** Whether the object we are editing is an Actor (i.e. requires a Scene Outliner to be displayed) */
 	bool bIsActor;
 
diff --git a/Engine/Source/Editor/PropertyEditor/Public/PropertyCustomizationHelpers.h b/Engine/Source/Editor/PropertyEditor/Public/PropertyCustomizationHelpers.h
--- a/Engine/Source/Editor/PropertyEditor/Public/PropertyCustomizationHelpers.h
+++ b/Engine/Source/Editor/PropertyEditor/Public/PropertyCustomizationHelpers.h
@@ -197,6 +197,8 @@ public:
 		SLATE_EVENT(FOnIsEnabled, OnIsEnabled)
 		/** Whether the asset can be 'None' */
 		SLATE_ARGUMENT(bool, AllowClear)
+		/** Whether the asset can be created from the asset picker directly */
+		SLATE_ARGUMENT(bool, AllowCreate)
 		/** Whether to show the 'Use Selected' button */
 		SLATE_ARGUMENT(bool, DisplayUseSelected)
 		/** Whether to show the 'Browse' button */
