SPropertyEditorNumeric's "ForceUnits" metadata now can also be set per property instance.

Reported on UDN.

diff --git a/Engine/Source/Editor/PropertyEditor/Private/UserInterface/PropertyEditor/SPropertyEditorNumeric.h b/Engine/Source/Editor/PropertyEditor/Private/UserInterface/PropertyEditor/SPropertyEditorNumeric.h
--- a/Engine/Source/Editor/PropertyEditor/Private/UserInterface/PropertyEditor/SPropertyEditorNumeric.h
+++ b/Engine/Source/Editor/PropertyEditor/Private/UserInterface/PropertyEditor/SPropertyEditorNumeric.h
@@ -270,6 +270,7 @@ public:
 			const FString& DeltaString = GetMetaDataFromKey("Delta");
 			const FString& ClampMinString = GetMetaDataFromKey("ClampMin");
 			const FString& ClampMaxString = GetMetaDataFromKey("ClampMax");
+			const FString& ForcedUnits = GetMetaDataFromKey("ForceUnits");
 
 			// If no UIMin/Max was specified then use the clamp string
 			const FString& UIMinString = MetaUIMinString.Len() ? MetaUIMinString : ClampMinString;
@@ -348,7 +349,6 @@ public:
 			// Set up the correct type interface if we want to display units on the property editor
 
 			// First off, check for ForceUnits= meta data. This meta tag tells us to interpret, and always display the value in these units. FUnitConversion::Settings().ShouldDisplayUnits does not apply to such properties
-			const FString& ForcedUnits = PropertyHandle->GetMetaData(TEXT("ForceUnits"));
 			auto PropertyUnits = FUnitConversion::UnitFromString(*ForcedUnits);
 			if (PropertyUnits.IsSet())
 			{
