Editor - When reading ObjectPropertyNodes, treat them as TObjectPtrs so that we can keep them unloaded in a latent loaded situations.

diff --git a/Engine/Source/Editor/PropertyEditor/Private/ObjectPropertyNode.cpp b/Engine/Source/Editor/PropertyEditor/Private/ObjectPropertyNode.cpp
--- a/Engine/Source/Editor/PropertyEditor/Private/ObjectPropertyNode.cpp
+++ b/Engine/Source/Editor/PropertyEditor/Private/ObjectPropertyNode.cpp
@@ -359,13 +359,13 @@ bool FObjectPropertyNode::GetReadAddressUncached(const FPropertyNode& InNode,
 
 					// If BaseObj is NON-NULL, all other objects' properties should also be non-NULL.
 					// If BaseObj is NULL, all other objects' properties should also be NULL.
-					UObject* BaseObj = CastField<FObjectPropertyBase>(InItemProperty)->GetObjectPropertyValue(BaseAddr);
+					const TObjectPtr<UObject> BaseObj = CastField<FObjectPropertyBase>(InItemProperty)->GetObjectPtrPropertyValue(BaseAddr);
 
 					for (int32 ObjIndex = 1; ObjIndex < GetNumObjects(); ObjIndex++)
 					{
-						UObject* CurObj = CastField<FObjectPropertyBase>(InItemProperty)->GetObjectPropertyValue(InNode.GetValueBaseAddressFromObject(GetUObject(ObjIndex)));
-						if (   ( !BaseObj && CurObj )			// BaseObj is NULL, but this InItemProperty is non-NULL!
-							|| ( BaseObj && !CurObj ) )			// BaseObj is non-NULL, but this InItemProperty is NULL!
+						const TObjectPtr<UObject> CurObj = CastField<FObjectPropertyBase>(InItemProperty)->GetObjectPtrPropertyValue(InNode.GetValueBaseAddressFromObject(GetUObject(ObjIndex)));
+						if (   ( !BaseObj && CurObj)			// BaseObj is NULL, but this InItemProperty is non-NULL!
+							|| ( BaseObj && !CurObj) )			// BaseObj is non-NULL, but this InItemProperty is NULL!
 						{
 							bAllTheSame = false;
 							break;
@@ -376,13 +376,12 @@ bool FObjectPropertyNode::GetReadAddressUncached(const FPropertyNode& InNode,
 		}
 	}
 
-		// Write addresses to the output.
+	// Write addresses to the output.
 	if (OutAddresses != nullptr)
-		{
+	{
 		for (int32 ObjIndex = 0; ObjIndex < GetNumObjects(); ++ObjIndex)
-		{	
-			const UObject* TempObject = GetUObject(ObjIndex);
-			if (TempObject)
+		{
+			if (const UObject* TempObject = GetUObject(ObjIndex))
 			{
 				OutAddresses->Add(TempObject, InNode.GetValueBaseAddressFromObject(TempObject));
 			}
@@ -448,7 +447,7 @@ uint8* FObjectPropertyNode::GetValueBaseAddress(uint8* StartAddress, bool bIsSpa
 void FObjectPropertyNode::InitBeforeNodeFlags()
 {
 	StoredProperty = Property;
-	Property = NULL;
+	Property = nullptr;
 
 	Finalize();
 }
@@ -466,7 +465,7 @@ static TSharedPtr<FPropertyNode> FindChildCategory( TSharedPtr<FPropertyNode> Pa
 		check( ChildNode.IsValid() );
 
 		// Is this a category node?
-		FCategoryPropertyNode* ChildCategoryNode = ChildNode->AsCategoryNode();
+		const FCategoryPropertyNode* ChildCategoryNode = ChildNode->AsCategoryNode();
 		if( ChildCategoryNode != nullptr && ChildCategoryNode->GetCategoryName() == CategoryName )
 		{
 			return ChildNode;
diff --git a/Engine/Source/Editor/PropertyEditor/Private/PropertyHandleImpl.cpp b/Engine/Source/Editor/PropertyEditor/Private/PropertyHandleImpl.cpp
--- a/Engine/Source/Editor/PropertyEditor/Private/PropertyHandleImpl.cpp
+++ b/Engine/Source/Editor/PropertyEditor/Private/PropertyHandleImpl.cpp
@@ -2069,6 +2069,8 @@ IMPLEMENT_PROPERTY_ACCESSOR( FQuat )
 IMPLEMENT_PROPERTY_ACCESSOR( FRotator )
 IMPLEMENT_PROPERTY_ACCESSOR( UObject* )
 IMPLEMENT_PROPERTY_ACCESSOR( const UObject* )
+IMPLEMENT_PROPERTY_ACCESSOR( TObjectPtr<UObject> )
+IMPLEMENT_PROPERTY_ACCESSOR( TObjectPtr<const UObject> )
 IMPLEMENT_PROPERTY_ACCESSOR( FAssetData )
 IMPLEMENT_PROPERTY_ACCESSOR( FProperty* )
 IMPLEMENT_PROPERTY_ACCESSOR( const FProperty* )
@@ -3783,10 +3785,26 @@ bool FPropertyHandleObject::Supports( TSharedRef<FPropertyNode> PropertyNode )
 
 FPropertyAccess::Result FPropertyHandleObject::GetValue( UObject*& OutValue ) const
 {
-	return FPropertyHandleObject::GetValue((const UObject*&)OutValue);
+	TObjectPtr<UObject> Temp;
+	FPropertyAccess::Result Result = FPropertyHandleObject::GetValue(Temp);
+	OutValue = Temp.Get();
+	return Result;
 }
 
 FPropertyAccess::Result FPropertyHandleObject::GetValue( const UObject*& OutValue ) const
+{
+	TObjectPtr<const UObject> Temp;
+	FPropertyAccess::Result Result = FPropertyHandleObject::GetValue(Temp);
+	OutValue = Temp.Get();
+	return Result;
+}
+
+FPropertyAccess::Result FPropertyHandleObject::GetValue(TObjectPtr<UObject>& OutValue) const
+{
+	return FPropertyHandleObject::GetValue((TObjectPtr<const UObject>&)OutValue);
+}
+
+FPropertyAccess::Result FPropertyHandleObject::GetValue(TObjectPtr<const UObject>& OutValue) const
 {
 	void* PropValue = nullptr;
 	FPropertyAccess::Result Res = Implementation->GetValueData( PropValue );
@@ -3797,7 +3815,7 @@ FPropertyAccess::Result FPropertyHandleObject::GetValue( const UObject*& OutValu
 
 		if (Property->IsA(FObjectPropertyBase::StaticClass()))
 		{
-			OutValue = Implementation->GetObjectPropertyValue(PropValue);
+			OutValue = Implementation->GetObjectPtrPropertyValue(PropValue);
 		}
 		else if (Property->IsA(FInterfaceProperty::StaticClass()))
 		{
@@ -3820,10 +3838,20 @@ FPropertyAccess::Result FPropertyHandleObject::GetValue( const UObject*& OutValu
 
 FPropertyAccess::Result FPropertyHandleObject::SetValue( UObject* const& NewValue, EPropertyValueSetFlags::Type Flags )
 {
-	return FPropertyHandleObject::SetValue((const UObject*)NewValue);
+	return FPropertyHandleObject::SetValue((TObjectPtr<const UObject>&)NewValue);
 }
 
 FPropertyAccess::Result FPropertyHandleObject::SetValue( const UObject* const& NewValue, EPropertyValueSetFlags::Type Flags )
+{
+	return FPropertyHandleObject::SetValue((TObjectPtr<const UObject>&)NewValue);
+}
+
+FPropertyAccess::Result FPropertyHandleObject::SetValue( TObjectPtr<UObject> const& NewValue, EPropertyValueSetFlags::Type Flags )
+{
+	return FPropertyHandleObject::SetValue((TObjectPtr<const UObject>&)NewValue);
+}
+
+FPropertyAccess::Result FPropertyHandleObject::SetValue( TObjectPtr<const UObject> const& NewValue, EPropertyValueSetFlags::Type Flags )
 {
 	const TSharedPtr<FPropertyNode>& PropertyNode = Implementation->GetPropertyNode();
 
diff --git a/Engine/Source/Editor/PropertyEditor/Private/PropertyHandleImpl.h b/Engine/Source/Editor/PropertyEditor/Private/PropertyHandleImpl.h
--- a/Engine/Source/Editor/PropertyEditor/Private/PropertyHandleImpl.h
+++ b/Engine/Source/Editor/PropertyEditor/Private/PropertyHandleImpl.h
@@ -108,6 +108,19 @@ public:
 		return CastFieldChecked<FObjectPropertyBase>(PropertyNodePin->GetProperty())->GetObjectPropertyValue(Address);
 	}
 
+	/**
+	 * Given an address, get the actual TObjectPtr value out
+	 *
+	 * @param Address	The location of the property value
+	 * @return The object property value
+	 */
+	TObjectPtr<UObject> GetObjectPtrPropertyValue(void const* Address) const
+	{
+		TSharedPtr<FPropertyNode> PropertyNodePin = PropertyNode.Pin();
+		check(PropertyNodePin.IsValid());
+		return CastFieldChecked<FObjectPropertyBase>(PropertyNodePin->GetProperty())->GetObjectPtrPropertyValue(Address);
+	}
+
 
 	/**
 	 * The core functionality for setting values on a property
@@ -393,6 +406,8 @@ public:
 	DECLARE_PROPERTY_ACCESSOR( FRotator )
 	DECLARE_PROPERTY_ACCESSOR( UObject* )
 	DECLARE_PROPERTY_ACCESSOR( const UObject* )
+	DECLARE_PROPERTY_ACCESSOR( TObjectPtr<UObject> )
+	DECLARE_PROPERTY_ACCESSOR( TObjectPtr<const UObject> )
 	DECLARE_PROPERTY_ACCESSOR( FAssetData )
 	DECLARE_PROPERTY_ACCESSOR( FProperty* )
 	DECLARE_PROPERTY_ACCESSOR( const FProperty* )
@@ -587,10 +602,14 @@ class FPropertyHandleObject : public FPropertyHandleBase
 public:
 	FPropertyHandleObject( TSharedRef<FPropertyNode> PropertyNode, FNotifyHook* NotifyHook, TSharedPtr<IPropertyUtilities> PropertyUtilities );
 	static bool Supports( TSharedRef<FPropertyNode> PropertyNode );
+	virtual FPropertyAccess::Result GetValue(TObjectPtr<UObject>& OutValue) const override;
+	virtual FPropertyAccess::Result GetValue(TObjectPtr<const UObject>& OutValue) const override;
 	virtual FPropertyAccess::Result GetValue( UObject*& OutValue ) const override;
 	virtual FPropertyAccess::Result SetValue( UObject* const& InValue, EPropertyValueSetFlags::Type Flags = EPropertyValueSetFlags::DefaultFlags ) override;
 	virtual FPropertyAccess::Result GetValue( const UObject*& OutValue ) const override;
 	virtual FPropertyAccess::Result SetValue( const UObject* const& InValue, EPropertyValueSetFlags::Type Flags = EPropertyValueSetFlags::DefaultFlags ) override;
+	virtual FPropertyAccess::Result SetValue( TObjectPtr<UObject> const& InValue, EPropertyValueSetFlags::Type Flags = EPropertyValueSetFlags::DefaultFlags) override;
+	virtual FPropertyAccess::Result SetValue( TObjectPtr<const UObject> const& InValue, EPropertyValueSetFlags::Type Flags = EPropertyValueSetFlags::DefaultFlags) override;
 	virtual FPropertyAccess::Result GetValue( FAssetData& OutValue ) const override;
 	virtual FPropertyAccess::Result SetValue( const FAssetData& InValue, EPropertyValueSetFlags::Type Flags = EPropertyValueSetFlags::DefaultFlags ) override;
 	virtual FPropertyAccess::Result SetValueFromFormattedString(const FString& InValue, EPropertyValueSetFlags::Type Flags = EPropertyValueSetFlags::DefaultFlags) override;
diff --git a/Engine/Source/Editor/PropertyEditor/Private/UserInterface/PropertyEditor/SPropertyEditorAsset.cpp b/Engine/Source/Editor/PropertyEditor/Private/UserInterface/PropertyEditor/SPropertyEditorAsset.cpp
--- a/Engine/Source/Editor/PropertyEditor/Private/UserInterface/PropertyEditor/SPropertyEditorAsset.cpp
+++ b/Engine/Source/Editor/PropertyEditor/Private/UserInterface/PropertyEditor/SPropertyEditorAsset.cpp
@@ -977,88 +977,83 @@ FPropertyAccess::Result SPropertyEditorAsset::GetValue( FObjectOrAssetData& OutV
 
 	if( PropertyEditor.IsValid() && PropertyEditor->GetPropertyHandle()->IsValidHandle() )
 	{
-		UObject* Object = nullptr;
-		Result = PropertyEditor->GetPropertyHandle()->GetValue(Object);
+		TObjectPtr<UObject> ObjectPtr;
+		Result = PropertyEditor->GetPropertyHandle()->GetValue(ObjectPtr);
 
-		if (Object == nullptr)
+		if (!ObjectPtr.IsResolved())
 		{
 			// Check to see if it's pointing to an unloaded object
-			FString CurrentObjectPath;
-			PropertyEditor->GetPropertyHandle()->GetValueAsFormattedString( CurrentObjectPath );
+			FSoftObjectPath SoftObjectPath = FSoftObjectPath(ObjectPtr);
 
-			if (CurrentObjectPath.Len() > 0 && CurrentObjectPath != TEXT("None"))
+			if (SoftObjectPath.IsAsset())
 			{
-				FSoftObjectPath SoftObjectPath = FSoftObjectPath(CurrentObjectPath);
-
-				if (SoftObjectPath.IsAsset())
+				if (!CachedAssetData.IsValid() || CachedAssetData.GetSoftObjectPath() != SoftObjectPath)
 				{
-					if (!CachedAssetData.IsValid() || CachedAssetData.GetObjectPathString() != CurrentObjectPath)
-					{
-						static FName AssetRegistryName("AssetRegistry");
-
-						FAssetRegistryModule& AssetRegistryModule = FModuleManager::Get().LoadModuleChecked<FAssetRegistryModule>(AssetRegistryName);
-						CachedAssetData = AssetRegistryModule.Get().GetAssetByObjectPath(FSoftObjectPath(CurrentObjectPath));
-					}
+					static FName AssetRegistryName("AssetRegistry");
 
-					Result = FPropertyAccess::Success;
-					OutValue = FObjectOrAssetData(CachedAssetData);
+					FAssetRegistryModule& AssetRegistryModule = FModuleManager::Get().LoadModuleChecked<FAssetRegistryModule>(AssetRegistryName);
+					CachedAssetData = AssetRegistryModule.Get().GetAssetByObjectPath(SoftObjectPath);
 				}
-				else
-				{
-					// This is an actor or other subobject reference
-					if (CachedAssetData.IsValid())
-					{
-						CachedAssetData = FAssetData();
-					}
 
-					Result = FPropertyAccess::Success;
-					OutValue = FObjectOrAssetData(SoftObjectPath);
+				Result = FPropertyAccess::Success;
+				OutValue = FObjectOrAssetData(CachedAssetData);
+			}
+			else
+			{
+				// This is an actor or other subobject reference
+				if (CachedAssetData.IsValid())
+				{
+					CachedAssetData = FAssetData();
 				}
 
-				return Result;
+				Result = FPropertyAccess::Success;
+				OutValue = FObjectOrAssetData(SoftObjectPath);
 			}
+
+			return Result;
 		}
 
 #if !UE_BUILD_SHIPPING
-		if (Object && !Object->IsValidLowLevel())
+		if (ObjectPtr.IsResolved() && ObjectPtr && !ObjectPtr->IsValidLowLevel())
 		{
 			const FProperty* Property = PropertyEditor->GetProperty();
 			UE_LOG(LogPropertyNode, Fatal, TEXT("Property \"%s\" (%s) contains invalid data."), *Property->GetName(), *Property->GetCPPType());
 		}
 #endif
 
-		OutValue = FObjectOrAssetData( Object );
+		OutValue = FObjectOrAssetData( ObjectPtr );
 	}
 	else
 	{
 		FSoftObjectPath SoftObjectPath;
-		UObject* Object = nullptr;
+		TObjectPtr<UObject> ObjectPtr;
 		if (PropertyHandle.IsValid())
 		{
-			Result = PropertyHandle->GetValue(Object);
+			Result = PropertyHandle->GetValue(ObjectPtr);
+			SoftObjectPath = FSoftObjectPath(ObjectPtr);
 		}
 		else
 		{
 			SoftObjectPath = FSoftObjectPath(ObjectPath.Get());
-			Object = SoftObjectPath.ResolveObject();
+			ObjectPtr = SoftObjectPath.ResolveObject();
 
-			if (Object != nullptr)
+			if (ObjectPtr != nullptr)
 			{
 				Result = FPropertyAccess::Success;
 			}
 		}
 
-		if (Object != nullptr)
+		if (ObjectPtr.IsResolved() && ObjectPtr)
 		{
 #if !UE_BUILD_SHIPPING
-			if (!Object->IsValidLowLevel())
+			if (!ObjectPtr->IsValidLowLevel())
 			{
 				const FProperty* Property = PropertyEditor->GetProperty();
 				UE_LOG(LogPropertyNode, Fatal, TEXT("Property \"%s\" (%s) contains invalid data."), *Property->GetName(), *Property->GetCPPType());
 			}
 #endif
 
-			OutValue = FObjectOrAssetData(Object);
+			OutValue = FObjectOrAssetData(ObjectPtr);
 		}
 		else
 		{
diff --git a/Engine/Source/Editor/PropertyEditor/Private/UserInterface/PropertyEditor/SPropertyEditorAsset.h b/Engine/Source/Editor/PropertyEditor/Private/UserInterface/PropertyEditor/SPropertyEditorAsset.h
--- a/Engine/Source/Editor/PropertyEditor/Private/UserInterface/PropertyEditor/SPropertyEditorAsset.h
+++ b/Engine/Source/Editor/PropertyEditor/Private/UserInterface/PropertyEditor/SPropertyEditorAsset.h
@@ -107,11 +107,20 @@ private:
 		FSoftObjectPath ObjectPath;
 		FAssetData AssetData;
 
-		FObjectOrAssetData( UObject* InObject = nullptr )
-			: Object( InObject )
-			, ObjectPath( Object )
+		FObjectOrAssetData( TObjectPtr<UObject> InObject = nullptr )
 		{
-			AssetData = InObject != nullptr && !InObject->IsA<AActor>() ? FAssetData( InObject ) : FAssetData();
+			// If the TObjectPtr isn't resolved yet, then treat it like a soft object ptr.
+			if (InObject.IsResolved() && InObject != nullptr)
+			{
+				Object = InObject;
+				ObjectPath = Object;
+				AssetData = Object->IsA<AActor>() ? FAssetData(Object) : FAssetData();
+			}
+			else
+			{
+				Object = nullptr;
+				ObjectPath = FSoftObjectPath(InObject);
+			}
 		}
 
 		FObjectOrAssetData( const FSoftObjectPath& InObjectPath )
diff --git a/Engine/Source/Editor/PropertyEditor/Public/PropertyHandle.h b/Engine/Source/Editor/PropertyEditor/Public/PropertyHandle.h
--- a/Engine/Source/Editor/PropertyEditor/Public/PropertyHandle.h
+++ b/Engine/Source/Editor/PropertyEditor/Public/PropertyHandle.h
@@ -294,6 +294,8 @@ public:
 	virtual FPropertyAccess::Result GetValue( FRotator& OutValue ) const = 0;
 	virtual FPropertyAccess::Result GetValue( UObject*& OutValue ) const = 0;
 	virtual FPropertyAccess::Result GetValue( const UObject*& OutValue ) const = 0;
+	virtual FPropertyAccess::Result GetValue( TObjectPtr<UObject>& OutValue ) const = 0;
+	virtual FPropertyAccess::Result GetValue( TObjectPtr<const UObject>& OutValue ) const = 0;
 	virtual FPropertyAccess::Result GetValue( FAssetData& OutValue ) const = 0;
 	virtual FPropertyAccess::Result GetValueData( void*& OutAddress ) const = 0;
 	virtual FPropertyAccess::Result GetValue(FProperty*& OutValue) const = 0;
@@ -324,6 +326,8 @@ public:
 	virtual FPropertyAccess::Result SetValue( const FVector4& InValue,  EPropertyValueSetFlags::Type Flags = EPropertyValueSetFlags::DefaultFlags ) = 0;
 	virtual FPropertyAccess::Result SetValue( const FQuat& InValue,  EPropertyValueSetFlags::Type Flags = EPropertyValueSetFlags::DefaultFlags ) = 0;
 	virtual FPropertyAccess::Result SetValue( const FRotator& InValue,  EPropertyValueSetFlags::Type Flags = EPropertyValueSetFlags::DefaultFlags ) = 0;
+	virtual FPropertyAccess::Result SetValue( TObjectPtr<UObject> const& InValue,  EPropertyValueSetFlags::Type Flags = EPropertyValueSetFlags::DefaultFlags ) = 0;
+	virtual FPropertyAccess::Result SetValue( TObjectPtr<const UObject> const& InValue,  EPropertyValueSetFlags::Type Flags = EPropertyValueSetFlags::DefaultFlags ) = 0;
 	virtual FPropertyAccess::Result SetValue( UObject* const& InValue,  EPropertyValueSetFlags::Type Flags = EPropertyValueSetFlags::DefaultFlags ) = 0;
 	virtual FPropertyAccess::Result SetValue( const UObject* const& InValue,  EPropertyValueSetFlags::Type Flags = EPropertyValueSetFlags::DefaultFlags ) = 0;
 	virtual FPropertyAccess::Result SetValue( const FAssetData& InValue,  EPropertyValueSetFlags::Type Flags = EPropertyValueSetFlags::DefaultFlags ) = 0;
diff --git a/Engine/Source/Editor/UnrealEd/Private/Kismet2/ComponentEditorUtils.cpp b/Engine/Source/Editor/UnrealEd/Private/Kismet2/ComponentEditorUtils.cpp
--- a/Engine/Source/Editor/UnrealEd/Private/Kismet2/ComponentEditorUtils.cpp
+++ b/Engine/Source/Editor/UnrealEd/Private/Kismet2/ComponentEditorUtils.cpp
@@ -189,10 +189,15 @@ FProperty* FComponentEditorUtils::GetPropertyForEditableNativeComponent(const UA
 				continue;
 			}
 
-			UObject* Object = ObjectProp->GetObjectPropertyValue(ObjectProp->ContainerPtrToValuePtr<void>(ComponentOuter));
-			if (Object != nullptr && Object->GetFName() == NativeComponent->GetFName())
+			TObjectPtr<UObject> ObjectPtr = ObjectProp->GetObjectPtrPropertyValue(ObjectProp->ContainerPtrToValuePtr<void>(ComponentOuter));
+
+			// Don't bother checking unless it's already resolved, a component member would already be resolved.
+			if (ObjectPtr.IsResolved())
 			{
-				return ObjectProp;
+				if (ObjectPtr != nullptr && ObjectPtr->GetFName() == NativeComponent->GetFName())
+				{
+					return ObjectProp;
+				}
 			}
 		}	
 	
diff --git a/Engine/Source/Runtime/Engine/Private/Actor.cpp b/Engine/Source/Runtime/Engine/Private/Actor.cpp
--- a/Engine/Source/Runtime/Engine/Private/Actor.cpp
+++ b/Engine/Source/Runtime/Engine/Private/Actor.cpp
@@ -1642,11 +1642,19 @@ bool AActor::Modify( bool bAlwaysMarkDirty/*=true*/ )
 			FObjectProperty* ObjProp = *PropertyIt;
 			if (!ObjProp->HasAllPropertyFlags(CPF_NonTransactional))
 			{
-				UActorComponent* ActorComponent = Cast<UActorComponent>(ObjProp->GetObjectPropertyValue(ObjProp->ContainerPtrToValuePtr<void>(this)));
-				if (ActorComponent && ActorComponent->IsCreatedByConstructionScript())
+				TObjectPtr<UObject> ObjectValue = ObjProp->GetObjectPtrPropertyValue(ObjProp->ContainerPtrToValuePtr<void>(this));
+				// @TODO: OBJPTR: UActorComponents can't be lazy loaded, so if this property is unresolved it's impossible that we're dealing
+				// with a case where we *need* resolution to occur.
+				if (ObjectValue.IsResolved())
 				{
-					ObjProp->SetPropertyFlags(CPF_NonTransactional);
-					TemporarilyNonTransactionalProperties.Add(ObjProp);
+					if (UActorComponent* ActorComponent = Cast<UActorComponent>(ObjectValue))
+					{
+						if (ActorComponent->IsCreatedByConstructionScript())
+						{
+							ObjProp->SetPropertyFlags(CPF_NonTransactional);
+							TemporarilyNonTransactionalProperties.Add(ObjProp);
+						}
+					}
 				}
 			}
 		}
