Fixed a bug where duplicating or deleting array elements in the property editor could cause a crash

diff --git a/Engine/Source/Editor/PropertyEditor/Private/PropertyHandleImpl.cpp b/Engine/Source/Editor/PropertyEditor/Private/PropertyHandleImpl.cpp
--- a/Engine/Source/Editor/PropertyEditor/Private/PropertyHandleImpl.cpp
+++ b/Engine/Source/Editor/PropertyEditor/Private/PropertyHandleImpl.cpp
@@ -1277,6 +1277,8 @@ void FPropertyValueImpl::ClearChildren()
 					PropertyNodePin->FixPropertiesInEvent(ChangeEvent);
 					PropertyUtilities.Pin()->NotifyFinishedChangingProperties(ChangeEvent);
 				}
+				
+				PropertyNodePin->RebuildChildren();
 			}
 		}
 	}
@@ -2008,13 +2010,16 @@ void FPropertyValueImpl::DuplicateChild( TSharedPtr<FPropertyNode> ChildNodeToDu
 		ChangeEvent.SetArrayIndexPerObject(ArrayIndicesPerObject);
 		ChangeEvent.SetInstancesChanged(MoveTemp(AllAffectedInstances));
 
-		PropertyNode.Pin()->NotifyPostChange(ChangeEvent, NotifyHook);
+		TSharedPtr<FPropertyNode> PropertyNodePin = PropertyNode.Pin();
+		PropertyNodePin->NotifyPostChange(ChangeEvent, NotifyHook);
 
 		if (PropertyUtilities.IsValid())
 		{
 			ChildNodePtr->FixPropertiesInEvent(ChangeEvent);
 			PropertyUtilities.Pin()->NotifyFinishedChangingProperties(ChangeEvent);
 		}
+
+		PropertyNodePin->RebuildChildren();
 	}
 }
 
