[Backout] - CL22508713
Original CL Desc
-----------------------------------------------------------------
Editor - When reading ObjectPropertyNodes, treat them as TObjectPtrs so that we can keep them unloaded in a latent loaded situations.

diff --git a/Engine/Source/Editor/PropertyEditor/Private/ObjectPropertyNode.cpp b/Engine/Source/Editor/PropertyEditor/Private/ObjectPropertyNode.cpp
--- a/Engine/Source/Editor/PropertyEditor/Private/ObjectPropertyNode.cpp
+++ b/Engine/Source/Editor/PropertyEditor/Private/ObjectPropertyNode.cpp
@@ -359,13 +359,13 @@ bool FObjectPropertyNode::GetReadAddressUncached(const FPropertyNode& InNode,
 
 					// If BaseObj is NON-NULL, all other objects' properties should also be non-NULL.
 					// If BaseObj is NULL, all other objects' properties should also be NULL.
-					const TObjectPtr<UObject> BaseObj = CastField<FObjectPropertyBase>(InItemProperty)->GetObjectPtrPropertyValue(BaseAddr);
+					UObject* BaseObj = CastField<FObjectPropertyBase>(InItemProperty)->GetObjectPropertyValue(BaseAddr);
 
 					for (int32 ObjIndex = 1; ObjIndex < GetNumObjects(); ObjIndex++)
 					{
-						const TObjectPtr<UObject> CurObj = CastField<FObjectPropertyBase>(InItemProperty)->GetObjectPtrPropertyValue(InNode.GetValueBaseAddressFromObject(GetUObject(ObjIndex)));
-						if (   ( !BaseObj && CurObj)			// BaseObj is NULL, but this InItemProperty is non-NULL!
-							|| ( BaseObj && !CurObj) )			// BaseObj is non-NULL, but this InItemProperty is NULL!
+						UObject* CurObj = CastField<FObjectPropertyBase>(InItemProperty)->GetObjectPropertyValue(InNode.GetValueBaseAddressFromObject(GetUObject(ObjIndex)));
+						if (   ( !BaseObj && CurObj )			// BaseObj is NULL, but this InItemProperty is non-NULL!
+							|| ( BaseObj && !CurObj ) )			// BaseObj is non-NULL, but this InItemProperty is NULL!
 						{
 							bAllTheSame = false;
 							break;
@@ -376,12 +376,13 @@ bool FObjectPropertyNode::GetReadAddressUncached(const FPropertyNode& InNode,
 		}
 	}
 
-	// Write addresses to the output.
+		// Write addresses to the output.
 	if (OutAddresses != nullptr)
-	{
-		for (int32 ObjIndex = 0; ObjIndex < GetNumObjects(); ++ObjIndex)
 		{
-			if (const UObject* TempObject = GetUObject(ObjIndex))
+		for (int32 ObjIndex = 0; ObjIndex < GetNumObjects(); ++ObjIndex)
+		{	
+			const UObject* TempObject = GetUObject(ObjIndex);
+			if (TempObject)
 			{
 				OutAddresses->Add(TempObject, InNode.GetValueBaseAddressFromObject(TempObject));
 			}
@@ -447,7 +448,7 @@ uint8* FObjectPropertyNode::GetValueBaseAddress(uint8* StartAddress, bool bIsSpa
 void FObjectPropertyNode::InitBeforeNodeFlags()
 {
 	StoredProperty = Property;
-	Property = nullptr;
+	Property = NULL;
 
 	Finalize();
 }
@@ -465,7 +466,7 @@ static TSharedPtr<FPropertyNode> FindChildCategory( TSharedPtr<FPropertyNode> Pa
 		check( ChildNode.IsValid() );
 
 		// Is this a category node?
-		const FCategoryPropertyNode* ChildCategoryNode = ChildNode->AsCategoryNode();
+		FCategoryPropertyNode* ChildCategoryNode = ChildNode->AsCategoryNode();
 		if( ChildCategoryNode != nullptr && ChildCategoryNode->GetCategoryName() == CategoryName )
 		{
 			return ChildNode;
diff --git a/Engine/Source/Editor/PropertyEditor/Private/PropertyHandleImpl.cpp b/Engine/Source/Editor/PropertyEditor/Private/PropertyHandleImpl.cpp
--- a/Engine/Source/Editor/PropertyEditor/Private/PropertyHandleImpl.cpp
+++ b/Engine/Source/Editor/PropertyEditor/Private/PropertyHandleImpl.cpp
@@ -2074,8 +2074,6 @@ IMPLEMENT_PROPERTY_ACCESSOR( FQuat )
 IMPLEMENT_PROPERTY_ACCESSOR( FRotator )
 IMPLEMENT_PROPERTY_ACCESSOR( UObject* )
 IMPLEMENT_PROPERTY_ACCESSOR( const UObject* )
-IMPLEMENT_PROPERTY_ACCESSOR( TObjectPtr<UObject> )
-IMPLEMENT_PROPERTY_ACCESSOR( TObjectPtr<const UObject> )
 IMPLEMENT_PROPERTY_ACCESSOR( FAssetData )
 IMPLEMENT_PROPERTY_ACCESSOR( FProperty* )
 IMPLEMENT_PROPERTY_ACCESSOR( const FProperty* )
@@ -3790,26 +3788,10 @@ bool FPropertyHandleObject::Supports( TSharedRef<FPropertyNode> PropertyNode )
 
 FPropertyAccess::Result FPropertyHandleObject::GetValue( UObject*& OutValue ) const
 {
-	TObjectPtr<UObject> Temp;
-	FPropertyAccess::Result Result = FPropertyHandleObject::GetValue(Temp);
-	OutValue = Temp.Get();
-	return Result;
+	return FPropertyHandleObject::GetValue((const UObject*&)OutValue);
 }
 
 FPropertyAccess::Result FPropertyHandleObject::GetValue( const UObject*& OutValue ) const
-{
-	TObjectPtr<const UObject> Temp;
-	FPropertyAccess::Result Result = FPropertyHandleObject::GetValue(Temp);
-	OutValue = Temp.Get();
-	return Result;
-}
-
-FPropertyAccess::Result FPropertyHandleObject::GetValue(TObjectPtr<UObject>& OutValue) const
-{
-	return FPropertyHandleObject::GetValue((TObjectPtr<const UObject>&)OutValue);
-}
-
-FPropertyAccess::Result FPropertyHandleObject::GetValue(TObjectPtr<const UObject>& OutValue) const
 {
 	void* PropValue = nullptr;
 	FPropertyAccess::Result Res = Implementation->GetValueData( PropValue );
@@ -3820,7 +3802,7 @@ FPropertyAccess::Result FPropertyHandleObject::GetValue(TObjectPtr<const UObject
 
 		if (Property->IsA(FObjectPropertyBase::StaticClass()))
 		{
-			OutValue = Implementation->GetObjectPtrPropertyValue(PropValue);
+			OutValue = Implementation->GetObjectPropertyValue(PropValue);
 		}
 		else if (Property->IsA(FInterfaceProperty::StaticClass()))
 		{
@@ -3843,20 +3825,10 @@ FPropertyAccess::Result FPropertyHandleObject::GetValue(TObjectPtr<const UObject
 
 FPropertyAccess::Result FPropertyHandleObject::SetValue( UObject* const& NewValue, EPropertyValueSetFlags::Type Flags )
 {
-	return FPropertyHandleObject::SetValue((TObjectPtr<const UObject>&)NewValue);
+	return FPropertyHandleObject::SetValue((const UObject*)NewValue);
 }
 
 FPropertyAccess::Result FPropertyHandleObject::SetValue( const UObject* const& NewValue, EPropertyValueSetFlags::Type Flags )
-{
-	return FPropertyHandleObject::SetValue((TObjectPtr<const UObject>&)NewValue);
-}
-
-FPropertyAccess::Result FPropertyHandleObject::SetValue( TObjectPtr<UObject> const& NewValue, EPropertyValueSetFlags::Type Flags )
-{
-	return FPropertyHandleObject::SetValue((TObjectPtr<const UObject>&)NewValue);
-}
-
-FPropertyAccess::Result FPropertyHandleObject::SetValue( TObjectPtr<const UObject> const& NewValue, EPropertyValueSetFlags::Type Flags )
 {
 	const TSharedPtr<FPropertyNode>& PropertyNode = Implementation->GetPropertyNode();
 
diff --git a/Engine/Source/Editor/PropertyEditor/Private/PropertyHandleImpl.h b/Engine/Source/Editor/PropertyEditor/Private/PropertyHandleImpl.h
--- a/Engine/Source/Editor/PropertyEditor/Private/PropertyHandleImpl.h
+++ b/Engine/Source/Editor/PropertyEditor/Private/PropertyHandleImpl.h
@@ -108,19 +108,6 @@ public:
 		return CastFieldChecked<FObjectPropertyBase>(PropertyNodePin->GetProperty())->GetObjectPropertyValue(Address);
 	}
 
-	/**
-	 * Given an address, get the actual TObjectPtr value out
-	 *
-	 * @param Address	The location of the property value
-	 * @return The object property value
-	 */
-	TObjectPtr<UObject> GetObjectPtrPropertyValue(void const* Address) const
-	{
-		TSharedPtr<FPropertyNode> PropertyNodePin = PropertyNode.Pin();
-		check(PropertyNodePin.IsValid());
-		return CastFieldChecked<FObjectPropertyBase>(PropertyNodePin->GetProperty())->GetObjectPtrPropertyValue(Address);
-	}
-
 
 	/**
 	 * The core functionality for setting values on a property
@@ -406,8 +393,6 @@ public:
 	DECLARE_PROPERTY_ACCESSOR( FRotator )
 	DECLARE_PROPERTY_ACCESSOR( UObject* )
 	DECLARE_PROPERTY_ACCESSOR( const UObject* )
-	DECLARE_PROPERTY_ACCESSOR( TObjectPtr<UObject> )
-	DECLARE_PROPERTY_ACCESSOR( TObjectPtr<const UObject> )
 	DECLARE_PROPERTY_ACCESSOR( FAssetData )
 	DECLARE_PROPERTY_ACCESSOR( FProperty* )
 	DECLARE_PROPERTY_ACCESSOR( const FProperty* )
@@ -602,14 +587,10 @@ class FPropertyHandleObject : public FPropertyHandleBase
 public:
 	FPropertyHandleObject( TSharedRef<FPropertyNode> PropertyNode, FNotifyHook* NotifyHook, TSharedPtr<IPropertyUtilities> PropertyUtilities );
 	static bool Supports( TSharedRef<FPropertyNode> PropertyNode );
-	virtual FPropertyAccess::Result GetValue(TObjectPtr<UObject>& OutValue) const override;
-	virtual FPropertyAccess::Result GetValue(TObjectPtr<const UObject>& OutValue) const override;
 	virtual FPropertyAccess::Result GetValue( UObject*& OutValue ) const override;
 	virtual FPropertyAccess::Result SetValue( UObject* const& InValue, EPropertyValueSetFlags::Type Flags = EPropertyValueSetFlags::DefaultFlags ) override;
 	virtual FPropertyAccess::Result GetValue( const UObject*& OutValue ) const override;
 	virtual FPropertyAccess::Result SetValue( const UObject* const& InValue, EPropertyValueSetFlags::Type Flags = EPropertyValueSetFlags::DefaultFlags ) override;
-	virtual FPropertyAccess::Result SetValue( TObjectPtr<UObject> const& InValue, EPropertyValueSetFlags::Type Flags = EPropertyValueSetFlags::DefaultFlags) override;
-	virtual FPropertyAccess::Result SetValue( TObjectPtr<const UObject> const& InValue, EPropertyValueSetFlags::Type Flags = EPropertyValueSetFlags::DefaultFlags) override;
 	virtual FPropertyAccess::Result GetValue( FAssetData& OutValue ) const override;
 	virtual FPropertyAccess::Result SetValue( const FAssetData& InValue, EPropertyValueSetFlags::Type Flags = EPropertyValueSetFlags::DefaultFlags ) override;
 	virtual FPropertyAccess::Result SetValueFromFormattedString(const FString& InValue, EPropertyValueSetFlags::Type Flags = EPropertyValueSetFlags::DefaultFlags) override;
diff --git a/Engine/Source/Editor/PropertyEditor/Private/UserInterface/PropertyEditor/SPropertyEditorAsset.cpp b/Engine/Source/Editor/PropertyEditor/Private/UserInterface/PropertyEditor/SPropertyEditorAsset.cpp
--- a/Engine/Source/Editor/PropertyEditor/Private/UserInterface/PropertyEditor/SPropertyEditorAsset.cpp
+++ b/Engine/Source/Editor/PropertyEditor/Private/UserInterface/PropertyEditor/SPropertyEditorAsset.cpp
@@ -977,83 +977,88 @@ FPropertyAccess::Result SPropertyEditorAsset::GetValue( FObjectOrAssetData& OutV
 
 	if( PropertyEditor.IsValid() && PropertyEditor->GetPropertyHandle()->IsValidHandle() )
 	{
-		TObjectPtr<UObject> ObjectPtr;
-		Result = PropertyEditor->GetPropertyHandle()->GetValue(ObjectPtr);
+		UObject* Object = nullptr;
+		Result = PropertyEditor->GetPropertyHandle()->GetValue(Object);
 
-		if (!ObjectPtr.IsResolved())
+		if (Object == nullptr)
 		{
 			// Check to see if it's pointing to an unloaded object
-			FSoftObjectPath SoftObjectPath = FSoftObjectPath(ObjectPtr);
+			FString CurrentObjectPath;
+			PropertyEditor->GetPropertyHandle()->GetValueAsFormattedString( CurrentObjectPath );
 
-			if (SoftObjectPath.IsAsset())
+			if (CurrentObjectPath.Len() > 0 && CurrentObjectPath != TEXT("None"))
 			{
-				if (!CachedAssetData.IsValid() || CachedAssetData.GetSoftObjectPath() != SoftObjectPath)
+				FSoftObjectPath SoftObjectPath = FSoftObjectPath(CurrentObjectPath);
+
+				if (SoftObjectPath.IsAsset())
 				{
-					static FName AssetRegistryName("AssetRegistry");
+					if (!CachedAssetData.IsValid() || CachedAssetData.GetObjectPathString() != CurrentObjectPath)
+					{
+						static FName AssetRegistryName("AssetRegistry");
 
-					FAssetRegistryModule& AssetRegistryModule = FModuleManager::Get().LoadModuleChecked<FAssetRegistryModule>(AssetRegistryName);
-					CachedAssetData = AssetRegistryModule.Get().GetAssetByObjectPath(SoftObjectPath);
-				}
+						FAssetRegistryModule& AssetRegistryModule = FModuleManager::Get().LoadModuleChecked<FAssetRegistryModule>(AssetRegistryName);
+						CachedAssetData = AssetRegistryModule.Get().GetAssetByObjectPath(FSoftObjectPath(CurrentObjectPath));
+					}
 
-				Result = FPropertyAccess::Success;
-				OutValue = FObjectOrAssetData(CachedAssetData);
-			}
-			else
-			{
-				// This is an actor or other subobject reference
-				if (CachedAssetData.IsValid())
+					Result = FPropertyAccess::Success;
+					OutValue = FObjectOrAssetData(CachedAssetData);
+				}
+				else
 				{
-					CachedAssetData = FAssetData();
+					// This is an actor or other subobject reference
+					if (CachedAssetData.IsValid())
+					{
+						CachedAssetData = FAssetData();
+					}
+
+					Result = FPropertyAccess::Success;
+					OutValue = FObjectOrAssetData(SoftObjectPath);
 				}
 
-				Result = FPropertyAccess::Success;
-				OutValue = FObjectOrAssetData(SoftObjectPath);
+				return Result;
 			}
-
-			return Result;
 		}
 
 #if !UE_BUILD_SHIPPING
-		if (ObjectPtr.IsResolved() && ObjectPtr && !ObjectPtr->IsValidLowLevel())
+		if (Object && !Object->IsValidLowLevel())
 		{
 			const FProperty* Property = PropertyEditor->GetProperty();
 			UE_LOG(LogPropertyNode, Fatal, TEXT("Property \"%s\" (%s) contains invalid data."), *Property->GetName(), *Property->GetCPPType());
 		}
 #endif
 
-		OutValue = FObjectOrAssetData( ObjectPtr );
+		OutValue = FObjectOrAssetData( Object );
 	}
 	else
 	{
 		FSoftObjectPath SoftObjectPath;
-		TObjectPtr<UObject> ObjectPtr;
+		UObject* Object = nullptr;
 		if (PropertyHandle.IsValid())
 		{
-			Result = PropertyHandle->GetValue(ObjectPtr);
-			SoftObjectPath = FSoftObjectPath(ObjectPtr);
+			Result = PropertyHandle->GetValue(Object);
 		}
 		else
 		{
 			SoftObjectPath = FSoftObjectPath(ObjectPath.Get());
-			ObjectPtr = SoftObjectPath.ResolveObject();
+			Object = SoftObjectPath.ResolveObject();
 
-			if (ObjectPtr != nullptr)
+			if (Object != nullptr)
 			{
 				Result = FPropertyAccess::Success;
 			}
 		}
 
-		if (ObjectPtr.IsResolved() && ObjectPtr)
+		if (Object != nullptr)
 		{
 #if !UE_BUILD_SHIPPING
-			if (!ObjectPtr->IsValidLowLevel())
+			if (!Object->IsValidLowLevel())
 			{
 				const FProperty* Property = PropertyEditor->GetProperty();
 				UE_LOG(LogPropertyNode, Fatal, TEXT("Property \"%s\" (%s) contains invalid data."), *Property->GetName(), *Property->GetCPPType());
 			}
 #endif
 
-			OutValue = FObjectOrAssetData(ObjectPtr);
+			OutValue = FObjectOrAssetData(Object);
 		}
 		else
 		{
diff --git a/Engine/Source/Editor/PropertyEditor/Private/UserInterface/PropertyEditor/SPropertyEditorAsset.h b/Engine/Source/Editor/PropertyEditor/Private/UserInterface/PropertyEditor/SPropertyEditorAsset.h
--- a/Engine/Source/Editor/PropertyEditor/Private/UserInterface/PropertyEditor/SPropertyEditorAsset.h
+++ b/Engine/Source/Editor/PropertyEditor/Private/UserInterface/PropertyEditor/SPropertyEditorAsset.h
@@ -107,20 +107,11 @@ private:
 		FSoftObjectPath ObjectPath;
 		FAssetData AssetData;
 
-		FObjectOrAssetData( TObjectPtr<UObject> InObject = nullptr )
+		FObjectOrAssetData( UObject* InObject = nullptr )
+			: Object( InObject )
+			, ObjectPath( Object )
 		{
-			// If the TObjectPtr isn't resolved yet, then treat it like a soft object ptr.
-			if (InObject.IsResolved() && InObject != nullptr)
-			{
-				Object = InObject;
-				ObjectPath = Object;
-				AssetData = Object->IsA<AActor>() ? FAssetData(Object) : FAssetData();
-			}
-			else
-			{
-				Object = nullptr;
-				ObjectPath = FSoftObjectPath(InObject);
-			}
+			AssetData = InObject != nullptr && !InObject->IsA<AActor>() ? FAssetData( InObject ) : FAssetData();
 		}
 
 		FObjectOrAssetData( const FSoftObjectPath& InObjectPath )
diff --git a/Engine/Source/Editor/PropertyEditor/Public/PropertyHandle.h b/Engine/Source/Editor/PropertyEditor/Public/PropertyHandle.h
--- a/Engine/Source/Editor/PropertyEditor/Public/PropertyHandle.h
+++ b/Engine/Source/Editor/PropertyEditor/Public/PropertyHandle.h
@@ -294,8 +294,6 @@ public:
 	virtual FPropertyAccess::Result GetValue( FRotator& OutValue ) const = 0;
 	virtual FPropertyAccess::Result GetValue( UObject*& OutValue ) const = 0;
 	virtual FPropertyAccess::Result GetValue( const UObject*& OutValue ) const = 0;
-	virtual FPropertyAccess::Result GetValue( TObjectPtr<UObject>& OutValue ) const = 0;
-	virtual FPropertyAccess::Result GetValue( TObjectPtr<const UObject>& OutValue ) const = 0;
 	virtual FPropertyAccess::Result GetValue( FAssetData& OutValue ) const = 0;
 	virtual FPropertyAccess::Result GetValueData( void*& OutAddress ) const = 0;
 	virtual FPropertyAccess::Result GetValue(FProperty*& OutValue) const = 0;
@@ -326,8 +324,6 @@ public:
 	virtual FPropertyAccess::Result SetValue( const FVector4& InValue,  EPropertyValueSetFlags::Type Flags = EPropertyValueSetFlags::DefaultFlags ) = 0;
 	virtual FPropertyAccess::Result SetValue( const FQuat& InValue,  EPropertyValueSetFlags::Type Flags = EPropertyValueSetFlags::DefaultFlags ) = 0;
 	virtual FPropertyAccess::Result SetValue( const FRotator& InValue,  EPropertyValueSetFlags::Type Flags = EPropertyValueSetFlags::DefaultFlags ) = 0;
-	virtual FPropertyAccess::Result SetValue( TObjectPtr<UObject> const& InValue,  EPropertyValueSetFlags::Type Flags = EPropertyValueSetFlags::DefaultFlags ) = 0;
-	virtual FPropertyAccess::Result SetValue( TObjectPtr<const UObject> const& InValue,  EPropertyValueSetFlags::Type Flags = EPropertyValueSetFlags::DefaultFlags ) = 0;
 	virtual FPropertyAccess::Result SetValue( UObject* const& InValue,  EPropertyValueSetFlags::Type Flags = EPropertyValueSetFlags::DefaultFlags ) = 0;
 	virtual FPropertyAccess::Result SetValue( const UObject* const& InValue,  EPropertyValueSetFlags::Type Flags = EPropertyValueSetFlags::DefaultFlags ) = 0;
 	virtual FPropertyAccess::Result SetValue( const FAssetData& InValue,  EPropertyValueSetFlags::Type Flags = EPropertyValueSetFlags::DefaultFlags ) = 0;
diff --git a/Engine/Source/Editor/UnrealEd/Private/Kismet2/ComponentEditorUtils.cpp b/Engine/Source/Editor/UnrealEd/Private/Kismet2/ComponentEditorUtils.cpp
--- a/Engine/Source/Editor/UnrealEd/Private/Kismet2/ComponentEditorUtils.cpp
+++ b/Engine/Source/Editor/UnrealEd/Private/Kismet2/ComponentEditorUtils.cpp
@@ -189,15 +189,10 @@ FProperty* FComponentEditorUtils::GetPropertyForEditableNativeComponent(const UA
 				continue;
 			}
 
-			TObjectPtr<UObject> ObjectPtr = ObjectProp->GetObjectPtrPropertyValue(ObjectProp->ContainerPtrToValuePtr<void>(ComponentOuter));
-
-			// Don't bother checking unless it's already resolved, a component member would already be resolved.
-			if (ObjectPtr.IsResolved())
+			UObject* Object = ObjectProp->GetObjectPropertyValue(ObjectProp->ContainerPtrToValuePtr<void>(ComponentOuter));
+			if (Object != nullptr && Object->GetFName() == NativeComponent->GetFName())
 			{
-				if (ObjectPtr != nullptr && ObjectPtr->GetFName() == NativeComponent->GetFName())
-				{
-					return ObjectProp;
-				}
+				return ObjectProp;
 			}
 		}	
 	
diff --git a/Engine/Source/Runtime/Engine/Private/Actor.cpp b/Engine/Source/Runtime/Engine/Private/Actor.cpp
--- a/Engine/Source/Runtime/Engine/Private/Actor.cpp
+++ b/Engine/Source/Runtime/Engine/Private/Actor.cpp
@@ -1642,19 +1642,11 @@ bool AActor::Modify( bool bAlwaysMarkDirty/*=true*/ )
 			FObjectProperty* ObjProp = *PropertyIt;
 			if (!ObjProp->HasAllPropertyFlags(CPF_NonTransactional))
 			{
-				TObjectPtr<UObject> ObjectValue = ObjProp->GetObjectPtrPropertyValue(ObjProp->ContainerPtrToValuePtr<void>(this));
-				// @TODO: OBJPTR: UActorComponents can't be lazy loaded, so if this property is unresolved it's impossible that we're dealing
-				// with a case where we *need* resolution to occur.
-				if (ObjectValue.IsResolved())
+				UActorComponent* ActorComponent = Cast<UActorComponent>(ObjProp->GetObjectPropertyValue(ObjProp->ContainerPtrToValuePtr<void>(this)));
+				if (ActorComponent && ActorComponent->IsCreatedByConstructionScript())
 				{
-					if (UActorComponent* ActorComponent = Cast<UActorComponent>(ObjectValue))
-					{
-						if (ActorComponent->IsCreatedByConstructionScript())
-						{
-							ObjProp->SetPropertyFlags(CPF_NonTransactional);
-							TemporarilyNonTransactionalProperties.Add(ObjProp);
-						}
-					}
+					ObjProp->SetPropertyFlags(CPF_NonTransactional);
+					TemporarilyNonTransactionalProperties.Add(ObjProp);
 				}
 			}
 		}
