DetailView: Plugins - Model View Viewmodel shows up for QAGame Project Settings searches regardless of search.

In PropertyNode::FilterNodes, filter the property key node, if present, to prevent this.

diff --git a/Engine/Source/Editor/PropertyEditor/Private/PropertyNode.cpp b/Engine/Source/Editor/PropertyEditor/Private/PropertyNode.cpp
--- a/Engine/Source/Editor/PropertyEditor/Private/PropertyNode.cpp
+++ b/Engine/Source/Editor/PropertyEditor/Private/PropertyNode.cpp
@@ -2296,6 +2296,11 @@ void FPropertyNode::DestroyTree(const bool bInDestroySelf)
  */
 void FPropertyNode::FilterNodes( const TArray<FString>& InFilterStrings, const bool bParentSeenDueToFiltering )
 {
+	if (const TSharedPtr<FPropertyNode>& KeyNode = GetPropertyKeyNode())
+	{
+		KeyNode->FilterNodes(InFilterStrings);						
+	}
+
 	//clear flags first.  Default to hidden
 	SetNodeFlags(EPropertyNodeFlags::IsSeenDueToFiltering | EPropertyNodeFlags::IsSeenDueToChildFiltering | EPropertyNodeFlags::IsParentSeenDueToFiltering, false);
 	SetNodeFlags(EPropertyNodeFlags::IsBeingFiltered, InFilterStrings.Num() > 0 );
