No longer crash if a property with an EditCondition does not have a parent node, eg. in a single property view.

#rnx

diff --git a/Engine/Source/Editor/PropertyEditor/Private/EditConditionContext.cpp b/Engine/Source/Editor/PropertyEditor/Private/EditConditionContext.cpp
--- a/Engine/Source/Editor/PropertyEditor/Private/EditConditionContext.cpp
+++ b/Engine/Source/Editor/PropertyEditor/Private/EditConditionContext.cpp
@@ -103,6 +103,11 @@ const T* FindTypedField(const TSharedPtr<FPropertyNode>& PropertyNode, const FSt
 static const FPropertyNode* GetEditConditionParentNode(const TSharedPtr<FPropertyNode>& PropertyNode)
 {
 	const FPropertyNode* ParentNode = PropertyNode->GetParentNode();
+	if (ParentNode == nullptr)
+	{
+		return nullptr;
+	}
+
 	FFieldVariant PropertyOuter = PropertyNode->GetProperty()->GetOwnerVariant();
 
 	if (PropertyOuter.Get<FArrayProperty>() != nullptr ||
