Fix an issue where prior valid to generate properties were not generated if we only specified full path to the property

diff --git a/Engine/Source/Editor/PropertyEditor/Private/DetailLayoutBuilderImpl.cpp b/Engine/Source/Editor/PropertyEditor/Private/DetailLayoutBuilderImpl.cpp
--- a/Engine/Source/Editor/PropertyEditor/Private/DetailLayoutBuilderImpl.cpp
+++ b/Engine/Source/Editor/PropertyEditor/Private/DetailLayoutBuilderImpl.cpp
@@ -947,7 +947,7 @@ bool FDetailLayoutBuilderImpl::IsPropertyPathAllowed(const FString& InPath) cons
 
 	for (const FString& PropertyName : PropertyGenerationAllowListPaths)
 	{
-		if (PropertyName.StartsWith(InPath))
+		if (InPath.StartsWith(PropertyName) || PropertyName.StartsWith(InPath))
 		{
 			return true;
 		}
