IWYU: Removed unnecessary PropertyNode.h include from PropertyEditor.h

#rnx

diff --git a/Engine/Source/Editor/PropertyEditor/Private/Presentation/PropertyEditor/PropertyEditor.h b/Engine/Source/Editor/PropertyEditor/Private/Presentation/PropertyEditor/PropertyEditor.h
--- a/Engine/Source/Editor/PropertyEditor/Private/Presentation/PropertyEditor/PropertyEditor.h
+++ b/Engine/Source/Editor/PropertyEditor/Private/Presentation/PropertyEditor/PropertyEditor.h
@@ -2,7 +2,6 @@
 #pragma once
 
 #include "CoreMinimal.h"
-#include "PropertyNode.h"
 #include "PropertyHandle.h"
 #include "SceneOutlinerFwd.h"
 
@@ -93,15 +92,15 @@ public:
 	void EditConfigHierarchy();
 
 	/**	In an ideal world we wouldn't expose these */
-	TSharedRef< FPropertyNode> GetPropertyNode() const;
+	TSharedRef<FPropertyNode> GetPropertyNode() const;
 	const FProperty* GetProperty() const;
-	TSharedRef< IPropertyHandle> GetPropertyHandle() const;
+	TSharedRef<IPropertyHandle> GetPropertyHandle() const;
 
 	static void SyncToObjectsInNode( const TWeakPtr<FPropertyNode>& WeakPropertyNode );
 
 	static const FString MultipleValuesDisplayName;
 private:
-	FPropertyEditor( const TSharedRef<FPropertyNode>& InPropertyNode, const TSharedRef< IPropertyUtilities>& InPropertyUtilities );
+	FPropertyEditor( const TSharedRef<FPropertyNode>& InPropertyNode, const TSharedRef<IPropertyUtilities>& InPropertyUtilities );
 
 	void OnUseSelected();
 	void OnAddItem();
