CIS: Fixed non-unity compilation errors due to removed PropertyNode.h include.

diff --git a/Engine/Source/Editor/PropertyEditor/Private/UserInterface/PropertyEditor/SPropertyEditorClass.cpp b/Engine/Source/Editor/PropertyEditor/Private/UserInterface/PropertyEditor/SPropertyEditorClass.cpp
--- a/Engine/Source/Editor/PropertyEditor/Private/UserInterface/PropertyEditor/SPropertyEditorClass.cpp
+++ b/Engine/Source/Editor/PropertyEditor/Private/UserInterface/PropertyEditor/SPropertyEditorClass.cpp
@@ -5,12 +5,12 @@
 #include "Misc/FeedbackContext.h"
 #include "Modules/ModuleManager.h"
 #include "Widgets/Layout/SBox.h"
-
 #include "DragAndDrop/ClassDragDropOp.h"
 #include "DragAndDrop/AssetDragDropOp.h"
 #include "ClassViewerModule.h"
 #include "ClassViewerFilter.h"
 #include "UObject/UObjectIterator.h"
+#include "PropertyNode.h"
 #include "PropertyRestriction.h"
 
 #define LOCTEXT_NAMESPACE "PropertyEditor"
diff --git a/Engine/Source/Editor/PropertyEditor/Private/UserInterface/PropertyEditor/SPropertyEditorColor.cpp b/Engine/Source/Editor/PropertyEditor/Private/UserInterface/PropertyEditor/SPropertyEditorColor.cpp
--- a/Engine/Source/Editor/PropertyEditor/Private/UserInterface/PropertyEditor/SPropertyEditorColor.cpp
+++ b/Engine/Source/Editor/PropertyEditor/Private/UserInterface/PropertyEditor/SPropertyEditorColor.cpp
@@ -5,6 +5,7 @@
 #include "Engine/Engine.h"
 #include "Components/LightComponent.h"
 #include "Editor.h"
+#include "PropertyNode.h"
 #include "Widgets/Colors/SColorBlock.h"
 #include "Materials/MaterialExpressionConstant3Vector.h"
 #include "Widgets/Colors/SColorPicker.h"
diff --git a/Engine/Source/Editor/PropertyEditor/Private/UserInterface/PropertyEditor/SPropertyEditorStruct.cpp b/Engine/Source/Editor/PropertyEditor/Private/UserInterface/PropertyEditor/SPropertyEditorStruct.cpp
--- a/Engine/Source/Editor/PropertyEditor/Private/UserInterface/PropertyEditor/SPropertyEditorStruct.cpp
+++ b/Engine/Source/Editor/PropertyEditor/Private/UserInterface/PropertyEditor/SPropertyEditorStruct.cpp
@@ -4,7 +4,7 @@
 #include "Misc/FeedbackContext.h"
 #include "Modules/ModuleManager.h"
 #include "Widgets/Layout/SBox.h"
-
+#include "PropertyNode.h"
 #include "DragAndDrop/AssetDragDropOp.h"
 #include "StructViewerModule.h"
 #include "StructViewerFilter.h"
