DetailsSectionsTests now correctly uses #if WITH_DEV_AUTOMATION_TESTS.

#rnx

diff --git a/Engine/Source/Editor/PropertyEditor/Private/Tests/DetailsSectionsTests.cpp b/Engine/Source/Editor/PropertyEditor/Private/Tests/DetailsSectionsTests.cpp
--- a/Engine/Source/Editor/PropertyEditor/Private/Tests/DetailsSectionsTests.cpp
+++ b/Engine/Source/Editor/PropertyEditor/Private/Tests/DetailsSectionsTests.cpp
@@ -1,6 +1,6 @@
 // Copyright Epic Games, Inc. All Rights Reserved.
 
-//#if WITH_DEV_AUTOMATION_TESTS
+#if WITH_DEV_AUTOMATION_TESTS
 
 #include "GameFramework/Actor.h"
 #include "Editor/PropertyEditorTestObject.h"
@@ -100,4 +100,4 @@ bool FDetailsViewTests_SectionsInherited::RunTest(const FString& Parameters)
 	return true;
 }
 
-//#endif // WITH_DEV_AUTOMATION_TESTS
\ No newline at end of file
+#endif // WITH_DEV_AUTOMATION_TESTS
\ No newline at end of file
