Add browe to nanite override material button on material list.
Only shows if a nanite material is set.
Waiting on custom icon, and uses duplicate of content browse icon for now.

diff --git a/Engine/Source/Editor/PropertyEditor/Private/MaterialList.cpp b/Engine/Source/Editor/PropertyEditor/Private/MaterialList.cpp
--- a/Engine/Source/Editor/PropertyEditor/Private/MaterialList.cpp
+++ b/Engine/Source/Editor/PropertyEditor/Private/MaterialList.cpp
@@ -204,6 +204,14 @@ TSharedRef<SWidget> FMaterialItemView::CreateValueContent(IDetailLayoutBuilder&
 					.VAlign(VAlign_Center)
 					[
 						SNew(SHorizontalBox)
+						+ SHorizontalBox::Slot()
+						.VAlign(VAlign_Center)
+						.Padding(0.0f, 0.0f, 3.0f, 0.0f)
+						.AutoWidth()
+						[
+							// Add a button to browse to any nanite override material
+							MakeBrowseNaniteOverrideMaterialButton()
+						]
 						+SHorizontalBox::Slot()
 						.VAlign(VAlign_Center)
 						.Padding(0.0f, 0.0f, 3.0f, 0.0f)
@@ -305,6 +313,49 @@ void FMaterialItemView::OnSetObject( const FAssetData& AssetData )
 	ReplaceMaterial( NewMaterial, bReplaceAll );
 }
 
+TSharedRef<SWidget> FMaterialItemView::MakeBrowseNaniteOverrideMaterialButton() const
+{
+	TSharedRef<SWidget> Widget =
+		SNew(SBox)
+		.HAlign(HAlign_Center)
+		.VAlign(VAlign_Center)
+		.WidthOverride(22)
+		.HeightOverride(22)
+		.ToolTipText(LOCTEXT("BrowseToNaniteOverride_Tip", "Browse to the Nanite Override Material in Content Browser"))
+		[
+			SNew(SButton)
+			.ButtonStyle(FAppStyle::Get(), "SimpleButton")
+			.ContentPadding(0)
+			.IsFocusable(false)
+			.OnClicked(FOnClicked::CreateLambda([WeakMaterial = MaterialItem.Material]()
+				{
+					UMaterialInterface* Material = WeakMaterial.Get();
+					UMaterialInterface* NaniteOverrideMaterial = Material != nullptr ? Material->GetNaniteOverride() : nullptr;
+					if (GEditor && NaniteOverrideMaterial != nullptr)
+					{
+						TArray<UObject*> Objects;
+						Objects.Add(NaniteOverrideMaterial);
+						GEditor->SyncBrowserToObjects(Objects);
+					}
+					return FReply::Handled();
+				}))
+			[ 
+				SNew(SImage)
+				.Image(FAppStyle::Get().GetBrush("Icons.BrowseContent")) //todo: UE-168435 Get custom icon for this.
+				.ColorAndOpacity(FSlateColor::UseForeground())
+			]
+		];
+
+	Widget->SetVisibility(TAttribute<EVisibility>::CreateLambda([WeakMaterial = MaterialItem.Material]()
+		{
+			UMaterialInterface* Material = WeakMaterial.Get();
+			return Material != nullptr && Material->GetNaniteOverride() != nullptr ? EVisibility::Visible : EVisibility::Collapsed;
+		}));
+
+	return Widget;
+}
+
+
 bool FMaterialItemView::IsTexturesMenuEnabled() const
 {
 	return MaterialItem.Material.Get() != NULL;
diff --git a/Engine/Source/Editor/PropertyEditor/Public/MaterialList.h b/Engine/Source/Editor/PropertyEditor/Public/MaterialList.h
--- a/Engine/Source/Editor/PropertyEditor/Public/MaterialList.h
+++ b/Engine/Source/Editor/PropertyEditor/Public/MaterialList.h
@@ -217,6 +217,11 @@ private:
 	 */
 	TSharedRef<SWidget> OnGetTexturesMenuForMaterial();
 
+	/**
+	 * Create button for browse to nanite override material
+	 */
+	TSharedRef<SWidget> MakeBrowseNaniteOverrideMaterialButton() const;
+
 	/**
 	 * On Get object path handler
 	 */
