After a details view is created, it is no longer refiltered every frame for N frames where N is the number of FPropertyNodes in the property tree, causing a noticeable slowdown.

When FPropertyNode::EnsureDataIsValid() returns ChildrenRebuilt, it was only marking itself as bChildrenRebuilt = false.

SDetailsViewBase now calls FPropertyNode::MarkChildrenAsRebuilt, which recursively clears the flag from everything.

diff --git a/Engine/Source/Editor/PropertyEditor/Private/PropertyNode.cpp b/Engine/Source/Editor/PropertyEditor/Private/PropertyNode.cpp
--- a/Engine/Source/Editor/PropertyEditor/Private/PropertyNode.cpp
+++ b/Engine/Source/Editor/PropertyEditor/Private/PropertyNode.cpp
@@ -420,12 +420,21 @@ FObjectPropertyNode* FPropertyNode::FindRootObjectItemParent()
 	return TopmostObjectItem;
 }
 
-
 bool FPropertyNode::DoesChildPropertyRequireValidation(FProperty* InChildProp)
 {
 	return InChildProp != nullptr && (CastField<FObjectProperty>(InChildProp) != nullptr || CastField<FStructProperty>(InChildProp) != nullptr);
 }
 
+void FPropertyNode::MarkChildrenAsRebuilt()
+{
+	bChildrenRebuilt = false;
+
+	for (const TSharedPtr<FPropertyNode>& ChildNode : ChildNodes)
+	{
+		ChildNode->MarkChildrenAsRebuilt();
+	}
+}
+
 /** 
  * Used to see if any data has been destroyed from under the property tree.
  */
@@ -437,7 +446,7 @@ EPropertyDataValidationResult FPropertyNode::EnsureDataIsValid()
 	// If we have rebuilt children since last EnsureDataIsValid call let the caller know
 	if (bChildrenRebuilt)
 	{
-		bChildrenRebuilt = false;
+		MarkChildrenAsRebuilt();
 		return EPropertyDataValidationResult::ChildrenRebuilt;
 	}
 
diff --git a/Engine/Source/Editor/PropertyEditor/Private/PropertyNode.h b/Engine/Source/Editor/PropertyEditor/Private/PropertyNode.h
--- a/Engine/Source/Editor/PropertyEditor/Private/PropertyNode.h
+++ b/Engine/Source/Editor/PropertyEditor/Private/PropertyNode.h
@@ -24,19 +24,19 @@ namespace EPropertyNodeFlags
 {
 	typedef uint32 Type;
 
-	const Type	IsSeen							= 1 << 0;		/** true if this node can be seen based on current parent expansion.  Does not take into account clipping*/
-	const Type	IsSeenDueToFiltering			= 1 << 1;		/** true if this node has been accepted by the filter*/
-	const Type	IsSeenDueToChildFiltering		= 1 << 2;		/** true if this node or one of it's children is seen due to filtering.  It will then be forced on as well.*/
-	const Type	IsParentSeenDueToFiltering		= 1 << 3;		/** True if the parent was visible due to filtering*/
-	const Type	IsSeenDueToChildFavorite		= 1 << 4;		/** True if this node is seen to it having a favorite as a child */
+	const Type	IsSeen							= 1 << 0;	/** true if this node can be seen based on current parent expansion.  Does not take into account clipping*/
+	const Type	IsSeenDueToFiltering			= 1 << 1;	/** true if this node has been accepted by the filter*/
+	const Type	IsSeenDueToChildFiltering		= 1 << 2;	/** true if this node or one of it's children is seen due to filtering.  It will then be forced on as well.*/
+	const Type	IsParentSeenDueToFiltering		= 1 << 3;	/** True if the parent was visible due to filtering*/
+	const Type	IsSeenDueToChildFavorite		= 1 << 4;	/** True if this node is seen to it having a favorite as a child */
 	
-	const Type	Expanded						= 1 << 5;		/** true if this node should display its children*/
-	const Type	CanBeExpanded					= 1 << 6;		/** true if this node is able to be expanded */
+	const Type	Expanded						= 1 << 5;	/** true if this node should display its children*/
+	const Type	CanBeExpanded					= 1 << 6;	/** true if this node is able to be expanded */
 
-	const Type	EditInlineNew					= 1 << 7;		/** true if the property can be expanded into the property window. */
+	const Type	EditInlineNew					= 1 << 7;	/** true if the property can be expanded into the property window. */
 
-	const Type	SingleSelectOnly				= 1 << 8;		/** true if only a single object is selected. */
-	const Type  ShowCategories					= 1 << 9;		/** true if this node should show categories.  Different*/
+	const Type	SingleSelectOnly				= 1 << 8;	/** true if only a single object is selected. */
+	const Type  ShowCategories					= 1 << 9;	/** true if this node should show categories.  Different*/
 
 	const Type  HasEverBeenExpanded				= 1 << 10;	/** true if expand has ever been called on this node */
 
@@ -51,13 +51,13 @@ namespace EPropertyNodeFlags
 	const Type	IsAdvanced						= 1 << 16;	/** true if the property node is advanced (i.e it only shows up in advanced sections) */
 	const Type	IsCustomized					= 1 << 17;	/** true if this node's visual representation has been customized by the editor */
 	
-	const Type	RequiresValidation				= 1 << 18; /** true if this node could unexpectedly change (array changes, editinlinenew changes) */
+	const Type	RequiresValidation				= 1 << 18;	/** true if this node could unexpectedly change (array changes, editinlinenew changes) */
 
-	const Type	ShouldShowDisableEditOnInstance = 1 << 19; /** true if this node should show child properties marked CPF_DisableEditOnInstance */
+	const Type	ShouldShowDisableEditOnInstance = 1 << 19;	/** true if this node should show child properties marked CPF_DisableEditOnInstance */
 
-	const Type	IsReadOnly						= 1 << 20; /** true if this node is overridden to appear as read-only */
+	const Type	IsReadOnly						= 1 << 20;	/** true if this node is overridden to appear as read-only */
 
-	const Type	SkipChildValidation				= 1 << 21; /** true if this node should skip child validation */
+	const Type	SkipChildValidation				= 1 << 21;	/** true if this node should skip child validation */
 
 	const Type  ShowInnerObjectProperties		= 1 << 22;
 
@@ -65,7 +65,7 @@ namespace EPropertyNodeFlags
 
 	const Type	IsSparseClassData				= 1 << 24;	/** true if the property on this node is part of a sparse class data structure */
 
-	const Type	ShouldShowInViewport			= 1 << 25; /** true if the property should be shown in the viewport context menu */
+	const Type	ShouldShowInViewport			= 1 << 25;	/** true if the property should be shown in the viewport context menu */
 
 	const Type 	NoFlags							= 0;
 
@@ -302,6 +302,11 @@ public:
 	 */
 	void RebuildChildren();
 
+	/**
+	 * Mark this and all children as having been rebuilt.
+	 */
+	void MarkChildrenAsRebuilt();
+
 	/**
 	 * For derived windows to be able to add their nodes to the child array
 	 */
diff --git a/Engine/Source/Editor/PropertyEditor/Private/SDetailsViewBase.cpp b/Engine/Source/Editor/PropertyEditor/Private/SDetailsViewBase.cpp
--- a/Engine/Source/Editor/PropertyEditor/Private/SDetailsViewBase.cpp
+++ b/Engine/Source/Editor/PropertyEditor/Private/SDetailsViewBase.cpp
@@ -1030,7 +1030,7 @@ void SDetailsViewBase::Tick( const FGeometry& AllottedGeometry, const double InC
 		}
 		else // standard validation behavior
 		{
-			for (TSharedPtr<FComplexPropertyNode>& RootPropertyNode : RootPropertyNodes)
+			for (const TSharedPtr<FComplexPropertyNode>& RootPropertyNode : RootPropertyNodes)
 			{
 				EPropertyDataValidationResult Result = RootPropertyNode->EnsureDataIsValid();
 				if (Result == EPropertyDataValidationResult::PropertiesChanged || Result == EPropertyDataValidationResult::EditInlineNewValueChanged)
@@ -1038,10 +1038,16 @@ void SDetailsViewBase::Tick( const FGeometry& AllottedGeometry, const double InC
 					UpdatePropertyMaps();
 					bUpdateFilteredDetails = true;
 				}
-				else if (Result == EPropertyDataValidationResult::ArraySizeChanged || Result == EPropertyDataValidationResult::ChildrenRebuilt)
+				else if (Result == EPropertyDataValidationResult::ArraySizeChanged)
 				{
 					bUpdateFilteredDetails = true;
 				}
+				else if (Result == EPropertyDataValidationResult::ChildrenRebuilt)
+				{
+					RootPropertyNode->MarkChildrenAsRebuilt();
+				
+					bUpdateFilteredDetails = true;
+				}
 				else if (Result == EPropertyDataValidationResult::ObjectInvalid)
 				{
 					bValidateExternalNodes = false;
@@ -1170,8 +1176,9 @@ static void SetExpandedItems(TSharedPtr<FPropertyNode> InPropertyNode, const TSe
 	{
 		// check if this path is a prefix to some other path
 		// if it's not, we can early out and skip checking every child node, which is a win for very large child arrays
-		if (Item.Len() > Path.Len() && Item.StartsWith(Path) &&
-			(Item[Path.Len()] == '[' || Item[Path.Len()] == '.'))
+		if (Item.Len() > Path.Len() && 
+			(Item[Path.Len()] == '[' || Item[Path.Len()] == '.') &&
+			Item.StartsWith(Path))
 		{
 			bAnyPrefix = true;
 			break;
@@ -1340,8 +1347,6 @@ void SDetailsViewBase::RestoreAllExpandedItems()
 
 void SDetailsViewBase::RestoreExpandedItems(TSharedRef<FPropertyNode> StartNode)
 {
-	FString ExpandedCustomItems;
-
 	UStruct* BestBaseStruct = StartNode->FindComplexParent()->GetBaseStructure();
 
 	//while a valid class, and we're either the same as the base class (for multiple actors being selected and base class is AActor) OR we're not down to AActor yet)
@@ -1357,7 +1362,9 @@ void SDetailsViewBase::RestoreExpandedItems(TSharedRef<FPropertyNode> StartNode)
 
 	if (BestBaseStruct)
 	{
+		FString ExpandedCustomItems;
 		GConfig->GetString(TEXT("DetailCustomWidgetExpansion"), *BestBaseStruct->GetName(), ExpandedCustomItems, GEditorPerProjectIni);
+
 		TArray<FString> ExpandedCustomItemsArray;
 		ExpandedCustomItems.ParseIntoArray(ExpandedCustomItemsArray, TEXT(","), true);
 
