Add method to set the minimum right column width in a details panel

diff --git a/Engine/Source/Editor/PropertyEditor/Private/SDetailsViewBase.cpp b/Engine/Source/Editor/PropertyEditor/Private/SDetailsViewBase.cpp
--- a/Engine/Source/Editor/PropertyEditor/Private/SDetailsViewBase.cpp
+++ b/Engine/Source/Editor/PropertyEditor/Private/SDetailsViewBase.cpp
@@ -284,6 +284,11 @@ void SDetailsViewBase::SetOnDisplayedPropertiesChanged(FOnDisplayedPropertiesCha
 	OnDisplayedPropertiesChangedDelegate = InOnDisplayedPropertiesChangedDelegate;
 }
 
+void SDetailsViewBase::SetRightColumnMinWidth(float InMinWidth)
+{
+	ColumnSizeData.SetRightColumnMinWidth(InMinWidth);
+}
+
 void SDetailsViewBase::RerunCurrentFilter()
 {
 	UpdateFilteredDetails();
diff --git a/Engine/Source/Editor/PropertyEditor/Private/SDetailsViewBase.h b/Engine/Source/Editor/PropertyEditor/Private/SDetailsViewBase.h
--- a/Engine/Source/Editor/PropertyEditor/Private/SDetailsViewBase.h
+++ b/Engine/Source/Editor/PropertyEditor/Private/SDetailsViewBase.h
@@ -121,6 +121,8 @@ public:
 	{
 		CustomValidatePropertyNodesFunction = MoveTemp(InCustomValidatePropertyNodesFunction);
 	}
+	virtual void SetRightColumnMinWidth(float InMinWidth) override;
+	
 	/** IDetailsViewPrivate interface */
 	virtual void RerunCurrentFilter() override;
 	void SetNodeExpansionState(TSharedRef<FDetailTreeNode> InTreeNode, bool bExpand, bool bRecursive) override;
diff --git a/Engine/Source/Editor/PropertyEditor/Public/IDetailsView.h b/Engine/Source/Editor/PropertyEditor/Public/IDetailsView.h
--- a/Engine/Source/Editor/PropertyEditor/Public/IDetailsView.h
+++ b/Engine/Source/Editor/PropertyEditor/Public/IDetailsView.h
@@ -261,5 +261,6 @@ public:
 	 * This is useful if your implementation doesn't need to validate nodes every tick or needs to perform some other form of validation. */
 	virtual void SetCustomValidatePropertyNodesFunction(FOnValidateDetailsViewPropertyNodes InCustomValidatePropertyNodesFunction) = 0;
 
-
+	/** Sets the minimum width of the detail panel's right column. */
+	virtual void SetRightColumnMinWidth(float InMinWidth) = 0;
 };
