Removed inline forward declarations from PropertyNode.h.

#rnx

diff --git a/Engine/Source/Editor/PropertyEditor/Private/PropertyNode.h b/Engine/Source/Editor/PropertyEditor/Private/PropertyNode.h
--- a/Engine/Source/Editor/PropertyEditor/Private/PropertyNode.h
+++ b/Engine/Source/Editor/PropertyEditor/Private/PropertyNode.h
@@ -9,13 +9,17 @@
 #include "PropertyEditorModule.h"
 #include "EditConditionParser.h"
 
+class FCategoryPropertyNode;
 class FComplexPropertyNode;
+class FDetailTreeNode;
 class FEditConditionContext;
 class FEditConditionExpression;
+class FItemPropertyNode;
 class FNotifyHook;
 class FObjectPropertyNode;
 class FPropertyItemValueDataTrackerSlate;
 class FPropertyNode;
+class FPropertyRestriction;
 class FStructurePropertyNode;
 
 DECLARE_LOG_CATEGORY_EXTERN(LogPropertyNode, Log, All);
@@ -94,9 +98,6 @@ private:
 	bool bShowHiddenProperties;
 };
 
-class FPropertyItemValueDataTrackerSlate;
-class FPropertyNode;
-
 struct FAddressPair
 {
 	FAddressPair(const UObject* InObject, uint8* Address, bool bInIsStruct)
@@ -258,8 +259,6 @@ struct EPropertyArrayChangeType
 	};
 };
 
-class FComplexPropertyNode;
-
 enum EPropertyDataValidationResult : uint8
 {
 	/** The object(s) being viewed are now invalid */
@@ -320,43 +319,43 @@ public:
 	/**
 	 * Interface function to get at the derived FObjectPropertyNode class
 	 */
-	virtual class FObjectPropertyNode* AsObjectNode() { return nullptr; }
+	virtual FObjectPropertyNode* AsObjectNode() { return nullptr; }
 	virtual const FObjectPropertyNode* AsObjectNode() const { return nullptr; }
 
 	/**
 	 * Interface function to get at the derived FComplexPropertyNode class
 	 */
-	virtual class FComplexPropertyNode* AsComplexNode() { return nullptr; }
+	virtual FComplexPropertyNode* AsComplexNode() { return nullptr; }
 	virtual const FComplexPropertyNode* AsComplexNode() const { return nullptr; }
 
 	/**
 	 * Interface function to get at the derived FCategoryPropertyNode class
 	 */
-	virtual class FCategoryPropertyNode* AsCategoryNode() { return nullptr; }
+	virtual FCategoryPropertyNode* AsCategoryNode() { return nullptr; }
 	virtual const FCategoryPropertyNode* AsCategoryNode() const { return nullptr; }
 
 	/**
 	 * Interface function to get at the derived FItemPropertyNode class
 	 */
-	virtual class FItemPropertyNode* AsItemPropertyNode() { return nullptr; }
+	virtual FItemPropertyNode* AsItemPropertyNode() { return nullptr; }
 	virtual const FItemPropertyNode* AsItemPropertyNode() const { return nullptr; }
 
 	/**
 	 * Follows the chain of items upwards until it finds the complex property that houses this item.
 	 */
-	class FComplexPropertyNode* FindComplexParent();
+	FComplexPropertyNode* FindComplexParent();
 	const FComplexPropertyNode* FindComplexParent() const;
 
 	/**
 	 * Follows the chain of items upwards until it finds the object property that houses this item.
 	 */
-	class FObjectPropertyNode* FindObjectItemParent();
+	FObjectPropertyNode* FindObjectItemParent();
 	const FObjectPropertyNode* FindObjectItemParent() const;
 
 	/**
 	 * Follows the top-most object window that contains this property window item.
 	 */
-	class FObjectPropertyNode* FindRootObjectItemParent();
+	FObjectPropertyNode* FindRootObjectItemParent();
 
 	/**
 	 * Used to see if any data has been destroyed from under the property tree.  Should only be called during Tick
@@ -592,13 +591,13 @@ public:
 	bool IsReorderable();
 
 	/**Walks up the hierarchy and return true if any parent node is a favorite*/
-	bool IsChildOfFavorite(void) const;
+	bool IsChildOfFavorite() const;
 
-	void NotifyPreChange(FProperty* PropertyAboutToChange, class FNotifyHook* InNotifyHook);
-	void NotifyPreChange(FProperty* PropertyAboutToChange, class FNotifyHook* InNotifyHook, const TSet<UObject*>& AffectedInstances);
-	void NotifyPreChange(FProperty* PropertyAboutToChange, class FNotifyHook* InNotifyHook, TSet<UObject*>&& AffectedInstances);
+	void NotifyPreChange(FProperty* PropertyAboutToChange, FNotifyHook* InNotifyHook);
+	void NotifyPreChange(FProperty* PropertyAboutToChange, FNotifyHook* InNotifyHook, const TSet<UObject*>& AffectedInstances);
+	void NotifyPreChange(FProperty* PropertyAboutToChange, FNotifyHook* InNotifyHook, TSet<UObject*>&& AffectedInstances);
 
-	void NotifyPostChange(FPropertyChangedEvent& InPropertyChangedEvent, class FNotifyHook* InNotifyHook);
+	void NotifyPostChange(FPropertyChangedEvent& InPropertyChangedEvent, FNotifyHook* InNotifyHook);
 
 	void SetOnRebuildChildren(FSimpleDelegate InOnRebuildChildren);
 
@@ -690,7 +689,7 @@ public:
 	/**
 	 * Marks windows as visible based their favorites status
 	 */
-	void ProcessSeenFlagsForFavorites(void);
+	void ProcessSeenFlagsForFavorites();
 
 	/**
 	 * @return true if this node should be visible in a tree
@@ -825,7 +824,7 @@ public:
 	 * Adds a restriction to the possible values for this property.
 	 * @param Restriction	The restriction being added to this property.
 	 */
-	virtual void AddRestriction(TSharedRef<const class FPropertyRestriction> Restriction);
+	virtual void AddRestriction(TSharedRef<const FPropertyRestriction> Restriction);
 
 	/**
 	* Tests if a value is hidden for this property
@@ -884,7 +883,7 @@ public:
 	 */
 	virtual bool GenerateRestrictionToolTip(const FString& Value, FText& OutTooltip)const;
 
-	const TArray<TSharedRef<const class FPropertyRestriction>>& GetRestrictions() const
+	const TArray<TSharedRef<const FPropertyRestriction>>& GetRestrictions() const
 	{
 		return Restrictions;
 	}
@@ -1132,10 +1131,10 @@ protected:
 	bool bChildrenRebuilt;
 
 	/** An array of restrictions limiting this property's potential values in property editors.*/
-	TArray<TSharedRef<const class FPropertyRestriction>> Restrictions;
+	TArray<TSharedRef<const FPropertyRestriction>> Restrictions;
 
 	/** Optional reference to a tree node that is displaying this property */
-	TWeakPtr< class FDetailTreeNode > TreeNode;
+	TWeakPtr<FDetailTreeNode> TreeNode;
 
 	/**
 	 * Stores metadata for this instance of the property (in contrast
@@ -1178,7 +1177,7 @@ public:
 	virtual FComplexPropertyNode* AsComplexNode() override { return this; }
 	virtual const FComplexPropertyNode* AsComplexNode() const override { return this; }
 
-	virtual class FStructurePropertyNode* AsStructureNode() { return nullptr; }
+	virtual FStructurePropertyNode* AsStructureNode() { return nullptr; }
 	virtual const FStructurePropertyNode* AsStructureNode() const { return nullptr; }
 
 	virtual UStruct* GetBaseStructure() = 0;
