Removed inline forward declarations.

#rnx

diff --git a/Engine/Source/Editor/PropertyEditor/Public/DetailsViewArgs.h b/Engine/Source/Editor/PropertyEditor/Public/DetailsViewArgs.h
--- a/Engine/Source/Editor/PropertyEditor/Public/DetailsViewArgs.h
+++ b/Engine/Source/Editor/PropertyEditor/Public/DetailsViewArgs.h
@@ -4,6 +4,12 @@
 
 #include "CoreMinimal.h"
 
+class FDetailsViewObjectFilter;
+class FNotifyHook;
+class FTabManager;
+class FUICommandList;
+class IClassViewerFilter;
+
 enum class EEditDefaultsOnlyNodeVisibility : uint8
 {
 	/** Always show nodes that have the EditDefaultsOnly aka CPF_DisableEditOnInstance flag. */
@@ -14,7 +20,7 @@ enum class EEditDefaultsOnlyNodeVisibility : uint8
 	Automatic,
 };
 
-/**
+/**éé
  * Init params for a details view widget
  */
 struct FDetailsViewArgs
@@ -34,18 +40,18 @@ struct FDetailsViewArgs
 	/** Controls how CPF_DisableEditOnInstance nodes will be treated */
 	EEditDefaultsOnlyNodeVisibility DefaultsOnlyVisibility;
 	/** The command list from the host of the details view, allowing child widgets to bind actions with a bound chord */
-	TSharedPtr<class FUICommandList> HostCommandList;
+	TSharedPtr<FUICommandList> HostCommandList;
 	/** The tab manager from the host of the details view, allowing child widgets to spawn tabs */
-	TSharedPtr<class FTabManager> HostTabManager;
+	TSharedPtr<FTabManager> HostTabManager;
 	/** Optional object filter to use for more complex handling of what a details panel is viewing. */
-	TSharedPtr<class FDetailsViewObjectFilter> ObjectFilter;
+	TSharedPtr<FDetailsViewObjectFilter> ObjectFilter;
 	/** Optional custom filter(s) to apply to the class viewer widget for class object property values. */
-	TArray<TSharedRef<class IClassViewerFilter>> ClassViewerFilters;
+	TArray<TSharedRef<IClassViewerFilter>> ClassViewerFilters;
 
 	/** Identifier for this details view; NAME_None if this view is anonymous */
 	FName ViewIdentifier;
 	/** Notify hook to call when properties are changed */
-	class FNotifyHook* NotifyHook;
+	FNotifyHook* NotifyHook;
 	/** Settings for displaying the name area (@see ENameAreaSettings) */
 	int32 NameAreaSettings;
 	/** The default value column width, as a percentage, 0-1. */
diff --git a/Engine/Source/Editor/PropertyEditor/Public/IDetailPropertyExtensionHandler.h b/Engine/Source/Editor/PropertyEditor/Public/IDetailPropertyExtensionHandler.h
--- a/Engine/Source/Editor/PropertyEditor/Public/IDetailPropertyExtensionHandler.h
+++ b/Engine/Source/Editor/PropertyEditor/Public/IDetailPropertyExtensionHandler.h
@@ -14,7 +14,7 @@ class IDetailPropertyExtensionHandler
 public:
 	virtual ~IDetailPropertyExtensionHandler(){ }
 
-	virtual bool IsPropertyExtendable(const UClass* InObjectClass, const class IPropertyHandle& PropertyHandle) const = 0;
+	virtual bool IsPropertyExtendable(const UClass* InObjectClass, const IPropertyHandle& PropertyHandle) const = 0;
 
 	UE_DEPRECATED(4.24, "Please use ExtendWidgetRow")
 	virtual TSharedRef<SWidget> GenerateExtensionWidget(const UClass* InObjectClass, TSharedPtr<IPropertyHandle> PropertyHandle) 
