Added FStringPrefixTree, which stores a set of strings as a tree of nodes that share common prefixes. This results in far fewer string comparisons if many strings share a common prefix, eg. with property paths.

The set "A.B", "A.C", "B" would result in a tree like this, with the asterisked nodes being leaf nodes:

|- A.
|  |- B *
|  |- C *
|- B *

#rnx

diff --git a/Engine/Source/Editor/PropertyEditor/Private/StringPrefixTree.cpp b/Engine/Source/Editor/PropertyEditor/Private/StringPrefixTree.cpp
new file mode 100644
--- /dev/null
+++ b/Engine/Source/Editor/PropertyEditor/Private/StringPrefixTree.cpp
@@ -0,0 +1,363 @@
+// Copyright Epic Games, Inc. All Rights Reserved.
+
+#include "StringPrefixTree.h"
+
+void FStringPrefixTree::Insert(FStringView ToAdd)
+{
+	if (ToAdd.IsEmpty())
+	{
+		Root.IsLeaf = true;
+		return;
+	}
+
+	Root.Insert(ToAdd);
+}
+
+void FStringPrefixTree::InsertAll(const TArray<FString>& Array)
+{
+	for (const FString& Entry : Array)
+	{
+		Insert(Entry);
+	}
+}
+
+void FStringPrefixTree::InsertAll(const TArray<FStringView>& Array)
+{
+	for (const FStringView& Entry : Array)
+	{
+		Insert(Entry);
+	}
+}
+
+void FStringPrefixTree::Remove(FStringView ToRemove)
+{
+	if (ToRemove.IsEmpty())
+	{
+		Root.IsLeaf = false;
+		return;
+	}
+
+	Root.Remove(ToRemove);
+}
+
+void FStringPrefixTree::Clear()
+{
+	Root.Children.Reset();
+	Root.IsLeaf = false;
+	Root.Prefix = FString();
+}
+
+bool FStringPrefixTree::Contains(FStringView ToFind) const
+{
+	if (ToFind.IsEmpty())
+	{
+		return Root.IsLeaf;
+	}
+
+	return Root.Contains(ToFind);
+}
+
+bool FStringPrefixTree::AnyStartsWith(FStringView Prefix) const
+{
+	if (Prefix.IsEmpty())
+	{
+		return Root.IsLeaf || Root.Children.Num() > 0;
+	}
+
+	return Root.AnyStartsWith(Prefix);
+}
+
+int32 FStringPrefixTree::Size() const
+{
+	return Root.Size();
+}
+
+int32 FStringPrefixTree::NumNodes() const
+{
+	return Root.NumNodes();
+}
+
+TArray<FString> FStringPrefixTree::GetAllEntries() const 
+{
+	TArray<FString> Entries;
+
+	Root.GetAllEntries(FString(), Entries);
+
+	return MoveTemp(Entries);
+}
+
+FString FStringPrefixTree::DumpToString() const
+{
+	TStringBuilder<512> Builder;
+	TArray<bool, TInlineAllocator<32>> DrawLines;
+	Root.DumpToString(Builder, DrawLines, false);
+	return Builder.ToString();
+}
+
+static int32 GetCommonPrefixLength(FStringView A, FStringView B)
+{
+	int32 CharIdx = 0;
+	for (; CharIdx < A.Len() && CharIdx < B.Len(); ++CharIdx)
+	{
+		if (A[CharIdx] != B[CharIdx])
+		{
+			break;
+		}
+	}
+	return CharIdx;
+}
+
+void FStringPrefixTree::FNode::Insert(FStringView ToAdd)
+{
+	bool bAdded = false;
+
+	// find a child that shares a prefix
+	for (FNode& Child : Children)
+	{
+		int32 CharIdx = GetCommonPrefixLength(Child.Prefix, ToAdd);
+		if (CharIdx > 0)
+		{
+			// the paths share a common prefix, there can be only one such node
+
+			bAdded = true;
+
+			if (CharIdx == Child.Prefix.Len())
+			{
+				// the prefix is the whole child path, we don't need to split the child node, but rather just add the postfix as a child
+				if (CharIdx == ToAdd.Len())
+				{
+					// the paths are identical, we don't need to add a node, just mark this as a leaf node
+					IsLeaf = true;
+					break;
+				}
+
+				// eg. Child = Foo, ToAdd = Foo.Bar
+				Child.Insert(ToAdd.RightChop(CharIdx));
+			}
+
+			// this is a partial prefix of the child
+			// eg. Child = Foobar, ToAdd = Foo
+			else if (CharIdx < Child.Prefix.Len())
+			{
+				// modify the child node such that we take the common prefix and insert a node above it
+				// eg. Child = "Foobar", ToAdd = "Foo",
+				// Child -> "Foo", Split -> "bar"
+				
+				FNode Split;
+				Split.Prefix = Child.Prefix.RightChop(CharIdx);
+				Split.Children = MoveTemp(Child.Children);
+				Split.IsLeaf = Child.IsLeaf;
+
+				Child.Prefix.LeftInline(CharIdx);
+				Child.Children.Reset();
+				Child.Children.Add(MoveTemp(Split));
+
+				if (CharIdx < ToAdd.Len())
+				{
+					// this was a partial match that still has leftover chars, 
+					// eg. Child = "Foo", ToAdd = "Foz"
+					// Child -> "Fo", Split -> "o", Added -> "z"
+					FNode& AddedNode = Child.Children.AddDefaulted_GetRef();
+					AddedNode.Prefix = ToAdd.RightChop(CharIdx);
+					AddedNode.IsLeaf = true;
+				
+					Child.IsLeaf = false;
+				}
+				else
+				{
+					// this was a full match, eg. Child = "Foobar", ToAdd = "Foo"
+					// this is a leaf because this represents ToAdd now
+					Child.IsLeaf = true;
+				}
+			}
+				
+			break;
+		}
+	}
+
+	if (!bAdded)
+	{
+		FNode& NewChild = Children.AddDefaulted_GetRef();
+		NewChild.Prefix = ToAdd;
+		NewChild.IsLeaf = true;
+	}
+}
+
+/** 
+	* Remove an entry. This can fragment the tree and should be avoided.
+	*/
+void FStringPrefixTree::FNode::Remove(FStringView ToRemove)
+{
+	for (int32 Idx = 0; Idx < Children.Num(); ++Idx)
+	{
+		FNode& Child = Children[Idx];
+
+		if (ToRemove.StartsWith(Child.Prefix))
+		{
+			if (ToRemove.Len() == Child.Prefix.Len())
+			{
+				// matched the entire string
+				Child.IsLeaf = false;
+			}
+			else
+			{
+				Child.Remove(ToRemove.RightChop(Child.Prefix.Len()));
+			}
+
+			if (Child.CanPrune())
+			{
+				Children.RemoveAt(Idx);
+			}
+
+			break;
+		}
+	}
+}
+
+bool FStringPrefixTree::FNode::Contains(FStringView ToFind) const
+{
+	for (const FNode& Child : Children)
+	{
+		if (ToFind.StartsWith(Child.Prefix))
+		{
+			if (Child.Prefix.Len() == ToFind.Len())
+			{
+				// matched the entire string
+				return Child.IsLeaf;
+			}
+
+			return Child.Contains(ToFind.RightChop(Child.Prefix.Len()));
+		}
+	}
+
+	return false;
+}
+
+bool FStringPrefixTree::FNode::AnyStartsWith(FStringView ToFind) const
+{
+	for (const FNode& Child : Children)
+	{
+		int32 CharIdx = GetCommonPrefixLength(Child.Prefix, ToFind);
+		if (CharIdx > 0)
+		{
+			if (CharIdx <= Child.Prefix.Len())
+			{
+				// identical strings or child is longer than the prefix we're searching for
+				return true;
+			}
+			
+			return Child.AnyStartsWith(ToFind.RightChop(Child.Prefix.Len()));
+		}
+
+	}
+
+	return false;
+}
+
+void FStringPrefixTree::FNode::GetAllEntries(FStringView CurrentPrefix, TArray<FString>& OutEntries) const 
+{
+	FString NewPrefix = CurrentPrefix + Prefix;
+
+	if (IsLeaf)
+	{
+		OutEntries.Add(NewPrefix);
+	}
+
+	for (const FNode& Child : Children)
+	{
+		Child.GetAllEntries(NewPrefix, OutEntries);
+	}
+}
+
+int32 FStringPrefixTree::FNode::Size() const 
+{
+	int32 Result = 0;
+	for (const FNode& Child : Children)
+	{
+		// only leaf children count
+		if (Child.IsLeaf)
+		{
+			++Result;
+		}
+		Result += Child.Size();
+	}
+	return Result;
+}
+
+int32 FStringPrefixTree::FNode::NumNodes() const
+{
+	int32 Result = 1;
+	for (const FNode& Child : Children)
+	{
+		Result += Child.NumNodes();
+	}
+	return Result;
+}
+
+void FStringPrefixTree::FNode::DumpToString(FStringBuilderBase& Builder, TArray<bool, TInlineAllocator<32>>& DrawLines, bool bLastChild) const
+{
+	for (int32 Level = 0; Level < DrawLines.Num(); ++Level)
+	{
+		if (DrawLines[Level])
+		{
+			Builder += TEXT("|");
+		}
+		else
+		{
+			Builder += TEXT(" ");
+		}
+		
+		if (Level != DrawLines.Num() - 1)
+		{
+			Builder += TEXT(" ");
+		}
+	}
+
+	if (bLastChild)
+	{
+		DrawLines.Last() = false;
+	}
+
+	if (Prefix.IsEmpty())
+	{
+		Builder += TEXT("\"\"");
+	}
+	else
+	{
+		Builder += TEXT("- ");
+		Builder += Prefix;
+	}
+
+	if (IsLeaf)
+	{
+		Builder += TEXT("*");
+	}
+
+	Builder += TEXT("\n");
+
+	for (int32 Idx = 0; Idx < Children.Num(); ++Idx)
+	{
+		DrawLines.Push(true);
+		Children[Idx].DumpToString(Builder, DrawLines, Idx == Children.Num() - 1);
+		DrawLines.Pop();
+	}
+}
+
+// we can prune any node that is not a leaf and does not contain any children which are leaves anymore
+bool FStringPrefixTree::FNode::CanPrune() const 
+{
+	if (IsLeaf)
+	{
+		return false;
+	}
+
+	for (const FNode& Child : Children)
+	{
+		if (!Child.CanPrune())
+		{
+			return false;
+		}
+	}
+
+	return true;
+}
diff --git a/Engine/Source/Editor/PropertyEditor/Private/StringPrefixTree.h b/Engine/Source/Editor/PropertyEditor/Private/StringPrefixTree.h
new file mode 100644
--- /dev/null
+++ b/Engine/Source/Editor/PropertyEditor/Private/StringPrefixTree.h
@@ -0,0 +1,112 @@
+// Copyright Epic Games, Inc. All Rights Reserved.
+
+#pragma once
+
+#include "Containers/Array.h"
+#include "Containers/StringView.h"
+#include "Containers/UnrealString.h"
+#include "Misc/StringBuilder.h"
+
+/**
+ * A string prefix tree is conceptually similar to a trie.
+ * However, instead of adding a new nodes per single character, nodes encode a string prefix, and entries are built up by concatenating the path to the node.
+ * For example, if we add the entries "ABC", "AC", and "D" to the tree, we will have:
+ *   <root>
+ *   |- "A"
+ *   | |- "B"
+ *   | | |- "C" *
+ *   | |- "C" *
+ *   |- "D" *
+ * 
+ * Where only the nodes marked with * are leaf nodes. 
+ * An entry is considered to be in the list only if it is a leaf node.
+ * This tree structure allows for long strings that may have long similar sections to be searched far faster than storing them in eg. a set.
+ */
+class FStringPrefixTree
+{
+public:
+
+	/**
+	 * Insert the given string into the tree.
+	 */
+	void Insert(FStringView ToAdd);
+
+	/**
+	 * Insert all the strings in the array.
+	 */
+	void InsertAll(const TArray<FString>& Array);
+
+	/**
+	 * Insert all the strings in the array.
+	 */
+	void InsertAll(const TArray<FStringView>& Array);
+
+	/** 
+	 * Remove an entry. This can fragment the tree and should be avoided.
+	 */
+	void Remove(FStringView ToRemove);
+
+	/**
+	 * Clear the tree.
+	 */
+	void Clear();
+
+	/**
+	 * Does the tree contain the given entry?
+	 */
+	bool Contains(FStringView Entry) const;
+
+	/**
+	 * Does any entry start with the given prefix.
+	 * Eg. if the tree contains "Foobar", this would return true if given "Foo".
+	 */
+	bool AnyStartsWith(FStringView Prefix) const;
+
+	/**
+	 * Returns the number of leaf entries in the tree.
+	 * Equivalent to the number of items added via Insert() or InsertAll().
+	 */
+	int32 Size() const;
+
+	/**
+	 * Dump all leaf nodes into an array.
+	 */
+	TArray<FString> GetAllEntries() const;
+
+	/**
+	 * For testing and debugging.
+	 * The number of internal nodes. This is a different (usually larger) number than the number of entries.
+	 */
+	int32 NumNodes() const;
+	
+	/**
+	 * For testing and debugging.
+	 * Pretty prints the internal tree structure as a string.
+	 */
+	FString DumpToString() const;
+
+private:
+
+	struct FNode
+	{
+		FString Prefix;
+		TArray<FNode> Children;
+		bool IsLeaf = false;
+
+		void Insert(FStringView ToAdd);
+		void Remove(FStringView ToRemove);
+		bool Contains(FStringView ToFind) const;
+		bool AnyStartsWith(FStringView ToFind) const;
+		int32 Size() const;
+
+		int32 NumNodes() const;
+		void GetAllEntries(FStringView CurrentPrefix, TArray<FString>& OutEntries) const;
+		void DumpToString(FStringBuilderBase& Builder, TArray<bool, TInlineAllocator<32>>& DrawLines, bool bLastChild) const;
+
+	private:
+		// we can prune any node that is not a leaf and does not contain any children which are leaves anymore
+		bool CanPrune() const;
+	};
+
+	FNode Root;
+};
\ No newline at end of file
diff --git a/Engine/Source/Editor/PropertyEditor/Private/StringPrefixTreeTests.cpp b/Engine/Source/Editor/PropertyEditor/Private/StringPrefixTreeTests.cpp
new file mode 100644
--- /dev/null
+++ b/Engine/Source/Editor/PropertyEditor/Private/StringPrefixTreeTests.cpp
@@ -0,0 +1,223 @@
+// Copyright Epic Games, Inc. All Rights Reserved.
+
+#if WITH_DEV_AUTOMATION_TESTS
+
+#include "Misc/AutomationTest.h"
+#include "StringPrefixTree.h"
+
+IMPLEMENT_SIMPLE_AUTOMATION_TEST(FStringPrefixTreeTests_Insert, "StringPrefixTree.Insert", EAutomationTestFlags::EditorContext | EAutomationTestFlags::EngineFilter)
+
+bool FStringPrefixTreeTests_Insert::RunTest(const FString& Parameters)
+{
+	FStringPrefixTree Tree;
+
+	Tree.Insert(TEXT("Foobar"));
+
+	TestEqual(TEXT("One"), Tree.Size(), 1);
+	TestTrue(TEXT("One"), Tree.Contains(TEXT("Foobar")));
+
+	Tree.Insert(TEXT("Foo"));
+
+	TestEqual(TEXT("Two"), Tree.Size(), 2);
+	TestTrue(TEXT("Two"), Tree.Contains(TEXT("Foo")));
+	TestTrue(TEXT("Two"), Tree.Contains(TEXT("Foobar")));
+	TestFalse(TEXT("Two"), Tree.Contains(TEXT("Fo")));
+	TestFalse(TEXT("Two"), Tree.Contains(TEXT("Foob")));
+
+	// check that the same holds if we build the tree in the opposite way
+	// ie. Foo, Foobar
+	Tree.Clear();
+
+	TestEqual(TEXT("Cleared"), Tree.Size(), 0);
+
+	Tree.Insert(TEXT("Foo"));
+
+	TestEqual(TEXT("One"), Tree.Size(), 1);
+	TestTrue(TEXT("One"), Tree.Contains(TEXT("Foo")));
+
+	Tree.Insert(TEXT("Foobar"));
+
+	TestEqual(TEXT("Two"), Tree.Size(), 2);
+	TestTrue(TEXT("Two"), Tree.Contains(TEXT("Foo")));
+	TestTrue(TEXT("Two"), Tree.Contains(TEXT("Foobar")));
+	TestFalse(TEXT("Two"), Tree.Contains(TEXT("Fo")));
+	TestFalse(TEXT("Two"), Tree.Contains(TEXT("Foob")));
+
+	Tree.Insert(TEXT("Foobaz"));
+	
+	TestEqual(TEXT("Three"), Tree.Size(), 3);
+	TestTrue(TEXT("Three"), Tree.Contains(TEXT("Foobar")));
+	TestTrue(TEXT("Three"), Tree.Contains(TEXT("Foobaz")));
+	TestFalse(TEXT("Three"), Tree.Contains(TEXT("Fooba")));
+
+	Tree.Insert(TEXT("Foobar.Foo"));
+	
+	TestEqual(TEXT("Four"), Tree.Size(), 4);
+	TestTrue(TEXT("Four"), Tree.Contains(TEXT("Foobar.Foo")));
+	TestTrue(TEXT("Four"), Tree.Contains(TEXT("Foobaz")));
+	TestTrue(TEXT("Four"), Tree.Contains(TEXT("Foobar")));
+	TestFalse(TEXT("Four"), Tree.Contains(TEXT("Foobar.")));
+
+	Tree.Insert(TEXT("Bar"));
+	Tree.Insert(TEXT("Bar.Foo"));
+
+	TestEqual(TEXT("Six"), Tree.Size(), 6);
+	TestTrue(TEXT("Six"), Tree.Contains(TEXT("Bar.Foo")));
+	TestTrue(TEXT("Six"), Tree.Contains(TEXT("Bar")));
+
+	return true;
+}
+
+IMPLEMENT_SIMPLE_AUTOMATION_TEST(FStringPrefixTreeTests_Contains, "StringPrefixTree.Contains", EAutomationTestFlags::EditorContext | EAutomationTestFlags::EngineFilter)
+
+bool FStringPrefixTreeTests_Contains::RunTest(const FString& Parameters)
+{
+	FStringPrefixTree Tree;
+
+	TestFalse(TEXT("Zero"), Tree.Contains(TEXT("")));
+
+	Tree.Insert(TEXT("Foobar"));
+
+	TestEqual(TEXT("One"), Tree.Size(), 1);
+
+	auto ContainsMatch = [](const FStringPrefixTree& Tree, const FString& Value)
+	{
+		TArray<FString> Entries = Tree.GetAllEntries();
+		return Tree.Contains(Value) == Entries.Contains(Value);
+	};
+
+	TestTrue(TEXT("One"), Tree.Contains(TEXT("Foobar")));
+
+	Tree.Insert(TEXT("Foo"));
+
+	TestEqual(TEXT("Two"), Tree.Size(), 2);
+	TestTrue(TEXT("Two"), ContainsMatch(Tree, TEXT("Foo")));
+	TestTrue(TEXT("Two"), ContainsMatch(Tree, TEXT("Foobar")));
+	TestTrue(TEXT("Two"), ContainsMatch(Tree, TEXT("Fo")));
+	TestTrue(TEXT("Two"), ContainsMatch(Tree, TEXT("Foob")));
+
+	// check that the same holds if we build the tree in the opposite way
+	// ie. Foo, Foobar
+	Tree.Clear();
+
+	TestEqual(TEXT("Cleared"), Tree.Size(), 0);
+
+	Tree.Insert(TEXT("Foo"));
+
+	TestEqual(TEXT("One"), Tree.Size(), 1);
+	TestTrue(TEXT("One"), ContainsMatch(Tree, TEXT("Foo")));
+
+	Tree.Insert(TEXT("Foobar"));
+
+	TestEqual(TEXT("Two"), Tree.Size(), 2);
+	TestTrue(TEXT("Two"), ContainsMatch(Tree, TEXT("Foo")));
+	TestTrue(TEXT("Two"), ContainsMatch(Tree, TEXT("Foobar")));
+	TestTrue(TEXT("Two"), ContainsMatch(Tree, TEXT("Fo")));
+	TestTrue(TEXT("Two"), ContainsMatch(Tree, TEXT("Foob")));
+
+	Tree.Insert(TEXT("Foobaz"));
+	
+	TestEqual(TEXT("Two"), Tree.Size(), 3);
+	TestTrue(TEXT("Three"), ContainsMatch(Tree, TEXT("Foobar")));
+	TestTrue(TEXT("Three"), ContainsMatch(Tree, TEXT("Foobaz")));
+
+	Tree.Insert(TEXT("Foobar.Foo"));
+	
+	TestEqual(TEXT("Four"), Tree.Size(), 4);
+	TestTrue(TEXT("Four"), ContainsMatch(Tree, TEXT("Foobar.Foo")));
+	TestTrue(TEXT("Four"), ContainsMatch(Tree, TEXT("Foobaz")));
+	TestTrue(TEXT("Four"), ContainsMatch(Tree, TEXT("Foobar")));
+	TestTrue(TEXT("Four"), ContainsMatch(Tree, TEXT("Foobar.")));
+
+	Tree.Insert(TEXT(""));
+
+	TestTrue(TEXT("Five"), Tree.Contains(TEXT("")));
+	TestTrue(TEXT("Five"), ContainsMatch(Tree, TEXT("")));
+
+	return true;
+}
+
+IMPLEMENT_SIMPLE_AUTOMATION_TEST(FStringPrefixTreeTests_AnyStartsWith, "StringPrefixTree.AnyStartsWith", EAutomationTestFlags::EditorContext | EAutomationTestFlags::EngineFilter)
+
+bool FStringPrefixTreeTests_AnyStartsWith::RunTest(const FString& Parameters)
+{
+	FStringPrefixTree Tree;
+	TestFalse(TEXT("Zero"), Tree.AnyStartsWith(TEXT("Foo")));
+	TestFalse(TEXT("Zero"), Tree.AnyStartsWith(TEXT("")));
+
+	Tree.Insert(TEXT("Foobar"));
+
+	TestTrue(TEXT("One"), Tree.AnyStartsWith(TEXT("Foo")));
+	TestTrue(TEXT("One"), Tree.AnyStartsWith(TEXT("Foobar")));
+	TestTrue(TEXT("One"), Tree.AnyStartsWith(TEXT("F")));
+
+	TestFalse(TEXT("One"), Tree.AnyStartsWith(TEXT("b")));
+	TestFalse(TEXT("One"), Tree.AnyStartsWith(TEXT("bar")));
+
+	return true;
+}
+
+IMPLEMENT_SIMPLE_AUTOMATION_TEST(FStringPrefixTreeTests_Remove, "StringPrefixTree.Remove", EAutomationTestFlags::EditorContext | EAutomationTestFlags::EngineFilter)
+
+bool FStringPrefixTreeTests_Remove::RunTest(const FString& Parameters)
+{
+	FStringPrefixTree Tree;
+
+	TestFalse(TEXT("Zero"), Tree.Contains(TEXT("")));
+
+	Tree.Insert(TEXT("Foobar"));
+	Tree.Remove(TEXT("Foobar"));
+	
+	TestFalse(TEXT("Zero"), Tree.Contains(TEXT("Foobar")));
+	TestEqual(TEXT("Zero"), Tree.NumNodes(), 1);
+	TestEqual(TEXT("Zero"), Tree.Size(), 0);
+
+	Tree.Insert(TEXT("Foo"));
+	Tree.Remove(TEXT("Foo"));
+
+	TestEqual(TEXT("Zero"), Tree.NumNodes(), 1);
+	TestEqual(TEXT("Zero"), Tree.Size(), 0);
+
+	Tree.Insert(TEXT("Foobar"));
+	Tree.Insert(TEXT("Foo"));
+
+	Tree.Remove(TEXT("Foo"));
+	TestFalse(TEXT("Remove One"), Tree.Contains(TEXT("Foo")));
+	TestTrue(TEXT("Remove One"), Tree.Contains(TEXT("Foobar")));
+	TestEqual(TEXT("Remove One"), Tree.NumNodes(), 3);
+	TestEqual(TEXT("Remove One"), Tree.Size(), 1);
+
+	Tree.Remove(TEXT("Foobar"));
+
+	TestFalse(TEXT("Remove Two"), Tree.Contains(TEXT("Foo")));
+	TestEqual(TEXT("Remove Two"), Tree.NumNodes(), 1);
+	TestEqual(TEXT("Remove Two"), Tree.Size(), 0);
+
+	return true;
+}
+
+IMPLEMENT_SIMPLE_AUTOMATION_TEST(FStringPrefixTreeTests_DumpToString, "StringPrefixTree.DumpToString", EAutomationTestFlags::EditorContext | EAutomationTestFlags::EngineFilter)
+
+bool FStringPrefixTreeTests_DumpToString::RunTest(const FString& Parameters)
+{
+	FStringPrefixTree Tree;
+	
+	Tree.Insert(TEXT("A"));
+	Tree.Insert(TEXT("A.1"));
+	Tree.Insert(TEXT("A.1.Z"));
+	Tree.Insert(TEXT("A.2.Y"));
+	Tree.Insert(TEXT("A.2.X"));
+	Tree.Insert(TEXT("B.3"));
+	Tree.Insert(TEXT("B.3.W"));
+	Tree.Insert(TEXT("C"));
+
+	const FString Dumped = Tree.DumpToString();
+
+	TArray<FString> Lines;
+	Dumped.ParseIntoArrayLines(Lines);
+	TestEqual(TEXT("Lines"), Lines.Num(), Tree.NumNodes());
+
+	return true;
+}
+
+#endif // WITH_DEV_AUTOMATION_TESTS
