Switched expanded nodes in SDetailsViewBase to use FStringPrefixTree.

This improved performance in one test case when calling RestoreExpandedItems() from ~500ms to ~100ms.

#rnx

diff --git a/Engine/Source/Editor/PropertyEditor/Private/SDetailsView.cpp b/Engine/Source/Editor/PropertyEditor/Private/SDetailsView.cpp
--- a/Engine/Source/Editor/PropertyEditor/Private/SDetailsView.cpp
+++ b/Engine/Source/Editor/PropertyEditor/Private/SDetailsView.cpp
@@ -896,7 +896,7 @@ void SDetailsView::PreSetObject(int32 InNewNumObjects)
 	}
 
 	RootPropertyNodes.Empty(InNewNumObjects);
-	ExpandedDetailNodes.Empty();
+	ExpandedDetailNodes.Clear();
 
 	for (int32 NewRootIndex = 0; NewRootIndex < InNewNumObjects; ++NewRootIndex)
 	{
diff --git a/Engine/Source/Editor/PropertyEditor/Private/SDetailsViewBase.cpp b/Engine/Source/Editor/PropertyEditor/Private/SDetailsViewBase.cpp
--- a/Engine/Source/Editor/PropertyEditor/Private/SDetailsViewBase.cpp
+++ b/Engine/Source/Editor/PropertyEditor/Private/SDetailsViewBase.cpp
@@ -27,7 +27,6 @@
 #include "SDetailNameArea.h"
 #include "SDetailsView.h"
 #include "PropertyEditorPermissionList.h"
-
 #include "ThumbnailRendering/ThumbnailManager.h"
 
 SDetailsViewBase::SDetailsViewBase() :
@@ -627,7 +626,7 @@ void SDetailsViewBase::SaveCustomExpansionState(const FString& NodePath, bool bI
 {
 	if (bIsExpanded)
 	{
-		ExpandedDetailNodes.Add(NodePath);
+		ExpandedDetailNodes.Insert(NodePath);
 	}
 	else
 	{
@@ -1136,7 +1135,7 @@ void SDetailsViewBase::Tick( const FGeometry& AllottedGeometry, const double InC
 * @param InPropertyNode			The node to get expanded items from
 * @param OutExpandedItems	List of expanded items that were found
 */
-static void GetExpandedItems(TSharedPtr<FPropertyNode> InPropertyNode, TSet<FString>& OutExpandedItems)
+static void GetExpandedItems(TSharedPtr<FPropertyNode> InPropertyNode, FStringPrefixTree& OutExpandedItems)
 {
 	if (InPropertyNode->HasNodeFlags(EPropertyNodeFlags::Expanded))
 	{
@@ -1145,7 +1144,7 @@ static void GetExpandedItems(TSharedPtr<FPropertyNode> InPropertyNode, TSet<FStr
 		Path.Empty(128);
 		InPropertyNode->GetQualifiedName(Path, bWithArrayIndex);
 
-		OutExpandedItems.Add(Path);
+		OutExpandedItems.Insert(Path);
 	}
 
 	for (int32 ChildIndex = 0; ChildIndex < InPropertyNode->GetNumChildNodes(); ++ChildIndex)
@@ -1160,7 +1159,7 @@ static void GetExpandedItems(TSharedPtr<FPropertyNode> InPropertyNode, TSet<FStr
 * @param InNode			The node to set expanded items on
 * @param OutExpandedItems	List of expanded items to set
 */
-static void SetExpandedItems(TSharedPtr<FPropertyNode> InPropertyNode, const TSet<FString>& InExpandedItems, bool bCollapseRest)
+static void SetExpandedItems(TSharedPtr<FPropertyNode> InPropertyNode, const FStringPrefixTree& InExpandedItems, bool bCollapseRest)
 {
 	const bool bWithArrayIndex = true;
 	FString Path;
@@ -1176,21 +1175,7 @@ static void SetExpandedItems(TSharedPtr<FPropertyNode> InPropertyNode, const TSe
 		InPropertyNode->SetNodeFlags(EPropertyNodeFlags::Expanded, false);
 	}
 
-	bool bAnyPrefix = false;
-	for (const FString& Item : InExpandedItems)
-	{
-		// check if this path is a prefix to some other path
-		// if it's not, we can early out and skip checking every child node, which is a win for very large child arrays
-		if (Item.Len() > Path.Len() && 
-			(Item[Path.Len()] == '[' || Item[Path.Len()] == '.') &&
-			Item.StartsWith(Path))
-		{
-			bAnyPrefix = true;
-			break;
-		}
-	}
-
-	if (bAnyPrefix)
+	if (InExpandedItems.AnyStartsWith(Path))
 	{
 		for (int32 NodeIndex = 0; NodeIndex < InPropertyNode->GetNumChildNodes(); ++NodeIndex)
 		{
@@ -1201,7 +1186,7 @@ static void SetExpandedItems(TSharedPtr<FPropertyNode> InPropertyNode, const TSe
 
 void SDetailsViewBase::SavePreSearchExpandedItems()
 {
-	PreSearchExpandedItems.Reset();
+	PreSearchExpandedItems.Clear();
 
 	FRootPropertyNodeList& RootPropertyNodes = GetRootNodes();
 
@@ -1219,7 +1204,7 @@ void SDetailsViewBase::SavePreSearchExpandedItems()
 		}
 	}
 
-	PreSearchExpandedCategories.Reset();
+	PreSearchExpandedCategories.Clear();
 
 	for (const TSharedRef<FDetailTreeNode>& RootNode : RootTreeNodes)
 	{
@@ -1228,7 +1213,7 @@ void SDetailsViewBase::SavePreSearchExpandedItems()
 			FDetailCategoryImpl& Category = (FDetailCategoryImpl&) RootNode.Get();
 			if (Category.ShouldBeExpanded())
 			{
-				PreSearchExpandedCategories.Add(Category.GetCategoryPathName());
+				PreSearchExpandedCategories.Insert(Category.GetCategoryPathName());
 			}
 		}
 	}
@@ -1252,7 +1237,7 @@ void SDetailsViewBase::RestorePreSearchExpandedItems()
 		}
 	}
 
-	PreSearchExpandedItems.Reset();
+	PreSearchExpandedItems.Clear();
 
 	for (const TSharedRef<FDetailTreeNode>& RootNode : RootTreeNodes)
 	{
@@ -1265,17 +1250,17 @@ void SDetailsViewBase::RestorePreSearchExpandedItems()
 		}
 	}
 
-	PreSearchExpandedCategories.Reset();
+	PreSearchExpandedCategories.Clear();
 }
 
 void SDetailsViewBase::SaveExpandedItems(TSharedRef<FPropertyNode> StartNode)
 {
 	UStruct* BestBaseStruct = StartNode->FindComplexParent()->GetBaseStructure();
 
-	TSet<FString> ExpandedPropertyItemSet;
+	FStringPrefixTree ExpandedPropertyItemSet;
 	GetExpandedItems(StartNode, ExpandedPropertyItemSet);
 
-	TArray<FString> ExpandedPropertyItems = ExpandedPropertyItemSet.Array();
+	TArray<FString> ExpandedPropertyItems = ExpandedPropertyItemSet.GetAllEntries();
 
 	// Handle spaces in expanded node names by wrapping them in quotes
 	for (FString& String : ExpandedPropertyItems)
@@ -1284,17 +1269,6 @@ void SDetailsViewBase::SaveExpandedItems(TSharedRef<FPropertyNode> StartNode)
 		String.AppendChar('"');
 	}
 
-	TArray<FString> ExpandedCustomItems = ExpandedDetailNodes.Array();
-
-	// Expanded custom items may have spaces but SetSingleLineArray doesn't support spaces (treats it as another element in the array)
-	// Append a '|' after each element instead
-	FString ExpandedCustomItemsString;
-	for (const FString& DetailNode : ExpandedDetailNodes)
-	{
-		ExpandedCustomItemsString += DetailNode;
-		ExpandedCustomItemsString += TEXT(",");
-	}
-
 	//while a valid class, and we're either the same as the base class (for multiple actors being selected and base class is AActor) OR we're not down to AActor yet)
 	for (UStruct* Struct = BestBaseStruct; Struct && ((BestBaseStruct == Struct) || (Struct != AActor::StaticClass())); Struct = Struct->GetSuperStruct())
 	{
@@ -1317,6 +1291,18 @@ void SDetailsViewBase::SaveExpandedItems(TSharedRef<FPropertyNode> StartNode)
 
 	if (DetailLayouts.Num() > 0 && BestBaseStruct)
 	{
+		TArray<FString> ExpandedCustomItems = ExpandedDetailNodes.GetAllEntries();
+
+		// Expanded custom items may have spaces but SetSingleLineArray doesn't support spaces (treats it as another element in the array)
+		// Append a ',' after each element instead
+		FString ExpandedCustomItemsString;
+		for (const FString& DetailNode : ExpandedCustomItems)
+		{
+			ExpandedCustomItemsString += DetailNode;
+			ExpandedCustomItemsString += TEXT(",");
+		}
+
+		// if there are no expanded custom items currently, save if there used to be items
 		bool bShouldSave = !ExpandedCustomItemsString.IsEmpty();
 		if (!bShouldSave)
 		{
@@ -1324,6 +1310,7 @@ void SDetailsViewBase::SaveExpandedItems(TSharedRef<FPropertyNode> StartNode)
 			GConfig->GetString(TEXT("DetailCustomWidgetExpansion"), *BestBaseStruct->GetName(), DummyExpandedCustomItemsString, GEditorPerProjectIni);
 			bShouldSave = !DummyExpandedCustomItemsString.IsEmpty();
 		}
+
 		if (bShouldSave)
 		{
 			GConfig->SetString(TEXT("DetailCustomWidgetExpansion"), *BestBaseStruct->GetName(), *ExpandedCustomItemsString, GEditorPerProjectIni);
@@ -1361,9 +1348,10 @@ void SDetailsViewBase::RestoreExpandedItems(TSharedRef<FPropertyNode> StartNode)
 		GConfig->GetSingleLineArray(TEXT("DetailPropertyExpansion"), *Struct->GetName(), DetailPropertyExpansionStrings, GEditorPerProjectIni);
 	}
 
-	TSet<FString> ExpandedPropertyItems;
-	ExpandedPropertyItems.Append(DetailPropertyExpansionStrings);
-	SetExpandedItems(StartNode, ExpandedPropertyItems, false);
+	FStringPrefixTree PrefixTree;
+	PrefixTree.InsertAll(DetailPropertyExpansionStrings);
+
+	SetExpandedItems(StartNode, PrefixTree, false);
 
 	if (BestBaseStruct)
 	{
@@ -1373,7 +1361,7 @@ void SDetailsViewBase::RestoreExpandedItems(TSharedRef<FPropertyNode> StartNode)
 		TArray<FString> ExpandedCustomItemsArray;
 		ExpandedCustomItems.ParseIntoArray(ExpandedCustomItemsArray, TEXT(","), true);
 
-		ExpandedDetailNodes.Append(ExpandedCustomItemsArray);
+		ExpandedDetailNodes.InsertAll(ExpandedCustomItemsArray);
 	}
 }
 
diff --git a/Engine/Source/Editor/PropertyEditor/Private/SDetailsViewBase.h b/Engine/Source/Editor/PropertyEditor/Private/SDetailsViewBase.h
--- a/Engine/Source/Editor/PropertyEditor/Private/SDetailsViewBase.h
+++ b/Engine/Source/Editor/PropertyEditor/Private/SDetailsViewBase.h
@@ -25,6 +25,7 @@
 #include "Widgets/Views/STableRow.h"
 #include "Widgets/Views/STableViewBase.h"
 #include "Widgets/Views/STreeView.h"
+#include "StringPrefixTree.h"
 
 class FAssetThumbnailPool;
 class FDetailCategoryImpl;
@@ -374,7 +375,7 @@ protected:
 	/** Map of nodes that are requesting an automatic expansion/collapse due to being filtered */
 	TMap< TWeakPtr<FDetailTreeNode>, bool > FilteredNodesRequestingExpansionState;
 	/** Current set of expanded detail nodes (by path) that should be saved when the details panel closes */
-	TSet<FString> ExpandedDetailNodes;
+	FStringPrefixTree ExpandedDetailNodes;
 	/** Tree view */
 	TSharedPtr<SDetailTree> DetailTree;
 	/** Root tree nodes visible in the tree */
@@ -427,8 +428,8 @@ protected:
 	/** Timer for that current node's widget animation duration. */
 	FTimerHandle AnimateNodeTimer;
 
-	TSet<FString> PreSearchExpandedItems;
-	TSet<FString> PreSearchExpandedCategories;
+	FStringPrefixTree PreSearchExpandedItems;
+	FStringPrefixTree PreSearchExpandedCategories;
 
 	/** Executed when the tree is refreshed */
 	FOnDisplayedPropertiesChanged OnDisplayedPropertiesChangedDelegate;
diff --git a/Engine/Source/Editor/PropertyEditor/Private/SStructureDetailsView.cpp b/Engine/Source/Editor/PropertyEditor/Private/SStructureDetailsView.cpp
--- a/Engine/Source/Editor/PropertyEditor/Private/SStructureDetailsView.cpp
+++ b/Engine/Source/Editor/PropertyEditor/Private/SStructureDetailsView.cpp
@@ -216,7 +216,7 @@ void SStructureDetailsView::SetStructureData(TSharedPtr<FStructOnScope> InStruct
 	RootNodesPendingKill.Add(RootNode);
 
 	RootNodes.Empty(1);
-	ExpandedDetailNodes.Empty();
+	ExpandedDetailNodes.Clear();
 
 	RootNode = MakeShareable(new FStructurePropertyNode);
 	RootNodes.Add(RootNode);
