FStructurePropertyNode::GetValuePtrOfInstance() no longer checks() if the struct instance has been removed.

This was most apparent with DataTables that used EditConditions that had the EditConditionHides metadata. After clearing the table, the rows would still refresh for a single frame, and the edit condition parser would try to check the edit condition state. However, this would end up calling into FProperty::ContainerPtrToValuePtr(nullptr), which causes a check.

Now, GetValuePtrOfInstance() early outs before making the call.

Reported on UDN.

diff --git a/Engine/Source/Editor/PropertyEditor/Private/StructurePropertyNode.h b/Engine/Source/Editor/PropertyEditor/Private/StructurePropertyNode.h
--- a/Engine/Source/Editor/PropertyEditor/Private/StructurePropertyNode.h
+++ b/Engine/Source/Editor/PropertyEditor/Private/StructurePropertyNode.h
@@ -124,11 +124,26 @@ public:
 	}
 	virtual uint8* GetValuePtrOfInstance(int32 Index, const FProperty* InProperty, const FPropertyNode* InParentNode) const override
 	{ 
-		check(0 == Index);
-		uint8* StructBaseAddress = HasValidStructData() ? StructData->GetStructMemory() : nullptr;
-		uint8* ParentBaseAddress = InParentNode ? InParentNode->GetValueAddress(StructBaseAddress, false) : nullptr;
-		return InProperty ? InProperty->ContainerPtrToValuePtr<uint8>(ParentBaseAddress) : nullptr;
+		if (InProperty == nullptr || InParentNode == nullptr)
+		{
+			return nullptr;
+		}
+
+		uint8* StructBaseAddress = GetMemoryOfInstance(Index);
+		if (StructBaseAddress == nullptr)
+		{
+			return nullptr;
+		}
+
+		uint8* ParentBaseAddress = InParentNode->GetValueAddress(StructBaseAddress, false);
+		if (ParentBaseAddress == nullptr)
+		{
+			return nullptr;
+		}
+
+		return InProperty->ContainerPtrToValuePtr<uint8>(ParentBaseAddress);
 	}
+
 	virtual TWeakObjectPtr<UObject> GetInstanceAsUObject(int32 Index) const override
 	{
 		check(0 == Index);
