Enabled FStringPrefixTree usage in details view.

Changes from previous submit:
- ExpandedCustomItems uses TStringBuilder::Join()
- FStringPrefixTree::GetAllEntries() uses a const FString& to avoid a potential memcpy crash.
- Added checks to FStringPrefixTree for empty prefixes being added to the tree.
- Fixed bug where re-adding a substring along a node split boundary (eg. "FooFoo", "FooBar", "FooBar") would cause the intermediate node to be marked as a leaf (instead of the child), added tests for it.

#rnx

diff --git a/Engine/Source/Editor/PropertyEditor/Private/SDetailsView.cpp b/Engine/Source/Editor/PropertyEditor/Private/SDetailsView.cpp
--- a/Engine/Source/Editor/PropertyEditor/Private/SDetailsView.cpp
+++ b/Engine/Source/Editor/PropertyEditor/Private/SDetailsView.cpp
@@ -896,7 +896,7 @@ void SDetailsView::PreSetObject(int32 InNewNumObjects)
 	}
 
 	RootPropertyNodes.Empty(InNewNumObjects);
-	ExpandedDetailNodes.Empty();
+	ExpandedDetailNodes.Clear();
 
 	for (int32 NewRootIndex = 0; NewRootIndex < InNewNumObjects; ++NewRootIndex)
 	{
diff --git a/Engine/Source/Editor/PropertyEditor/Private/SDetailsViewBase.cpp b/Engine/Source/Editor/PropertyEditor/Private/SDetailsViewBase.cpp
--- a/Engine/Source/Editor/PropertyEditor/Private/SDetailsViewBase.cpp
+++ b/Engine/Source/Editor/PropertyEditor/Private/SDetailsViewBase.cpp
@@ -27,16 +27,15 @@
 #include "SDetailNameArea.h"
 #include "SDetailsView.h"
 #include "PropertyEditorPermissionList.h"
-
 #include "ThumbnailRendering/ThumbnailManager.h"
 
 SDetailsViewBase::SDetailsViewBase() :
-	bHasActiveFilter(false)
+	  NumVisibleTopLevelObjectNodes(0)
 	, bIsLocked(false)
 	, bHasOpenColorPicker(false)
-	, bDisableCustomDetailLayouts( false )
-	, NumVisibleTopLevelObjectNodes(0)
+	, bDisableCustomDetailLayouts(false)
 	, bPendingCleanupTimerSet(false)
+	, bRunningDeferredActions(false)
 {
 	UDetailsConfig::Initialize();
 	UDetailsConfig::Get()->LoadEditorConfig();
@@ -627,7 +626,7 @@ void SDetailsViewBase::SaveCustomExpansionState(const FString& NodePath, bool bI
 {
 	if (bIsExpanded)
 	{
-		ExpandedDetailNodes.Add(NodePath);
+		ExpandedDetailNodes.Insert(NodePath);
 	}
 	else
 	{
@@ -837,10 +836,6 @@ void SDetailsViewBase::FilterView(const FString& InFilter)
 	ParseString.TrimStartAndEndInline();
 	ParseString.ParseIntoArray(CurrentFilterStrings, TEXT(" "), true);
 
-PRAGMA_DISABLE_DEPRECATION_WARNINGS
-	bHasActiveFilter = CurrentFilterStrings.Num() > 0;
-PRAGMA_ENABLE_DEPRECATION_WARNINGS
-	
 	CurrentFilter.FilterStrings = CurrentFilterStrings;
 
 	if (!bHadActiveFilter && CurrentFilter.FilterStrings.Num() > 0)
@@ -968,7 +963,7 @@ void SDetailsViewBase::Tick( const FGeometry& AllottedGeometry, const double InC
 		}
 
 		if (bHadDeferredActions)
-		{		
+		{
 			// Any deferred actions are likely to cause the node tree to be at least partially rebuilt
 			// Save the expansion state of existing nodes so we can expand them later
 			SaveExpandedItems(RootPropertyNode.ToSharedRef());
@@ -991,6 +986,9 @@ void SDetailsViewBase::Tick( const FGeometry& AllottedGeometry, const double InC
 
 	if (bHadDeferredActions)
 	{
+		bRunningDeferredActions = true;
+		ON_SCOPE_EXIT { bRunningDeferredActions = false; };
+
 		TArray<FSimpleDelegate> DeferredActionsCopy;
 		
 		do
@@ -1136,7 +1134,7 @@ void SDetailsViewBase::Tick( const FGeometry& AllottedGeometry, const double InC
 * @param InPropertyNode			The node to get expanded items from
 * @param OutExpandedItems	List of expanded items that were found
 */
-static void GetExpandedItems(TSharedPtr<FPropertyNode> InPropertyNode, TSet<FString>& OutExpandedItems)
+static void GetExpandedItems(TSharedPtr<FPropertyNode> InPropertyNode, FStringPrefixTree& OutExpandedItems)
 {
 	if (InPropertyNode->HasNodeFlags(EPropertyNodeFlags::Expanded))
 	{
@@ -1145,7 +1143,7 @@ static void GetExpandedItems(TSharedPtr<FPropertyNode> InPropertyNode, TSet<FStr
 		Path.Empty(128);
 		InPropertyNode->GetQualifiedName(Path, bWithArrayIndex);
 
-		OutExpandedItems.Add(Path);
+		OutExpandedItems.Insert(Path);
 	}
 
 	for (int32 ChildIndex = 0; ChildIndex < InPropertyNode->GetNumChildNodes(); ++ChildIndex)
@@ -1160,7 +1158,7 @@ static void GetExpandedItems(TSharedPtr<FPropertyNode> InPropertyNode, TSet<FStr
 * @param InNode			The node to set expanded items on
 * @param OutExpandedItems	List of expanded items to set
 */
-static void SetExpandedItems(TSharedPtr<FPropertyNode> InPropertyNode, const TSet<FString>& InExpandedItems, bool bCollapseRest)
+static void SetExpandedItems(TSharedPtr<FPropertyNode> InPropertyNode, const FStringPrefixTree& InExpandedItems, bool bCollapseRest)
 {
 	const bool bWithArrayIndex = true;
 	FString Path;
@@ -1176,21 +1174,7 @@ static void SetExpandedItems(TSharedPtr<FPropertyNode> InPropertyNode, const TSe
 		InPropertyNode->SetNodeFlags(EPropertyNodeFlags::Expanded, false);
 	}
 
-	bool bAnyPrefix = false;
-	for (const FString& Item : InExpandedItems)
-	{
-		// check if this path is a prefix to some other path
-		// if it's not, we can early out and skip checking every child node, which is a win for very large child arrays
-		if (Item.Len() > Path.Len() && 
-			(Item[Path.Len()] == '[' || Item[Path.Len()] == '.') &&
-			Item.StartsWith(Path))
-		{
-			bAnyPrefix = true;
-			break;
-		}
-	}
-
-	if (bAnyPrefix)
+	if (InExpandedItems.AnyStartsWith(Path))
 	{
 		for (int32 NodeIndex = 0; NodeIndex < InPropertyNode->GetNumChildNodes(); ++NodeIndex)
 		{
@@ -1201,7 +1185,7 @@ static void SetExpandedItems(TSharedPtr<FPropertyNode> InPropertyNode, const TSe
 
 void SDetailsViewBase::SavePreSearchExpandedItems()
 {
-	PreSearchExpandedItems.Reset();
+	PreSearchExpandedItems.Clear();
 
 	FRootPropertyNodeList& RootPropertyNodes = GetRootNodes();
 
@@ -1219,7 +1203,7 @@ void SDetailsViewBase::SavePreSearchExpandedItems()
 		}
 	}
 
-	PreSearchExpandedCategories.Reset();
+	PreSearchExpandedCategories.Clear();
 
 	for (const TSharedRef<FDetailTreeNode>& RootNode : RootTreeNodes)
 	{
@@ -1228,7 +1212,7 @@ void SDetailsViewBase::SavePreSearchExpandedItems()
 			FDetailCategoryImpl& Category = (FDetailCategoryImpl&) RootNode.Get();
 			if (Category.ShouldBeExpanded())
 			{
-				PreSearchExpandedCategories.Add(Category.GetCategoryPathName());
+				PreSearchExpandedCategories.Insert(Category.GetCategoryPathName());
 			}
 		}
 	}
@@ -1252,7 +1236,7 @@ void SDetailsViewBase::RestorePreSearchExpandedItems()
 		}
 	}
 
-	PreSearchExpandedItems.Reset();
+	PreSearchExpandedItems.Clear();
 
 	for (const TSharedRef<FDetailTreeNode>& RootNode : RootTreeNodes)
 	{
@@ -1265,17 +1249,25 @@ void SDetailsViewBase::RestorePreSearchExpandedItems()
 		}
 	}
 
-	PreSearchExpandedCategories.Reset();
+	PreSearchExpandedCategories.Clear();
 }
 
 void SDetailsViewBase::SaveExpandedItems(TSharedRef<FPropertyNode> StartNode)
 {
+	if (bRunningDeferredActions)
+	{
+		// Deferred actions can manipulate the tree, eg. Add Item
+		// However, there's no reason for us to save expanded items during deferred actions, 
+		// because the expansion state was already stored at the beginning and will be restored afterwards.
+		return;
+	}
+
 	UStruct* BestBaseStruct = StartNode->FindComplexParent()->GetBaseStructure();
 
-	TSet<FString> ExpandedPropertyItemSet;
+	FStringPrefixTree ExpandedPropertyItemSet;
 	GetExpandedItems(StartNode, ExpandedPropertyItemSet);
 
-	TArray<FString> ExpandedPropertyItems = ExpandedPropertyItemSet.Array();
+	TArray<FString> ExpandedPropertyItems = ExpandedPropertyItemSet.GetAllEntries();
 
 	// Handle spaces in expanded node names by wrapping them in quotes
 	for (FString& String : ExpandedPropertyItems)
@@ -1284,17 +1276,6 @@ void SDetailsViewBase::SaveExpandedItems(TSharedRef<FPropertyNode> StartNode)
 		String.AppendChar('"');
 	}
 
-	TArray<FString> ExpandedCustomItems = ExpandedDetailNodes.Array();
-
-	// Expanded custom items may have spaces but SetSingleLineArray doesn't support spaces (treats it as another element in the array)
-	// Append a '|' after each element instead
-	FString ExpandedCustomItemsString;
-	for (const FString& DetailNode : ExpandedDetailNodes)
-	{
-		ExpandedCustomItemsString += DetailNode;
-		ExpandedCustomItemsString += TEXT(",");
-	}
-
 	//while a valid class, and we're either the same as the base class (for multiple actors being selected and base class is AActor) OR we're not down to AActor yet)
 	for (UStruct* Struct = BestBaseStruct; Struct && ((BestBaseStruct == Struct) || (Struct != AActor::StaticClass())); Struct = Struct->GetSuperStruct())
 	{
@@ -1317,22 +1298,39 @@ void SDetailsViewBase::SaveExpandedItems(TSharedRef<FPropertyNode> StartNode)
 
 	if (DetailLayouts.Num() > 0 && BestBaseStruct)
 	{
-		bool bShouldSave = !ExpandedCustomItemsString.IsEmpty();
+		TArray<FString> ExpandedCustomItems = ExpandedDetailNodes.GetAllEntries();
+
+		// Expanded custom items may have spaces but SetSingleLineArray doesn't support spaces (treats it as another element in the array)
+		// Append a ',' after each element instead
+		TStringBuilder<256> ExpandedCustomItemsBuilder;
+		ExpandedCustomItemsBuilder.Join(ExpandedCustomItems, TEXT(","));
+
+		// if there are no expanded custom items currently, save if there used to be items
+		bool bShouldSave = ExpandedCustomItemsBuilder.Len() > 0;
 		if (!bShouldSave)
 		{
 			FString DummyExpandedCustomItemsString;
 			GConfig->GetString(TEXT("DetailCustomWidgetExpansion"), *BestBaseStruct->GetName(), DummyExpandedCustomItemsString, GEditorPerProjectIni);
 			bShouldSave = !DummyExpandedCustomItemsString.IsEmpty();
 		}
+
 		if (bShouldSave)
 		{
-			GConfig->SetString(TEXT("DetailCustomWidgetExpansion"), *BestBaseStruct->GetName(), *ExpandedCustomItemsString, GEditorPerProjectIni);
+			GConfig->SetString(TEXT("DetailCustomWidgetExpansion"), *BestBaseStruct->GetName(), *ExpandedCustomItemsBuilder, GEditorPerProjectIni);
 		}
 	}
 }
 
 void SDetailsViewBase::RestoreAllExpandedItems()
 {
+	if (bRunningDeferredActions)
+	{
+		// Deferred actions can manipulate the tree, eg. Add Item
+		// However, there's no reason for us to restore expanded items during deferred actions, 
+		// because the expansion state was already stored at the beginning and will be restored afterwards.
+		return;
+	}
+
 	for (TSharedPtr<FComplexPropertyNode>& RootPropertyNode : GetRootNodes())
 	{
 		check(RootPropertyNode.IsValid());
@@ -1352,6 +1350,14 @@ void SDetailsViewBase::RestoreAllExpandedItems()
 
 void SDetailsViewBase::RestoreExpandedItems(TSharedRef<FPropertyNode> StartNode)
 {
+	if (bRunningDeferredActions)
+	{
+		// Deferred actions can manipulate the tree, eg. Add Item
+		// However, there's no reason for us to restore expanded items during deferred actions,
+		// because the expansion state was already stored at the beginning and will be restored afterwards.
+		return;
+	}
+
 	UStruct* BestBaseStruct = StartNode->FindComplexParent()->GetBaseStructure();
 
 	//while a valid class, and we're either the same as the base class (for multiple actors being selected and base class is AActor) OR we're not down to AActor yet)
@@ -1361,9 +1367,10 @@ void SDetailsViewBase::RestoreExpandedItems(TSharedRef<FPropertyNode> StartNode)
 		GConfig->GetSingleLineArray(TEXT("DetailPropertyExpansion"), *Struct->GetName(), DetailPropertyExpansionStrings, GEditorPerProjectIni);
 	}
 
-	TSet<FString> ExpandedPropertyItems;
-	ExpandedPropertyItems.Append(DetailPropertyExpansionStrings);
-	SetExpandedItems(StartNode, ExpandedPropertyItems, false);
+	FStringPrefixTree PrefixTree;
+	PrefixTree.InsertAll(DetailPropertyExpansionStrings);
+
+	SetExpandedItems(StartNode, PrefixTree, false);
 
 	if (BestBaseStruct)
 	{
@@ -1373,7 +1380,7 @@ void SDetailsViewBase::RestoreExpandedItems(TSharedRef<FPropertyNode> StartNode)
 		TArray<FString> ExpandedCustomItemsArray;
 		ExpandedCustomItems.ParseIntoArray(ExpandedCustomItemsArray, TEXT(","), true);
 
-		ExpandedDetailNodes.Append(ExpandedCustomItemsArray);
+		ExpandedDetailNodes.InsertAll(ExpandedCustomItemsArray);
 	}
 }
 
diff --git a/Engine/Source/Editor/PropertyEditor/Private/SDetailsViewBase.h b/Engine/Source/Editor/PropertyEditor/Private/SDetailsViewBase.h
--- a/Engine/Source/Editor/PropertyEditor/Private/SDetailsViewBase.h
+++ b/Engine/Source/Editor/PropertyEditor/Private/SDetailsViewBase.h
@@ -20,6 +20,7 @@
 #include "PropertyNode.h"
 #include "PropertyPath.h"
 #include "PropertyRowGenerator.h"
+#include "StringPrefixTree.h"
 #include "Widgets/Layout/SSplitter.h"
 #include "Widgets/SWindow.h"
 #include "Widgets/Views/STableRow.h"
@@ -373,8 +374,6 @@ protected:
 	TArray< TSharedPtr<FDetailLayoutBuilderImpl> > DetailLayoutsPendingDelete;
 	/** Map of nodes that are requesting an automatic expansion/collapse due to being filtered */
 	TMap< TWeakPtr<FDetailTreeNode>, bool > FilteredNodesRequestingExpansionState;
-	/** Current set of expanded detail nodes (by path) that should be saved when the details panel closes */
-	TSet<FString> ExpandedDetailNodes;
 	/** Tree view */
 	TSharedPtr<SDetailTree> DetailTree;
 	/** Root tree nodes visible in the tree */
@@ -393,15 +392,8 @@ protected:
 	mutable FOnFinishedChangingProperties OnFinishedChangingPropertiesDelegate;
 	/** Container for passing around column size data to rows in the tree (each row has a splitter which can affect the column size)*/
 	FDetailColumnSizeData ColumnSizeData;
-	/** True if there is an active filter (text in the filter box) */
-	UE_DEPRECATED(4.26, "Use HasActiveSearch function instead.")
-	bool bHasActiveFilter;
-	/** True if this property view is currently locked (I.E The objects being observed are not changed automatically due to user selection)*/
-	bool bIsLocked;
 	/** The property node that the color picker is currently editing. */
 	TWeakPtr<FPropertyNode> ColorPropertyNode;
-	/** Whether or not this instance of the details view opened a color picker and it is not closed yet */
-	bool bHasOpenColorPicker;
 	/** Settings for this view */
 	TSharedPtr<IPropertyUtilities> PropertyUtilities;
 	/** Gets internal utilities for generating property layouts. */
@@ -427,15 +419,14 @@ protected:
 	/** Timer for that current node's widget animation duration. */
 	FTimerHandle AnimateNodeTimer;
 
-	TSet<FString> PreSearchExpandedItems;
-	TSet<FString> PreSearchExpandedCategories;
+	/** Current set of expanded detail nodes (by path) that should be saved when the details panel closes */
+	FStringPrefixTree ExpandedDetailNodes;
+	FStringPrefixTree PreSearchExpandedItems;
+	FStringPrefixTree PreSearchExpandedCategories;
 
 	/** Executed when the tree is refreshed */
 	FOnDisplayedPropertiesChanged OnDisplayedPropertiesChangedDelegate;
 
-	/** True if we want to skip generation of custom layouts for displayed object */
-	bool bDisableCustomDetailLayouts;
-
 	int32 NumVisibleTopLevelObjectNodes;
 
 	/** Used to refresh the tree when the allow list filter changes */
@@ -444,14 +435,25 @@ protected:
 
 	/** Delegate for overriding the show modified filter */
 	FSimpleDelegate CustomFilterDelegate;
+	FText CustomFilterLabel;
+
+	/** True if this property view is currently locked (I.E The objects being observed are not changed automatically due to user selection)*/
+	bool bIsLocked : 1;
+
+	/** Whether or not this instance of the details view opened a color picker and it is not closed yet */
+	bool bHasOpenColorPicker : 1;
+
+	/** True if we want to skip generation of custom layouts for displayed object */
+	bool bDisableCustomDetailLayouts : 1;
 
 	/** When overriding show modified, you can't use the filter state to determine what is overridden anymore. Use this variable instead. */
-	bool bCustomFilterActive;
+	bool bCustomFilterActive : 1;
 
 	/** Timer has already been set to be run next tick */
-	bool bPendingCleanupTimerSet;
+	bool bPendingCleanupTimerSet : 1;
 
-	FText CustomFilterLabel;
+	/** Are we currently running deferred actions? */
+	bool bRunningDeferredActions : 1;
 
 	mutable TSharedPtr<FEditConditionParser> EditConditionParser;
 	
diff --git a/Engine/Source/Editor/PropertyEditor/Private/SStructureDetailsView.cpp b/Engine/Source/Editor/PropertyEditor/Private/SStructureDetailsView.cpp
--- a/Engine/Source/Editor/PropertyEditor/Private/SStructureDetailsView.cpp
+++ b/Engine/Source/Editor/PropertyEditor/Private/SStructureDetailsView.cpp
@@ -216,7 +216,7 @@ void SStructureDetailsView::SetStructureData(TSharedPtr<FStructOnScope> InStruct
 	RootNodesPendingKill.Add(RootNode);
 
 	RootNodes.Empty(1);
-	ExpandedDetailNodes.Empty();
+	ExpandedDetailNodes.Clear();
 
 	RootNode = MakeShareable(new FStructurePropertyNode);
 	RootNodes.Add(RootNode);
diff --git a/Engine/Source/Editor/PropertyEditor/Private/StringPrefixTree.cpp b/Engine/Source/Editor/PropertyEditor/Private/StringPrefixTree.cpp
--- a/Engine/Source/Editor/PropertyEditor/Private/StringPrefixTree.cpp
+++ b/Engine/Source/Editor/PropertyEditor/Private/StringPrefixTree.cpp
@@ -2,6 +2,11 @@
 
 #include "StringPrefixTree.h"
 
+FStringPrefixTree::~FStringPrefixTree()
+{
+	Clear();
+}
+
 void FStringPrefixTree::Insert(FStringView ToAdd)
 {
 	if (ToAdd.IsEmpty())
@@ -94,6 +99,13 @@ FString FStringPrefixTree::DumpToString() const
 	return Builder.ToString();
 }
 
+FStringPrefixTree::FNode::~FNode()
+{
+	Children.Reset();
+	Prefix.Reset();
+	IsLeaf = false;
+}
+
 static int32 GetCommonPrefixLength(FStringView A, FStringView B)
 {
 	int32 CharIdx = 0;
@@ -118,7 +130,6 @@ void FStringPrefixTree::FNode::Insert(FStringView ToAdd)
 		if (CharIdx > 0)
 		{
 			// the paths share a common prefix, there can be only one such node
-
 			bAdded = true;
 
 			if (CharIdx == Child.Prefix.Len())
@@ -127,7 +138,7 @@ void FStringPrefixTree::FNode::Insert(FStringView ToAdd)
 				if (CharIdx == ToAdd.Len())
 				{
 					// the paths are identical, we don't need to add a node, just mark this as a leaf node
-					IsLeaf = true;
+					Child.IsLeaf = true;
 					break;
 				}
 
@@ -145,10 +156,14 @@ void FStringPrefixTree::FNode::Insert(FStringView ToAdd)
 				
 				FNode Split;
 				Split.Prefix = Child.Prefix.RightChop(CharIdx);
+				check(Split.Prefix.Len() > 0);
+
 				Split.Children = MoveTemp(Child.Children);
 				Split.IsLeaf = Child.IsLeaf;
 
 				Child.Prefix.LeftInline(CharIdx);
+				check(Child.Prefix.Len() > 0);
+
 				Child.Children.Reset();
 				Child.Children.Add(MoveTemp(Split));
 
@@ -159,6 +174,8 @@ void FStringPrefixTree::FNode::Insert(FStringView ToAdd)
 					// Child -> "Fo", Split -> "o", Added -> "z"
 					FNode& AddedNode = Child.Children.AddDefaulted_GetRef();
 					AddedNode.Prefix = ToAdd.RightChop(CharIdx);
+					check(AddedNode.Prefix.Len() > 0);
+
 					AddedNode.IsLeaf = true;
 				
 					Child.IsLeaf = false;
@@ -179,6 +196,8 @@ void FStringPrefixTree::FNode::Insert(FStringView ToAdd)
 	{
 		FNode& NewChild = Children.AddDefaulted_GetRef();
 		NewChild.Prefix = ToAdd;
+		check(NewChild.Prefix.Len() > 0);
+
 		NewChild.IsLeaf = true;
 	}
 }
@@ -254,7 +273,7 @@ bool FStringPrefixTree::FNode::AnyStartsWith(FStringView ToFind) const
 	return false;
 }
 
-void FStringPrefixTree::FNode::GetAllEntries(FStringView CurrentPrefix, TArray<FString>& OutEntries) const 
+void FStringPrefixTree::FNode::GetAllEntries(const FString& CurrentPrefix, TArray<FString>& OutEntries) const 
 {
 	FString NewPrefix = CurrentPrefix + Prefix;
 
diff --git a/Engine/Source/Editor/PropertyEditor/Private/StringPrefixTree.h b/Engine/Source/Editor/PropertyEditor/Private/StringPrefixTree.h
--- a/Engine/Source/Editor/PropertyEditor/Private/StringPrefixTree.h
+++ b/Engine/Source/Editor/PropertyEditor/Private/StringPrefixTree.h
@@ -26,6 +26,9 @@ class FStringPrefixTree
 {
 public:
 
+	FStringPrefixTree() = default;
+	~FStringPrefixTree();
+
 	/**
 	 * Insert the given string into the tree.
 	 */
@@ -93,6 +96,9 @@ private:
 		TArray<FNode> Children;
 		bool IsLeaf = false;
 
+		FNode() = default;
+		~FNode();
+
 		void Insert(FStringView ToAdd);
 		void Remove(FStringView ToRemove);
 		bool Contains(FStringView ToFind) const;
@@ -100,7 +106,7 @@ private:
 		int32 Size() const;
 
 		int32 NumNodes() const;
-		void GetAllEntries(FStringView CurrentPrefix, TArray<FString>& OutEntries) const;
+		void GetAllEntries(const FString& CurrentPrefix, TArray<FString>& OutEntries) const;
 		void DumpToString(FStringBuilderBase& Builder, TArray<bool, TInlineAllocator<32>>& DrawLines, bool bLastChild) const;
 
 	private:
diff --git a/Engine/Source/Editor/PropertyEditor/Private/StringPrefixTreeTests.cpp b/Engine/Source/Editor/PropertyEditor/Private/StringPrefixTreeTests.cpp
--- a/Engine/Source/Editor/PropertyEditor/Private/StringPrefixTreeTests.cpp
+++ b/Engine/Source/Editor/PropertyEditor/Private/StringPrefixTreeTests.cpp
@@ -14,16 +14,29 @@ bool FStringPrefixTreeTests_Insert::RunTest(const FString& Parameters)
 	Tree.Insert(TEXT("Foobar"));
 
 	TestEqual(TEXT("One"), Tree.Size(), 1);
+	TestEqual(TEXT("One"), Tree.NumNodes(), 2);
 	TestTrue(TEXT("One"), Tree.Contains(TEXT("Foobar")));
 
+
 	Tree.Insert(TEXT("Foo"));
 
 	TestEqual(TEXT("Two"), Tree.Size(), 2);
+	TestEqual(TEXT("Two"), Tree.NumNodes(), 3);
 	TestTrue(TEXT("Two"), Tree.Contains(TEXT("Foo")));
 	TestTrue(TEXT("Two"), Tree.Contains(TEXT("Foobar")));
 	TestFalse(TEXT("Two"), Tree.Contains(TEXT("Fo")));
 	TestFalse(TEXT("Two"), Tree.Contains(TEXT("Foob")));
 
+	// check that we don't add nodes if we add another identical entry
+	Tree.Insert(TEXT("Foo"));
+	TestEqual(TEXT("Two"), Tree.Size(), 2);
+	TestEqual(TEXT("Two"), Tree.NumNodes(), 3);
+
+	// check that we don't add nodes if we add another identical entry
+	Tree.Insert(TEXT("Foobar"));
+	TestEqual(TEXT("Two"), Tree.Size(), 2);
+	TestEqual(TEXT("Two"), Tree.NumNodes(), 3);
+
 	// check that the same holds if we build the tree in the opposite way
 	// ie. Foo, Foobar
 	Tree.Clear();
@@ -53,11 +66,17 @@ bool FStringPrefixTreeTests_Insert::RunTest(const FString& Parameters)
 	Tree.Insert(TEXT("Foobar.Foo"));
 	
 	TestEqual(TEXT("Four"), Tree.Size(), 4);
+	TestEqual(TEXT("Four"), Tree.NumNodes(), 6);
 	TestTrue(TEXT("Four"), Tree.Contains(TEXT("Foobar.Foo")));
 	TestTrue(TEXT("Four"), Tree.Contains(TEXT("Foobaz")));
 	TestTrue(TEXT("Four"), Tree.Contains(TEXT("Foobar")));
 	TestFalse(TEXT("Four"), Tree.Contains(TEXT("Foobar.")));
 
+	Tree.Insert(TEXT("Foobar"));
+	
+	TestEqual(TEXT("Four"), Tree.Size(), 4);
+	TestEqual(TEXT("Four"), Tree.NumNodes(), 6);
+
 	Tree.Insert(TEXT("Bar"));
 	Tree.Insert(TEXT("Bar.Foo"));
 
@@ -68,6 +87,44 @@ bool FStringPrefixTreeTests_Insert::RunTest(const FString& Parameters)
 	return true;
 }
 
+IMPLEMENT_SIMPLE_AUTOMATION_TEST(FStringPrefixTreeTests_Fuzz, "StringPrefixTree.Fuzz", EAutomationTestFlags::EditorContext | EAutomationTestFlags::EngineFilter)
+
+bool FStringPrefixTreeTests_Fuzz::RunTest(const FString& Parameters)
+{
+	const int32 MAX_STRINGS = 1000;
+
+	TArray<FString> Added;
+	Added.Reserve(MAX_STRINGS);
+
+	FStringPrefixTree Tree;
+
+	FString Current;
+	Current.Reserve(MAX_STRINGS);
+
+	for (int32 Idx = 0; Idx < MAX_STRINGS; ++Idx)
+	{
+		TCHAR ToAdd = TEXT('A') + (TCHAR) FMath::RandRange(0, 25);
+		
+		int32 Where = FMath::RandRange(0, Current.Len());
+		Current.InsertAt(Where, ToAdd);
+
+		Tree.Insert(Current);
+		Added.Add(Current);
+
+		TestEqual(TEXT("Size"), Tree.Size(), Added.Num());
+	}
+
+	for (const FString& Str : Added)
+	{
+		TestTrue(TEXT("Contains"), Tree.Contains(Str));
+	}
+	
+	TArray<FString> Actual = Tree.GetAllEntries();
+	TestEqual(TEXT("GetAllEntries"), Actual.Num(), Tree.Size());
+
+	return true;
+}
+
 IMPLEMENT_SIMPLE_AUTOMATION_TEST(FStringPrefixTreeTests_Contains, "StringPrefixTree.Contains", EAutomationTestFlags::EditorContext | EAutomationTestFlags::EngineFilter)
 
 bool FStringPrefixTreeTests_Contains::RunTest(const FString& Parameters)
