Material list improvements:
* Added tooltips to Select material's texture in asset browser button in order to help understanding what this is about when there's no texture referenced by the material
* Disabled the button when there's no texture, since it's basically useless then

diff --git a/Engine/Source/Editor/PropertyEditor/Private/MaterialList.cpp b/Engine/Source/Editor/PropertyEditor/Private/MaterialList.cpp
--- a/Engine/Source/Editor/PropertyEditor/Private/MaterialList.cpp
+++ b/Engine/Source/Editor/PropertyEditor/Private/MaterialList.cpp
@@ -224,6 +224,7 @@ TSharedRef<SWidget> FMaterialItemView::CreateValueContent(IDetailLayoutBuilder&
 							.VAlign(VAlign_Center)
 							.IsEnabled( this, &FMaterialItemView::IsTexturesMenuEnabled )
 							.Visibility(bShowUsedTextures ? EVisibility::Visible : EVisibility::Hidden)
+							.ToolTipText(this, &FMaterialItemView::GetTexturesMenuToolTipText)
 							.ButtonContent()
 							[
 								SNew(SImage)
@@ -358,7 +359,37 @@ TSharedRef<SWidget> FMaterialItemView::MakeBrowseNaniteOverrideMaterialButton()
 
 bool FMaterialItemView::IsTexturesMenuEnabled() const
 {
-	return MaterialItem.Material.Get() != NULL;
+	if (UMaterialInterface* Material = MaterialItem.Material.Get())
+	{
+		// Don't enable the menu unless there are textures, otherwise, the user might want to think the button is broken while really, it just has nothing to display
+		TArray<UTexture*> Textures;
+		Material->GetUsedTextures(Textures, EMaterialQualityLevel::Num, false, ERHIFeatureLevel::Num, true);
+		return !Textures.IsEmpty();
+	}
+	return false;
+}
+
+FText FMaterialItemView::GetTexturesMenuToolTipText() const
+{
+	if (UMaterialInterface* Material = MaterialItem.Material.Get())
+	{
+		// Don't enable the menu unless there are textures, otherwise, the user might want to think the button is broken while really, it just has nothing to display
+		TArray<UTexture*> Textures;
+		Material->GetUsedTextures(Textures, EMaterialQualityLevel::Num, false, ERHIFeatureLevel::Num, true);
+
+		FFormatNamedArguments Arguments;
+		Arguments.Add(TEXT("MaterialName"), FText::AsCultureInvariant(Material->GetName()));
+		if (Textures.IsEmpty())
+		{
+			return LOCTEXT("GetTexturesMenuToolTipText_NoTexture", "Find the material's textures in the content browser (no texture)");
+		}
+		else
+		{
+			return FText::Format(LOCTEXT("GetTexturesMenuToolTipText_MaterialName", "Find {MaterialName}'s textures in the content browser"), Arguments);
+		}
+	}
+
+	return LOCTEXT("GetTexturesMenuToolTipText_Default", "Find the material's textures in the content browser");
 }
 
 TSharedRef<SWidget> FMaterialItemView::OnGetTexturesMenuForMaterial()
diff --git a/Engine/Source/Editor/PropertyEditor/Public/MaterialList.h b/Engine/Source/Editor/PropertyEditor/Public/MaterialList.h
--- a/Engine/Source/Editor/PropertyEditor/Public/MaterialList.h
+++ b/Engine/Source/Editor/PropertyEditor/Public/MaterialList.h
@@ -222,6 +222,11 @@ private:
 	 */
 	TSharedRef<SWidget> MakeBrowseNaniteOverrideMaterialButton() const;
 
+	/**
+	 * @return The tool tip for the textures menuuuuu
+	 */
+	FText GetTexturesMenuToolTipText() const;
+
 	/**
 	 * On Get object path handler
 	 */
