Fixes for c++ headerunits/modules.
* Adding forward declarations
* Adding "inline" in front of const/constexpr variables in headers
* Adding includes
* A few moves of ctor/dtor to cpp file

diff --git a/Engine/Plugins/EnhancedInput/Source/EnhancedInput/Public/InputMappingQuery.h b/Engine/Plugins/EnhancedInput/Source/EnhancedInput/Public/InputMappingQuery.h
--- a/Engine/Plugins/EnhancedInput/Source/EnhancedInput/Public/InputMappingQuery.h
+++ b/Engine/Plugins/EnhancedInput/Source/EnhancedInput/Public/InputMappingQuery.h
@@ -56,9 +56,9 @@ ENUM_CLASS_FLAGS(EMappingQueryIssue);
 // Useful mapping query issue collections.
 namespace DefaultMappingIssues
 {
-constexpr EMappingQueryIssue NoCollisions = EMappingQueryIssue::HidesExistingMapping | EMappingQueryIssue::HiddenByExistingMapping | EMappingQueryIssue::CollisionWithMappingInSameContext;
-constexpr EMappingQueryIssue TypeMismatch = EMappingQueryIssue::ForcesTypePromotion | EMappingQueryIssue::ForcesTypeDemotion;
-constexpr EMappingQueryIssue StandardFatal = EMappingQueryIssue::ReservedByAction | NoCollisions | TypeMismatch;	// Default fatal value for QueryMapKeyIn... function calls.
+inline constexpr EMappingQueryIssue NoCollisions = EMappingQueryIssue::HidesExistingMapping | EMappingQueryIssue::HiddenByExistingMapping | EMappingQueryIssue::CollisionWithMappingInSameContext;
+inline constexpr EMappingQueryIssue TypeMismatch = EMappingQueryIssue::ForcesTypePromotion | EMappingQueryIssue::ForcesTypeDemotion;
+inline constexpr EMappingQueryIssue StandardFatal = EMappingQueryIssue::ReservedByAction | NoCollisions | TypeMismatch;	// Default fatal value for QueryMapKeyIn... function calls.
 }
 
 class UInputMappingContext;
diff --git a/Engine/Plugins/Experimental/CommonUI/Source/CommonInput/Public/CommonInputModeTypes.h b/Engine/Plugins/Experimental/CommonUI/Source/CommonInput/Public/CommonInputModeTypes.h
--- a/Engine/Plugins/Experimental/CommonUI/Source/CommonInput/Public/CommonInputModeTypes.h
+++ b/Engine/Plugins/Experimental/CommonUI/Source/CommonInput/Public/CommonInputModeTypes.h
@@ -2,6 +2,8 @@
 
 #pragma once
 
+#include "CoreTypes.h"
+
 UENUM(BlueprintType)
 enum class ECommonInputMode : uint8
 {
diff --git a/Engine/Plugins/Experimental/CommonUI/Source/CommonUI/Public/ICommonUIModule.h b/Engine/Plugins/Experimental/CommonUI/Source/CommonUI/Public/ICommonUIModule.h
--- a/Engine/Plugins/Experimental/CommonUI/Source/CommonUI/Public/ICommonUIModule.h
+++ b/Engine/Plugins/Experimental/CommonUI/Source/CommonUI/Public/ICommonUIModule.h
@@ -4,6 +4,7 @@
 
 struct FStreamableManager;
 class UCommonUISettings;
+typedef int32 TAsyncLoadPriority;
 
 #if WITH_EDITOR
 class UCommonUIEditorSettings;
diff --git a/Engine/Plugins/Experimental/GameFeatures/Source/GameFeatures/Public/GameFeatureTypesFwd.h b/Engine/Plugins/Experimental/GameFeatures/Source/GameFeatures/Public/GameFeatureTypesFwd.h
--- a/Engine/Plugins/Experimental/GameFeatures/Source/GameFeatures/Public/GameFeatureTypesFwd.h
+++ b/Engine/Plugins/Experimental/GameFeatures/Source/GameFeatures/Public/GameFeatureTypesFwd.h
@@ -2,6 +2,9 @@
 
 #pragma once
 
+#include "CoreTypes.h"
+
+class FString;
 enum class EGameFeaturePluginProtocol : uint8;
 
 namespace GameFeaturePluginStatePrivate
diff --git a/Engine/Plugins/Online/OnlineSubsystem/Source/Public/Interfaces/OnlineCodeRedemptionInterface.h b/Engine/Plugins/Online/OnlineSubsystem/Source/Public/Interfaces/OnlineCodeRedemptionInterface.h
--- a/Engine/Plugins/Online/OnlineSubsystem/Source/Public/Interfaces/OnlineCodeRedemptionInterface.h
+++ b/Engine/Plugins/Online/OnlineSubsystem/Source/Public/Interfaces/OnlineCodeRedemptionInterface.h
@@ -4,6 +4,7 @@
 
 #include "CoreMinimal.h"
 
+class FUniqueNetId;
 struct FOnlineError;
 
 class FCodeLockInfo
diff --git a/Engine/Plugins/Online/OnlineSubsystem/Source/Public/Interfaces/OnlineSessionDelegates.h b/Engine/Plugins/Online/OnlineSubsystem/Source/Public/Interfaces/OnlineSessionDelegates.h
--- a/Engine/Plugins/Online/OnlineSubsystem/Source/Public/Interfaces/OnlineSessionDelegates.h
+++ b/Engine/Plugins/Online/OnlineSubsystem/Source/Public/Interfaces/OnlineSessionDelegates.h
@@ -2,6 +2,8 @@
 
 #pragma once
 
+#include "Delegates/Delegate.h"
+
 /**
  * Delegate fired when a session create request has completed
  *
diff --git a/Engine/Plugins/Runtime/StateTree/Source/StateTreeModule/Public/StateTreeTypes.h b/Engine/Plugins/Runtime/StateTree/Source/StateTreeModule/Public/StateTreeTypes.h
--- a/Engine/Plugins/Runtime/StateTree/Source/StateTreeModule/Public/StateTreeTypes.h
+++ b/Engine/Plugins/Runtime/StateTree/Source/StateTreeModule/Public/StateTreeTypes.h
@@ -96,9 +96,9 @@ enum class EStateTreeStateType : uint8
 
 namespace UE::StateTree
 {
-	constexpr int32 MaxConditionIndent = 4;
+	inline constexpr int32 MaxConditionIndent = 4;
 
-	const FName SchemaTag(TEXT("Schema"));
+	inline const FName SchemaTag(TEXT("Schema"));
 }; // UE::StateTree
 
 
diff --git a/Engine/Source/Developer/TargetPlatform/Public/Interfaces/ITargetPlatform.h b/Engine/Source/Developer/TargetPlatform/Public/Interfaces/ITargetPlatform.h
--- a/Engine/Source/Developer/TargetPlatform/Public/Interfaces/ITargetPlatform.h
+++ b/Engine/Source/Developer/TargetPlatform/Public/Interfaces/ITargetPlatform.h
@@ -132,34 +132,34 @@ enum class EOfflineBVHMode
 namespace ETargetPlatformReadyStatus
 {
 	/** Ready */
-	const int32 Ready = 0;
+	inline const int32 Ready = 0;
 
 	/** SDK Not Found*/
-	const int32 SDKNotFound = 1;
+	inline const int32 SDKNotFound = 1;
 
 	/** Code Build Not Supported */
-	const int32 CodeUnsupported = 2;
+	inline const int32 CodeUnsupported = 2;
 
 	/** Plugins Not Supported */
-	const int32 PluginsUnsupported = 4;
+	inline const int32 PluginsUnsupported = 4;
 
 	/** Signing Key Not Found */
-	const int32 SigningKeyNotFound = 8;
+	inline const int32 SigningKeyNotFound = 8;
 
 	/** Provision Not Found */
-	const int32 ProvisionNotFound = 16;
+	inline const int32 ProvisionNotFound = 16;
 
 	/** Manifest Not Found */
-	const int32 ManifestNotFound = 32;
+	inline const int32 ManifestNotFound = 32;
 
 	/** Remote Server Name Empty */
-	const int32 RemoveServerNameEmpty = 64;
+	inline const int32 RemoveServerNameEmpty = 64;
 
 	/** License Not Accepted  */
-	const int32 LicenseNotAccepted = 128;
+	inline const int32 LicenseNotAccepted = 128;
 
 	/** Code Build Required */
-	const int32 CodeBuildRequired = 256;
+	inline const int32 CodeBuildRequired = 256;
 };
 
 
diff --git a/Engine/Source/Editor/PropertyEditor/Private/PropertyNode.h b/Engine/Source/Editor/PropertyEditor/Private/PropertyNode.h
--- a/Engine/Source/Editor/PropertyEditor/Private/PropertyNode.h
+++ b/Engine/Source/Editor/PropertyEditor/Private/PropertyNode.h
@@ -28,59 +28,59 @@ namespace EPropertyNodeFlags
 {
 	typedef uint32 Type;
 
-	const Type	IsSeen							= 1 << 0;	/** true if this node can be seen based on current parent expansion.  Does not take into account clipping*/
-	const Type	IsSeenDueToFiltering			= 1 << 1;	/** true if this node has been accepted by the filter*/
-	const Type	IsSeenDueToChildFiltering		= 1 << 2;	/** true if this node or one of it's children is seen due to filtering.  It will then be forced on as well.*/
-	const Type	IsParentSeenDueToFiltering		= 1 << 3;	/** True if the parent was visible due to filtering*/
-	const Type	IsSeenDueToChildFavorite		= 1 << 4;	/** True if this node is seen to it having a favorite as a child */
+	inline const Type	IsSeen							= 1 << 0;	/** true if this node can be seen based on current parent expansion.  Does not take into account clipping*/
+	inline const Type	IsSeenDueToFiltering			= 1 << 1;	/** true if this node has been accepted by the filter*/
+	inline const Type	IsSeenDueToChildFiltering		= 1 << 2;	/** true if this node or one of it's children is seen due to filtering.  It will then be forced on as well.*/
+	inline const Type	IsParentSeenDueToFiltering		= 1 << 3;	/** True if the parent was visible due to filtering*/
+	inline const Type	IsSeenDueToChildFavorite		= 1 << 4;	/** True if this node is seen to it having a favorite as a child */
 	
-	const Type	Expanded						= 1 << 5;	/** true if this node should display its children*/
-	const Type	CanBeExpanded					= 1 << 6;	/** true if this node is able to be expanded */
+	inline const Type	Expanded						= 1 << 5;	/** true if this node should display its children*/
+	inline const Type	CanBeExpanded					= 1 << 6;	/** true if this node is able to be expanded */
 
-	const Type	EditInlineNew					= 1 << 7;	/** true if the property can be expanded into the property window. */
+	inline const Type	EditInlineNew					= 1 << 7;	/** true if the property can be expanded into the property window. */
 
-	const Type	SingleSelectOnly				= 1 << 8;	/** true if only a single object is selected. */
-	const Type  ShowCategories					= 1 << 9;	/** true if this node should show categories.  Different*/
+	inline const Type	SingleSelectOnly				= 1 << 8;	/** true if only a single object is selected. */
+	inline const Type  ShowCategories					= 1 << 9;	/** true if this node should show categories.  Different*/
 
-	const Type  HasEverBeenExpanded				= 1 << 10;	/** true if expand has ever been called on this node */
+	inline const Type  HasEverBeenExpanded				= 1 << 10;	/** true if expand has ever been called on this node */
 
-	const Type	IsBeingFiltered					= 1 << 11;	/** true if the node is being filtered. If this is true, seen flags should be checked for visibility.  If this is false the node has no filter and is visible */
+	inline const Type	IsBeingFiltered					= 1 << 11;	/** true if the node is being filtered. If this is true, seen flags should be checked for visibility.  If this is false the node has no filter and is visible */
 
-	const Type  IsFavorite						= 1 << 12;	/** true if this item has been dubbed a favorite by the user */
+	inline const Type  IsFavorite						= 1 << 12;	/** true if this item has been dubbed a favorite by the user */
 
-	const Type  NoChildrenDueToCircularReference= 1 << 13;	/** true if this node has no children (but normally would) due to circular referencing */
+	inline const Type  NoChildrenDueToCircularReference= 1 << 13;	/** true if this node has no children (but normally would) due to circular referencing */
 
-	const Type	AutoExpanded					= 1 << 14;	/** true if this node was autoexpanded due to being filtered */
-	const Type	ShouldShowHiddenProperties		= 1 << 15;	/** true if this node should all properties not just those with the correct flag(s) to be shown in the editor */
-	const Type	IsAdvanced						= 1 << 16;	/** true if the property node is advanced (i.e it only shows up in advanced sections) */
-	const Type	IsCustomized					= 1 << 17;	/** true if this node's visual representation has been customized by the editor */
+	inline const Type	AutoExpanded					= 1 << 14;	/** true if this node was autoexpanded due to being filtered */
+	inline const Type	ShouldShowHiddenProperties		= 1 << 15;	/** true if this node should all properties not just those with the correct flag(s) to be shown in the editor */
+	inline const Type	IsAdvanced						= 1 << 16;	/** true if the property node is advanced (i.e it only shows up in advanced sections) */
+	inline const Type	IsCustomized					= 1 << 17;	/** true if this node's visual representation has been customized by the editor */
 	
-	const Type	RequiresValidation				= 1 << 18;	/** true if this node could unexpectedly change (array changes, editinlinenew changes) */
+	inline const Type	RequiresValidation				= 1 << 18;	/** true if this node could unexpectedly change (array changes, editinlinenew changes) */
 
-	const Type	ShouldShowDisableEditOnInstance = 1 << 19;	/** true if this node should show child properties marked CPF_DisableEditOnInstance */
+	inline const Type	ShouldShowDisableEditOnInstance = 1 << 19;	/** true if this node should show child properties marked CPF_DisableEditOnInstance */
 
-	const Type	IsReadOnly						= 1 << 20;	/** true if this node is overridden to appear as read-only */
+	inline const Type	IsReadOnly						= 1 << 20;	/** true if this node is overridden to appear as read-only */
 
-	const Type	SkipChildValidation				= 1 << 21;	/** true if this node should skip child validation */
+	inline const Type	SkipChildValidation				= 1 << 21;	/** true if this node should skip child validation */
 
-	const Type  ShowInnerObjectProperties		= 1 << 22;
+	inline const Type  ShowInnerObjectProperties		= 1 << 22;
 
-	const Type	HasCustomResetToDefault			= 1 << 23;	/** true if this node's visual representation of reset to default has been customized*/
+	inline const Type	HasCustomResetToDefault			= 1 << 23;	/** true if this node's visual representation of reset to default has been customized*/
 
-	const Type	IsSparseClassData				= 1 << 24;	/** true if the property on this node is part of a sparse class data structure */
+	inline const Type	IsSparseClassData				= 1 << 24;	/** true if the property on this node is part of a sparse class data structure */
 
-	const Type	ShouldShowInViewport			= 1 << 25;	/** true if the property should be shown in the viewport context menu */
+	inline const Type	ShouldShowInViewport			= 1 << 25;	/** true if the property should be shown in the viewport context menu */
 
-	const Type 	NoFlags							= 0;
+	inline const Type 	NoFlags							= 0;
 
 };
 
 namespace FPropertyNodeConstants
 {
-	const int32 NoDepthRestrictions = -1;
+	inline const int32 NoDepthRestrictions = -1;
 
 	/** Character used to deliminate sub-categories in category path names */
-	const TCHAR CategoryDelimiterChar = TCHAR( '|' );
+	inline const TCHAR CategoryDelimiterChar = TCHAR( '|' );
 };
 
 class FPropertySettings
diff --git a/Engine/Source/Editor/PropertyEditor/Public/PropertyHandle.h b/Engine/Source/Editor/PropertyEditor/Public/PropertyHandle.h
--- a/Engine/Source/Editor/PropertyEditor/Public/PropertyHandle.h
+++ b/Engine/Source/Editor/PropertyEditor/Public/PropertyHandle.h
@@ -20,13 +20,13 @@ namespace EPropertyValueSetFlags
 	typedef uint32 Type;
 
 	/** Normal way to call set value (make a transaction, call posteditchange) */
-	const Type DefaultFlags = 0;
+	inline const Type DefaultFlags = 0;
 	/** No transaction will be created when setting this value (no undo/redo) */
-	const Type NotTransactable = 1 << 0;
+	inline const Type NotTransactable = 1 << 0;
 	/** When PostEditChange is called mark the change as interactive (e.g, user is spinning a value in a spin box) */
-	const Type InteractiveChange = 1 << 1;
+	inline const Type InteractiveChange = 1 << 1;
 	/** If the property being imported to is an instanced object, create a new object rather than simply setting the value literally */
-	const Type InstanceObjects = 1 << 2;
+	inline const Type InstanceObjects = 1 << 2;
 
 };
 
diff --git a/Engine/Source/Editor/UnrealEd/Classes/Editor/AssetReferenceFilter.h b/Engine/Source/Editor/UnrealEd/Classes/Editor/AssetReferenceFilter.h
--- a/Engine/Source/Editor/UnrealEd/Classes/Editor/AssetReferenceFilter.h
+++ b/Engine/Source/Editor/UnrealEd/Classes/Editor/AssetReferenceFilter.h
@@ -4,6 +4,7 @@
 
 #include "CoreTypes.h"
 
+class FText;
 struct FAssetData;
 
 /** Used in filtering allowed references between assets. Implement a subclass of this and return it in OnMakeAssetReferenceFilter */
diff --git a/Engine/Source/Editor/UnrealEd/Private/Bookmarks/BookMark2DTypeActions.h b/Engine/Source/Editor/UnrealEd/Private/Bookmarks/BookMark2DTypeActions.h
--- a/Engine/Source/Editor/UnrealEd/Private/Bookmarks/BookMark2DTypeActions.h
+++ b/Engine/Source/Editor/UnrealEd/Private/Bookmarks/BookMark2DTypeActions.h
@@ -10,6 +10,7 @@
 #include "Engine/BookMark2D.h"
 #include "Engine/LevelStreaming.h"
 #include "Engine/World.h"
+#include "LevelEditorViewport.h"
 
 class FBookMark2DTypeActions : public IBookmarkTypeActions
 {
diff --git a/Engine/Source/Editor/UnrealEd/Private/Cooker/Algo/GraphConvert.h b/Engine/Source/Editor/UnrealEd/Private/Cooker/Algo/GraphConvert.h
--- a/Engine/Source/Editor/UnrealEd/Private/Cooker/Algo/GraphConvert.h
+++ b/Engine/Source/Editor/UnrealEd/Private/Cooker/Algo/GraphConvert.h
@@ -20,7 +20,7 @@ namespace Algo::Graph
 {
 
 typedef int32 FVertex;
-constexpr FVertex InvalidVertex = static_cast<FVertex>(INDEX_NONE);
+inline constexpr FVertex InvalidVertex = static_cast<FVertex>(INDEX_NONE);
 
 enum class EConvertToGraphOptions
 {
diff --git a/Engine/Source/Editor/UnrealEd/Private/Cooker/CookSockets.h b/Engine/Source/Editor/UnrealEd/Private/Cooker/CookSockets.h
--- a/Engine/Source/Editor/UnrealEd/Private/Cooker/CookSockets.h
+++ b/Engine/Source/Editor/UnrealEd/Private/Cooker/CookSockets.h
@@ -12,7 +12,7 @@ class FSocket;
 namespace UE::Cook::Sockets
 {
 
-constexpr int32 COOKDIRECTOR_DEFAULT_REQUEST_CONNECTION_PORT = 41897;
+inline constexpr int32 COOKDIRECTOR_DEFAULT_REQUEST_CONNECTION_PORT = 41897;
 
 /** Destruct Socket as gracefully as possible. Sets Socket to nullptr. Early exits if Socket is already null. */
 void CloseSocket(FSocket*& Socket);
diff --git a/Engine/Source/Editor/UnrealEd/Private/Cooker/CookTypes.h b/Engine/Source/Editor/UnrealEd/Private/Cooker/CookTypes.h
--- a/Engine/Source/Editor/UnrealEd/Private/Cooker/CookTypes.h
+++ b/Engine/Source/Editor/UnrealEd/Private/Cooker/CookTypes.h
@@ -518,8 +518,8 @@ bool LoadFromCompactBinary(FCbFieldView Field, FBeginCookContextForWorker& Value
 void LogCookerMessage(const FString& MessageText, EMessageSeverity::Type Severity);
 LLM_DECLARE_TAG(Cooker);
 
-constexpr uint32 ExpectedMaxNumPlatforms = 32;
+inline constexpr uint32 ExpectedMaxNumPlatforms = 32;
 #define REMAPPED_PLUGINS TEXT("RemappedPlugins")
 extern float GCookProgressWarnBusyTime;
 
-constexpr float TickCookableObjectsFrameTime = .100f;
\ No newline at end of file
+inline constexpr float TickCookableObjectsFrameTime = .100f;
\ No newline at end of file
diff --git a/Engine/Source/Editor/UnrealEd/Public/Bookmarks/BookMarkTypeActions.h b/Engine/Source/Editor/UnrealEd/Public/Bookmarks/BookMarkTypeActions.h
--- a/Engine/Source/Editor/UnrealEd/Public/Bookmarks/BookMarkTypeActions.h
+++ b/Engine/Source/Editor/UnrealEd/Public/Bookmarks/BookMarkTypeActions.h
@@ -3,13 +3,13 @@
 #pragma once
 
 #include "Bookmarks/IBookmarkTypeActions.h"
-#include "EditorViewportClient.h"
 #include "Editor.h"
 #include "Editor/EditorEngine.h"
 
 #include "Engine/BookMark.h"
 #include "Engine/LevelStreaming.h"
 #include "Engine/World.h"
+#include "LevelEditorViewport.h"
 
 class FBookMarkTypeActions : public IBookmarkTypeActions
 {
diff --git a/Engine/Source/Editor/UnrealEd/Public/Commandlets/AnalyzeReferencedContentStat.h b/Engine/Source/Editor/UnrealEd/Public/Commandlets/AnalyzeReferencedContentStat.h
--- a/Engine/Source/Editor/UnrealEd/Public/Commandlets/AnalyzeReferencedContentStat.h
+++ b/Engine/Source/Editor/UnrealEd/Public/Commandlets/AnalyzeReferencedContentStat.h
@@ -14,6 +14,7 @@ class UTexture;
 class UModel;
 class USkeletalMesh;
 class UStaticMesh;
+enum AnimationCompressionFormat : int;
 
 class FAnalyzeReferencedContentStat
 {
@@ -501,11 +502,11 @@ public:
 		/** Whether resource is forced to be uncompressed by human action               */
 		bool bMarkedAsDoNotOverrideCompression;
 		/** Type of compression used on this animation.									*/
-		enum AnimationCompressionFormat TranslationFormat;
+		AnimationCompressionFormat TranslationFormat;
 		/** Type of compression used on this animation.									*/
-		enum AnimationCompressionFormat RotationFormat;
+		AnimationCompressionFormat RotationFormat;
 		/** Type of compression used on this animation.									*/
-		enum AnimationCompressionFormat ScaleFormat;
+		AnimationCompressionFormat ScaleFormat;
 		/** Name of compression algo class used. */
 		FString CompressionScheme;
 		/** Size in bytes of this animation. */
diff --git a/Engine/Source/Editor/UnrealEd/Public/DragAndDrop/FolderDragDropOp.h b/Engine/Source/Editor/UnrealEd/Public/DragAndDrop/FolderDragDropOp.h
--- a/Engine/Source/Editor/UnrealEd/Public/DragAndDrop/FolderDragDropOp.h
+++ b/Engine/Source/Editor/UnrealEd/Public/DragAndDrop/FolderDragDropOp.h
@@ -3,6 +3,8 @@
 #pragma once
 
 #include "CoreMinimal.h"
+#include "Engine/World.h"
+#include "Folder.h"
 #include "Input/DragAndDrop.h"
 #include "DragAndDrop/DecoratedDragDropOp.h"
 #include "EditorFolderUtils.h"
diff --git a/Engine/Source/Editor/UnrealEd/Public/DragAndDrop/MultipleDataDragDropOp.h b/Engine/Source/Editor/UnrealEd/Public/DragAndDrop/MultipleDataDragDropOp.h
deleted file mode 100644
diff --git a/Engine/Source/Editor/UnrealEd/Public/Editor.h b/Engine/Source/Editor/UnrealEd/Public/Editor.h
--- a/Engine/Source/Editor/UnrealEd/Public/Editor.h
+++ b/Engine/Source/Editor/UnrealEd/Public/Editor.h
@@ -377,16 +377,16 @@ namespace MapChangeEventFlags
 	typedef uint32 Type;
 
 	/** Default flags */
-	const Type Default = 0;
+	inline const Type Default = 0;
 
 	/** Set when a new map is created, loaded from disk, imported, etc. */
-	const Type NewMap = 1 << 0;
+	inline const Type NewMap = 1 << 0;
 
 	/** Set when a map rebuild occurred */
-	const Type MapRebuild = 1 << 1;
+	inline const Type MapRebuild = 1 << 1;
 
 	/** Set when a world was destroyed (torn down) */
-	const Type WorldTornDown = 1 << 2;
+	inline const Type WorldTornDown = 1 << 2;
 }
 
 /**
diff --git a/Engine/Source/Editor/UnrealEd/Public/IProjectExternalContentInterface.h b/Engine/Source/Editor/UnrealEd/Public/IProjectExternalContentInterface.h
--- a/Engine/Source/Editor/UnrealEd/Public/IProjectExternalContentInterface.h
+++ b/Engine/Source/Editor/UnrealEd/Public/IProjectExternalContentInterface.h
@@ -3,6 +3,8 @@
 
 #include "CoreMinimal.h"
 
+class IPlugin;
+
 /** Interface to manage project references to external content */
 class IProjectExternalContentInterface
 {
diff --git a/Engine/Source/Editor/UnrealEd/Public/ObjectTools.h b/Engine/Source/Editor/UnrealEd/Public/ObjectTools.h
--- a/Engine/Source/Editor/UnrealEd/Public/ObjectTools.h
+++ b/Engine/Source/Editor/UnrealEd/Public/ObjectTools.h
@@ -702,7 +702,7 @@ namespace ThumbnailTools
 	/** Loads thumbnails for the specified objects (or copies them from a cache, if they're already loaded.) */
 	UNREALED_API bool ConditionallyLoadThumbnailsForObjects( const TArray< FName >& InObjectFullNames, FThumbnailMap& InOutThumbnails );
 	/** Standard thumbnail height setting used by generation */
-	const int32 DefaultThumbnailSize=256;
+	inline const int32 DefaultThumbnailSize=256;
 
 	/** Returns true if the given asset has a custom thumbnail cached or on the disk. */
 	UNREALED_API bool AssetHasCustomThumbnail(const FString& InAssetDataFullName);
diff --git a/Engine/Source/Editor/UnrealEd/Public/PackageUtilityWorkers.h b/Engine/Source/Editor/UnrealEd/Public/PackageUtilityWorkers.h
--- a/Engine/Source/Editor/UnrealEd/Public/PackageUtilityWorkers.h
+++ b/Engine/Source/Editor/UnrealEd/Public/PackageUtilityWorkers.h
@@ -7,6 +7,7 @@
 #pragma once
 
 #include "CoreMinimal.h"
+#include "Misc/FeedbackContext.h"
 
 /**
  * These bit flag values represent the different types of information that can be reported about a package
diff --git a/Engine/Source/Editor/UnrealEd/Public/Viewports.h b/Engine/Source/Editor/UnrealEd/Public/Viewports.h
--- a/Engine/Source/Editor/UnrealEd/Public/Viewports.h
+++ b/Engine/Source/Editor/UnrealEd/Public/Viewports.h
@@ -11,11 +11,11 @@
 namespace EditorViewportDefs
 {
 	/** Default camera position for level editor perspective viewports */
-	const FVector DefaultPerspectiveViewLocation( 0.0f, 1024.0f, 512.0f );
+	inline const FVector DefaultPerspectiveViewLocation( 0.0f, 1024.0f, 512.0f );
 
 	/** Default camera orientation for level editor perspective viewports */
-	const FRotator DefaultPerspectiveViewRotation( -15.0f, -90.0f, 0 );
+	inline const FRotator DefaultPerspectiveViewRotation( -15.0f, -90.0f, 0 );
 
 	/** Default camera field of view angle for level editor perspective viewports */
-	const float DefaultPerspectiveFOVAngle( 90.0f );
+	inline const float DefaultPerspectiveFOVAngle( 90.0f );
 }
diff --git a/Engine/Source/Runtime/AIModule/Classes/AIController.h b/Engine/Source/Runtime/AIModule/Classes/AIController.h
--- a/Engine/Source/Runtime/AIModule/Classes/AIController.h
+++ b/Engine/Source/Runtime/AIModule/Classes/AIController.h
@@ -54,11 +54,11 @@ namespace EAIFocusPriority
 {
 	typedef uint8 Type;
 
-	const Type Default = 0;
-	const Type Move = 1;
-	const Type Gameplay = 2;
+	inline const Type Default = 0;
+	inline const Type Move = 1;
+	inline const Type Gameplay = 2;
 
-	const Type LastFocusPriority = Gameplay;
+	inline const Type LastFocusPriority = Gameplay;
 }
 
 struct FFocusKnowledge
diff --git a/Engine/Source/Runtime/AIModule/Classes/AITypes.h b/Engine/Source/Runtime/AIModule/Classes/AITypes.h
--- a/Engine/Source/Runtime/AIModule/Classes/AITypes.h
+++ b/Engine/Source/Runtime/AIModule/Classes/AITypes.h
@@ -20,13 +20,13 @@ DECLARE_CYCLE_STAT_EXTERN(TEXT("Overall AI Time"), STAT_AI_Overall, STATGROUP_AI
 
 namespace FAISystem
 {
-	static const FRotator InvalidRotation = FRotator(FLT_MAX);
-	static const FQuat InvalidOrientation = FQuat(FLT_MAX, FLT_MAX, FLT_MAX, FLT_MAX);
-	static const FVector InvalidLocation = FVector(FLT_MAX);
-	static const FVector InvalidDirection = FVector::ZeroVector; 
-	static const float InvalidRange = -1.f;
-	static const float InfiniteInterval = -FLT_MAX;
-	static const uint32 InvalidUnsignedID = uint32(INDEX_NONE);
+	inline static const FRotator InvalidRotation = FRotator(FLT_MAX);
+	inline static const FQuat InvalidOrientation = FQuat(FLT_MAX, FLT_MAX, FLT_MAX, FLT_MAX);
+	inline static const FVector InvalidLocation = FVector(FLT_MAX);
+	inline static const FVector InvalidDirection = FVector::ZeroVector; 
+	inline static const float InvalidRange = -1.f;
+	inline static const float InfiniteInterval = -FLT_MAX;
+	inline static const uint32 InvalidUnsignedID = uint32(INDEX_NONE);
 
 	FORCEINLINE bool IsValidLocation(const FVector& TestLocation)
 	{
diff --git a/Engine/Source/Runtime/AIModule/Classes/BehaviorTree/BehaviorTreeTypes.h b/Engine/Source/Runtime/AIModule/Classes/BehaviorTree/BehaviorTreeTypes.h
--- a/Engine/Source/Runtime/AIModule/Classes/BehaviorTree/BehaviorTreeTypes.h
+++ b/Engine/Source/Runtime/AIModule/Classes/BehaviorTree/BehaviorTreeTypes.h
@@ -50,7 +50,7 @@ namespace FBlackboard
 #ifdef AI_BLACKBOARD_KEY_SIZE_8
 	// this is the legacy BB key size. Add AI_BLACKBOARD_KEY_SIZE_8 to your *.target.cs to enable it.
 	using FKey = uint8;
-	constexpr FKey InvalidKey = FKey(-1);
+	inline constexpr FKey InvalidKey = FKey(-1);
 #else
 	//the default BB key size is now 16
 	struct FKey
@@ -69,7 +69,7 @@ namespace FBlackboard
 		uint16 Key = static_cast<uint16>(-1);
 	};
 
-	constexpr FKey InvalidKey = FKey();
+	inline constexpr FKey InvalidKey = FKey();
 
 	inline uint32 GetTypeHash(const FKey& Key) { return ::GetTypeHash(Key.Key);}
 #endif
@@ -94,10 +94,10 @@ DECLARE_DELEGATE_RetVal_TwoParams(EBlackboardNotificationResult, FOnBlackboardCh
 
 namespace BTSpecialChild
 {
-	constexpr int32 NotInitialized = -1;	// special value for child indices: needs to be initialized
-	constexpr int32 ReturnToParent = -2;	// special value for child indices: return to parent node
+	inline constexpr int32 NotInitialized = -1;	// special value for child indices: needs to be initialized
+	inline constexpr int32 ReturnToParent = -2;	// special value for child indices: return to parent node
 	
-	constexpr uint8 OwnedByComposite = MAX_uint8;	// special value for aux node's child index: owned by composite node instead of a task
+	inline constexpr uint8 OwnedByComposite = MAX_uint8;	// special value for aux node's child index: owned by composite node instead of a task
 }
 
 UENUM(BlueprintType)
@@ -292,7 +292,7 @@ struct FBehaviorTreeExecutionStep
 	/** Game world's time stamp of this step */
 	double TimeStamp;
 
-	static constexpr int32 InvalidExecutionId = -1;
+	inline static constexpr int32 InvalidExecutionId = -1;
 
 	/** Id of execution step */
 	int32 ExecutionStepId;
diff --git a/Engine/Source/Runtime/AIModule/Classes/EnvironmentQuery/EnvQueryGenerator.h b/Engine/Source/Runtime/AIModule/Classes/EnvironmentQuery/EnvQueryGenerator.h
--- a/Engine/Source/Runtime/AIModule/Classes/EnvironmentQuery/EnvQueryGenerator.h
+++ b/Engine/Source/Runtime/AIModule/Classes/EnvironmentQuery/EnvQueryGenerator.h
@@ -11,10 +11,10 @@
 
 namespace EnvQueryGeneratorVersion
 {
-	const int32 Initial = 0;
-	const int32 DataProviders = 1;
+	inline const int32 Initial = 0;
+	inline const int32 DataProviders = 1;
 
-	const int32 Latest = DataProviders;
+	inline const int32 Latest = DataProviders;
 }
 
 UCLASS(EditInlineNew, Abstract, meta = (Category = "Generators"))
diff --git a/Engine/Source/Runtime/AIModule/Classes/EnvironmentQuery/EnvQueryTest.h b/Engine/Source/Runtime/AIModule/Classes/EnvironmentQuery/EnvQueryTest.h
--- a/Engine/Source/Runtime/AIModule/Classes/EnvironmentQuery/EnvQueryTest.h
+++ b/Engine/Source/Runtime/AIModule/Classes/EnvironmentQuery/EnvQueryTest.h
@@ -29,11 +29,11 @@ enum class EEQSNormalizationType : uint8
 
 namespace EnvQueryTestVersion
 {
-	const int32 Initial = 0;
-	const int32 DataProviders = 1;
-	const int32 ReferenceValueFix = 2;
+	inline const int32 Initial = 0;
+	inline const int32 DataProviders = 1;
+	inline const int32 ReferenceValueFix = 2;
 
-	const int32 Latest = ReferenceValueFix;
+	inline const int32 Latest = ReferenceValueFix;
 }
 
 #if WITH_EDITORONLY_DATA
diff --git a/Engine/Source/Runtime/AIModule/Classes/Navigation/PathFollowingComponent.h b/Engine/Source/Runtime/AIModule/Classes/Navigation/PathFollowingComponent.h
--- a/Engine/Source/Runtime/AIModule/Classes/Navigation/PathFollowingComponent.h
+++ b/Engine/Source/Runtime/AIModule/Classes/Navigation/PathFollowingComponent.h
@@ -78,42 +78,42 @@ namespace FPathFollowingResultFlags
 {
 	typedef uint16 Type;
 
-	const Type None = 0;
+	inline const Type None = 0;
 
 	/** Reached destination (EPathFollowingResult::Success) */
-	const Type Success = (1 << 0);
+	inline const Type Success = (1 << 0);
 
 	/** Movement was blocked (EPathFollowingResult::Blocked) */
-	const Type Blocked = (1 << 1);
+	inline const Type Blocked = (1 << 1);
 
 	/** Agent is not on path (EPathFollowingResult::OffPath) */
-	const Type OffPath = (1 << 2);
+	inline const Type OffPath = (1 << 2);
 
 	/** Aborted (EPathFollowingResult::Aborted) */
-	const Type UserAbort = (1 << 3);
+	inline const Type UserAbort = (1 << 3);
 
 	/** Abort details: owner no longer wants to move */
-	const Type OwnerFinished = (1 << 4);
+	inline const Type OwnerFinished = (1 << 4);
 
 	/** Abort details: path is no longer valid */
-	const Type InvalidPath = (1 << 5);
+	inline const Type InvalidPath = (1 << 5);
 
 	/** Abort details: unable to move */
-	const Type MovementStop = (1 << 6);
+	inline const Type MovementStop = (1 << 6);
 
 	/** Abort details: new movement request was received */
-	const Type NewRequest = (1 << 7);
+	inline const Type NewRequest = (1 << 7);
 
 	/** Abort details: blueprint MoveTo function was called */
-	const Type ForcedScript = (1 << 8);
+	inline const Type ForcedScript = (1 << 8);
 
 	/** Finish details: never started, agent was already at goal */
-	const Type AlreadyAtGoal = (1 << 9);
+	inline const Type AlreadyAtGoal = (1 << 9);
 
 	/** Can be used to create project specific reasons */
-	const Type FirstGameplayFlagShift = 10;
+	inline const Type FirstGameplayFlagShift = 10;
 
-	const Type UserAbortFlagMask = ~(Success | Blocked | OffPath);
+	inline const Type UserAbortFlagMask = ~(Success | Blocked | OffPath);
 
 	FString ToString(uint16 Value);
 }
diff --git a/Engine/Source/Runtime/AIModule/Public/GraphAStar.h b/Engine/Source/Runtime/AIModule/Public/GraphAStar.h
--- a/Engine/Source/Runtime/AIModule/Public/GraphAStar.h
+++ b/Engine/Source/Runtime/AIModule/Public/GraphAStar.h
@@ -20,7 +20,7 @@ enum EGraphAStarResult
 	InfiniteLoop
 };
 
-const int32 NO_COUNT = INT_MAX;
+inline const int32 NO_COUNT = INT_MAX;
 
 // To get AStar Graph tracing, enable this define
 #define ENABLE_GRAPH_ASTAR_LOGGING 0
diff --git a/Engine/Source/Runtime/Core/Private/Unix/UnixPlatformSyscallTable.cpp b/Engine/Source/Runtime/Core/Private/Unix/UnixPlatformSyscallTable.cpp
--- a/Engine/Source/Runtime/Core/Private/Unix/UnixPlatformSyscallTable.cpp
+++ b/Engine/Source/Runtime/Core/Private/Unix/UnixPlatformSyscallTable.cpp
@@ -1,6 +1,7 @@
 // Copyright Epic Games, Inc. All Rights Reserved.
 
 #include "UnixPlatformSyscallTable.h"
+#include "Containers/UnrealString.h"
 
 namespace
 {
diff --git a/Engine/Source/Runtime/Core/Private/Unix/UnixPlatformSyscallTable.h b/Engine/Source/Runtime/Core/Private/Unix/UnixPlatformSyscallTable.h
--- a/Engine/Source/Runtime/Core/Private/Unix/UnixPlatformSyscallTable.h
+++ b/Engine/Source/Runtime/Core/Private/Unix/UnixPlatformSyscallTable.h
@@ -4,4 +4,6 @@
 
 #include "CoreTypes.h"
 
+class FString;
+
 int UnixPlatformLookupSyscallNumberFromName(const FString& SyscallName);
diff --git a/Engine/Source/Runtime/CoreUObject/Public/Misc/AssetRegistryInterface.h b/Engine/Source/Runtime/CoreUObject/Public/Misc/AssetRegistryInterface.h
--- a/Engine/Source/Runtime/CoreUObject/Public/Misc/AssetRegistryInterface.h
+++ b/Engine/Source/Runtime/CoreUObject/Public/Misc/AssetRegistryInterface.h
@@ -41,10 +41,10 @@ namespace EAssetRegistryDependencyType
 		// Note: Also update FAssetRegistryDependencyOptions when adding more flags
 	};
 
-	static const Type None = (Type)(0);
-	static const Type All = (Type)(Soft | Hard | SearchableName | SoftManage | HardManage);
-	static const Type Packages = (Type)(Soft | Hard);
-	static const Type Manage = (Type)(SoftManage | HardManage);
+	inline static const Type None = (Type)(0);
+	inline static const Type All = (Type)(Soft | Hard | SearchableName | SoftManage | HardManage);
+	inline static const Type Packages = (Type)(Soft | Hard);
+	inline static const Type Manage = (Type)(SoftManage | HardManage);
 }
 
 class IAssetRegistry;
diff --git a/Engine/Source/Runtime/CoreUObject/Public/Misc/ITransactionObjectAnnotation.h b/Engine/Source/Runtime/CoreUObject/Public/Misc/ITransactionObjectAnnotation.h
--- a/Engine/Source/Runtime/CoreUObject/Public/Misc/ITransactionObjectAnnotation.h
+++ b/Engine/Source/Runtime/CoreUObject/Public/Misc/ITransactionObjectAnnotation.h
@@ -2,7 +2,7 @@
 
 #pragma once
 
-#include "CoreTypes.h"
+#include "Containers/Map.h"
 
 class UObject;
 struct FTransactionObjectChange;
diff --git a/Engine/Source/Runtime/CoreUObject/Public/Serialization/ArchiveReplaceObjectRef.h b/Engine/Source/Runtime/CoreUObject/Public/Serialization/ArchiveReplaceObjectRef.h
--- a/Engine/Source/Runtime/CoreUObject/Public/Serialization/ArchiveReplaceObjectRef.h
+++ b/Engine/Source/Runtime/CoreUObject/Public/Serialization/ArchiveReplaceObjectRef.h
@@ -3,8 +3,9 @@
 #pragma once
 
 #include "CoreMinimal.h"
-#include "UObject/Object.h"
 #include "Serialization/ArchiveUObject.h"
+#include "UObject/Object.h"
+#include "UObject/SoftObjectPath.h"
 
 /** Base class for object replacement archives */ 
 class COREUOBJECT_API FArchiveReplaceObjectRefBase : public FArchiveUObject
diff --git a/Engine/Source/Runtime/CoreUObject/Public/Serialization/ArchiveScriptReferenceCollector.h b/Engine/Source/Runtime/CoreUObject/Public/Serialization/ArchiveScriptReferenceCollector.h
--- a/Engine/Source/Runtime/CoreUObject/Public/Serialization/ArchiveScriptReferenceCollector.h
+++ b/Engine/Source/Runtime/CoreUObject/Public/Serialization/ArchiveScriptReferenceCollector.h
@@ -3,6 +3,7 @@
 #pragma once
 
 #include "CoreMinimal.h"
+#include "UObject/Field.h"
 #include "UObject/Object.h"
 #include "UObject/UObjectGlobals.h"
 #include "Serialization/ArchiveUObject.h"
diff --git a/Engine/Source/Runtime/CoreUObject/Public/Serialization/BulkData2.h b/Engine/Source/Runtime/CoreUObject/Public/Serialization/BulkData2.h
--- a/Engine/Source/Runtime/CoreUObject/Public/Serialization/BulkData2.h
+++ b/Engine/Source/Runtime/CoreUObject/Public/Serialization/BulkData2.h
@@ -2,6 +2,8 @@
 
 #pragma once
 
+#include "Misc/CoreMiscDefines.h"
+
 UE_DEPRECATED(5.1, "Including BulkData2.h has been deprecated. Please include BulkData.h instead.")
 inline void BulkData2CommonHeaderDeprecatedWarning()
 {
diff --git a/Engine/Source/Runtime/CoreUObject/Public/Serialization/BulkDataCommon.h b/Engine/Source/Runtime/CoreUObject/Public/Serialization/BulkDataCommon.h
--- a/Engine/Source/Runtime/CoreUObject/Public/Serialization/BulkDataCommon.h
+++ b/Engine/Source/Runtime/CoreUObject/Public/Serialization/BulkDataCommon.h
@@ -2,6 +2,8 @@
 
 #pragma once
 
+#include "Misc/CoreMiscDefines.h"
+
 UE_DEPRECATED(5.1, "Including BulkDataCommon.h has been deprecated. Please include BulkData.h instead.")
 inline void BulkDataCommonHeaderDeprecatedWarning()
 {
diff --git a/Engine/Source/Runtime/CoreUObject/Public/Serialization/DuplicatedObject.h b/Engine/Source/Runtime/CoreUObject/Public/Serialization/DuplicatedObject.h
--- a/Engine/Source/Runtime/CoreUObject/Public/Serialization/DuplicatedObject.h
+++ b/Engine/Source/Runtime/CoreUObject/Public/Serialization/DuplicatedObject.h
@@ -3,6 +3,7 @@
 #pragma once
 
 #include "CoreMinimal.h"
+#include "UObject/WeakObjectPtr.h"
 
 /**
  * Information about a duplicated object
diff --git a/Engine/Source/Runtime/CoreUObject/Public/UObject/FieldIterator.h b/Engine/Source/Runtime/CoreUObject/Public/UObject/FieldIterator.h
--- a/Engine/Source/Runtime/CoreUObject/Public/UObject/FieldIterator.h
+++ b/Engine/Source/Runtime/CoreUObject/Public/UObject/FieldIterator.h
@@ -8,6 +8,7 @@ FieldIterator.h: FField iterators.
 
 #include "CoreMinimal.h"
 #include "UObject/Field.h"
+#include "UObject/UnrealType.h"
 #include "UObject/UObjectIterator.h"
 
 /**
diff --git a/Engine/Source/Runtime/Engine/Public/Net/Iris/ReplicationSystem/ReplicationSystemUtil.h b/Engine/Source/Runtime/Engine/Public/Net/Iris/ReplicationSystem/ReplicationSystemUtil.h
--- a/Engine/Source/Runtime/Engine/Public/Net/Iris/ReplicationSystem/ReplicationSystemUtil.h
+++ b/Engine/Source/Runtime/Engine/Public/Net/Iris/ReplicationSystem/ReplicationSystemUtil.h
@@ -2,6 +2,8 @@
 
 #pragma once
 
+// HEADER_UNIT_SKIP - Can't find includes if UE_WITH_IRIS is not set
+
 #if UE_WITH_IRIS
 
 #include "Iris/ReplicationSystem/NetHandle.h"
diff --git a/Engine/Source/Runtime/Experimental/Chaos/Public/Chaos/EPA.h b/Engine/Source/Runtime/Experimental/Chaos/Public/Chaos/EPA.h
--- a/Engine/Source/Runtime/Experimental/Chaos/Public/Chaos/EPA.h
+++ b/Engine/Source/Runtime/Experimental/Chaos/Public/Chaos/EPA.h
@@ -9,7 +9,7 @@
 namespace Chaos
 {
 
-constexpr int32 ArraySizeEPA = 16;
+inline constexpr int32 ArraySizeEPA = 16;
 
 // Array type used in EPA to avoid heap allocation for small convex shapes
 // @todo(chaos): The inline size was picked to avoid allocations in box-box collision - it might need adjusting after more general purpose testing
diff --git a/Engine/Source/Runtime/Experimental/Chaos/Public/Chaos/ImplicitObject.h b/Engine/Source/Runtime/Experimental/Chaos/Public/Chaos/ImplicitObject.h
--- a/Engine/Source/Runtime/Experimental/Chaos/Public/Chaos/ImplicitObject.h
+++ b/Engine/Source/Runtime/Experimental/Chaos/Public/Chaos/ImplicitObject.h
@@ -86,7 +86,7 @@ namespace EImplicitObject
 		DisableCollisions = 1 << 2
 	};
 
-	const int32 FiniteConvex = IsConvex | HasBoundingBox;
+	inline const int32 FiniteConvex = IsConvex | HasBoundingBox;
 }
 
 
diff --git a/Engine/Source/Runtime/Experimental/ChaosCore/Public/Chaos/Core.h b/Engine/Source/Runtime/Experimental/ChaosCore/Public/Chaos/Core.h
--- a/Engine/Source/Runtime/Experimental/ChaosCore/Public/Chaos/Core.h
+++ b/Engine/Source/Runtime/Experimental/ChaosCore/Public/Chaos/Core.h
@@ -51,7 +51,7 @@ namespace Chaos
 	// NOTE: if you get a merge conflict on the GUID, you must replace it with a new GUID - do not accept the source or target
 	// or you will likely get DDC version conflicts resulting in crashes during load.
 	// Core version string for Chaos data. Any DDC builder dependent on Chaos for serialization should depend on this version
-	static const TCHAR* const ChaosVersionGUID = TEXT("FC60014F-0799-4BFD-B01C-44C1DEE81F32");
+	inline const TCHAR* const ChaosVersionGUID = TEXT("FC60014F-0799-4BFD-B01C-44C1DEE81F32");
 
 	inline FString GetChaosVersionStringInner()
 	{
diff --git a/Engine/Source/Runtime/GeometryCore/Public/IndexTypes.h b/Engine/Source/Runtime/GeometryCore/Public/IndexTypes.h
--- a/Engine/Source/Runtime/GeometryCore/Public/IndexTypes.h
+++ b/Engine/Source/Runtime/GeometryCore/Public/IndexTypes.h
@@ -10,7 +10,7 @@
 
 namespace IndexConstants
 {
-	constexpr int InvalidID = -1;
+	inline constexpr int InvalidID = -1;
 }
 
 namespace UE
diff --git a/Engine/Source/Runtime/Landscape/Public/LandscapeDataAccess.h b/Engine/Source/Runtime/Landscape/Public/LandscapeDataAccess.h
--- a/Engine/Source/Runtime/Landscape/Public/LandscapeDataAccess.h
+++ b/Engine/Source/Runtime/Landscape/Public/LandscapeDataAccess.h
@@ -21,8 +21,8 @@ class ULandscapeLayerInfoObject;
 
 namespace LandscapeDataAccess
 {
-	const int32 MaxValue = 65535;
-	const float MidValue = 32768.f;
+	inline const int32 MaxValue = 65535;
+	inline const float MidValue = 32768.f;
 	// Reserved 2 bits for other purpose
 	// Most significant bit - Visibility, 0 is visible(default), 1 is invisible
 	// 2nd significant bit - Triangle flip, not implemented yet
diff --git a/Engine/Source/Runtime/NavigationSystem/Public/NavMesh/RecastNavMesh.h b/Engine/Source/Runtime/NavigationSystem/Public/NavMesh/RecastNavMesh.h
--- a/Engine/Source/Runtime/NavigationSystem/Public/NavMesh/RecastNavMesh.h
+++ b/Engine/Source/Runtime/NavigationSystem/Public/NavMesh/RecastNavMesh.h
@@ -113,13 +113,13 @@ public:
 namespace ERecastPathFlags
 {
 	/** If set, path won't be post processed. */
-	const int32 SkipStringPulling = (1 << 0);
+	inline const int32 SkipStringPulling = (1 << 0);
 
 	/** If set, path will contain navigation corridor. */
-	const int32 GenerateCorridor = (1 << 1);
+	inline const int32 GenerateCorridor = (1 << 1);
 
 	/** Make your game-specific flags start at this index */
-	const uint8 FirstAvailableFlag = 2;
+	inline const uint8 FirstAvailableFlag = 2;
 }
 
 #if WITH_RECAST
diff --git a/Engine/Source/Runtime/Navmesh/Public/Detour/DetourLargeWorldCoordinates.h b/Engine/Source/Runtime/Navmesh/Public/Detour/DetourLargeWorldCoordinates.h
--- a/Engine/Source/Runtime/Navmesh/Public/Detour/DetourLargeWorldCoordinates.h
+++ b/Engine/Source/Runtime/Navmesh/Public/Detour/DetourLargeWorldCoordinates.h
@@ -16,5 +16,5 @@ typedef double dtReal;
 
 #endif // DT_LARGE_WORLD_COORDINATES_DISABLED
 
-constexpr dtReal DT_REAL_MAX = std::numeric_limits<dtReal>::max();
+inline constexpr dtReal DT_REAL_MAX = std::numeric_limits<dtReal>::max();
 
diff --git a/Engine/Source/Runtime/Navmesh/Public/Detour/DetourNavMesh.h b/Engine/Source/Runtime/Navmesh/Public/Detour/DetourNavMesh.h
--- a/Engine/Source/Runtime/Navmesh/Public/Detour/DetourNavMesh.h
+++ b/Engine/Source/Runtime/Navmesh/Public/Detour/DetourNavMesh.h
@@ -103,7 +103,7 @@ static const int DT_NAVMESH_STATE_VERSION = 1;
 static const unsigned short DT_EXT_LINK = 0x8000;
 
 /// A value that indicates the entity does not link to anything.
-static const unsigned int DT_NULL_LINK = 0xffffffff;
+inline const unsigned int DT_NULL_LINK = 0xffffffff;
 
 /// A flag that indicates that an off-mesh connection can be traversed in both directions. (Is bidirectional.)
 static const unsigned char DT_OFFMESH_CON_BIDIR = 0x01;
diff --git a/Engine/Source/Runtime/PhysicsCore/Private/SceneQueryChaosImp.h b/Engine/Source/Runtime/PhysicsCore/Private/SceneQueryChaosImp.h
--- a/Engine/Source/Runtime/PhysicsCore/Private/SceneQueryChaosImp.h
+++ b/Engine/Source/Runtime/PhysicsCore/Private/SceneQueryChaosImp.h
@@ -2,6 +2,11 @@
 
 #pragma once
 
+#include "Math/MathFwd.h"
+#include "PhysicsInterfaceDeclaresCore.h"
+
+enum class EHitFlags : uint16;
+
 inline bool LowLevelRaycastImp(const FVector& Start, const FVector& Dir, float DeltaMag, const Chaos::FImplicitObject& Shape, const FTransform ActorTM, EHitFlags OutputFlags, FHitRaycast& Hit)
 {
 	//TODO_SQ_IMPLEMENTATION
diff --git a/Engine/Source/Runtime/PhysicsCore/Public/PhysicsInterfaceWrapperShared.h b/Engine/Source/Runtime/PhysicsCore/Public/PhysicsInterfaceWrapperShared.h
--- a/Engine/Source/Runtime/PhysicsCore/Public/PhysicsInterfaceWrapperShared.h
+++ b/Engine/Source/Runtime/PhysicsCore/Public/PhysicsInterfaceWrapperShared.h
@@ -2,6 +2,8 @@
 
 #pragma once
 
+#include "HAL/Platform.h"
+
 enum class EQueryFlags : uint16
 {
 	None = 0,
diff --git a/Engine/Source/Runtime/PhysicsCore/Public/SQVerifier.h b/Engine/Source/Runtime/PhysicsCore/Public/SQVerifier.h
--- a/Engine/Source/Runtime/PhysicsCore/Public/SQVerifier.h
+++ b/Engine/Source/Runtime/PhysicsCore/Public/SQVerifier.h
@@ -7,6 +7,7 @@
 #include "Chaos/PBDRigidsEvolutionGBF.h"
 #include "PhysicsInterfaceUtilsCore.h"
 #include "PhysTestSerializer.h"
+#include "SQAccelerator.h"
 
 #ifndef SQ_REPLAY_TEST
 #define SQ_REPLAY_TEST(cond) bEnsureOnMismatch ? ensure(cond) : (cond)
diff --git a/Engine/Source/Runtime/RHI/Public/RHISurfaceDataConversion.h b/Engine/Source/Runtime/RHI/Public/RHISurfaceDataConversion.h
--- a/Engine/Source/Runtime/RHI/Public/RHISurfaceDataConversion.h
+++ b/Engine/Source/Runtime/RHI/Public/RHISurfaceDataConversion.h
@@ -6,7 +6,10 @@
 
 #pragma once
 
+#include "Math/Float16Color.h"
 #include "Math/PackedVector.h"
+#include "Math/Plane.h"
+#include "Math/UnrealMathUtility.h"
 
 /** Helper for accessing R10G10B10A2 colors. */
 struct FRHIR10G10B10A2
diff --git a/Engine/Source/Runtime/RHI/Public/RHIValidationUtils.h b/Engine/Source/Runtime/RHI/Public/RHIValidationUtils.h
--- a/Engine/Source/Runtime/RHI/Public/RHIValidationUtils.h
+++ b/Engine/Source/Runtime/RHI/Public/RHIValidationUtils.h
@@ -2,6 +2,8 @@
 
 #pragma once
 
+#include "Misc/Build.h"
+
 #if ENABLE_RHI_VALIDATION
 
 class FValidationRHIUtils
diff --git a/Engine/Source/Runtime/RenderCore/Public/ShaderMaterial.h b/Engine/Source/Runtime/RenderCore/Public/ShaderMaterial.h
--- a/Engine/Source/Runtime/RenderCore/Public/ShaderMaterial.h
+++ b/Engine/Source/Runtime/RenderCore/Public/ShaderMaterial.h
@@ -10,6 +10,7 @@
 #include "Templates/RefCounting.h"
 #include "Serialization/MemoryLayout.h"
 
+namespace ERHIFeatureLevel { enum Type : int; }
 struct FShaderCompilerEnvironment;
 
 struct FShaderGlobalDefines
diff --git a/Engine/Source/Runtime/RenderCore/Public/ShaderPermutation.h b/Engine/Source/Runtime/RenderCore/Public/ShaderPermutation.h
--- a/Engine/Source/Runtime/RenderCore/Public/ShaderPermutation.h
+++ b/Engine/Source/Runtime/RenderCore/Public/ShaderPermutation.h
@@ -8,6 +8,7 @@
 
 #include "CoreMinimal.h"
 
+struct FShaderCompilerEnvironment;
 
 /** Defines at compile time a boolean permutation dimension. */
 struct FShaderPermutationBool
diff --git a/Engine/Source/Runtime/Renderer/Public/MeshPassProcessor.h b/Engine/Source/Runtime/Renderer/Public/MeshPassProcessor.h
--- a/Engine/Source/Runtime/Renderer/Public/MeshPassProcessor.h
+++ b/Engine/Source/Runtime/Renderer/Public/MeshPassProcessor.h
@@ -680,7 +680,7 @@ struct FMeshProcessorShaders
  * This is tweaked so that the bindings for BasePass shaders of an average material using a FLocalVertexFactory fit into the inline storage.
  * Overflow of the inline storage will cause a heap allocation per draw (and corresponding cache miss on traversal)
  */
-const int32 NumInlineShaderBindings = 10;
+inline const int32 NumInlineShaderBindings = 10;
 
 /**
 * Debug only data for being able to backtrack the origin of given FMeshDrawCommand.
diff --git a/Engine/Source/Runtime/Slate/Public/Framework/Commands/Contexts/UIContentContext.h b/Engine/Source/Runtime/Slate/Public/Framework/Commands/Contexts/UIContentContext.h
--- a/Engine/Source/Runtime/Slate/Public/Framework/Commands/Contexts/UIContentContext.h
+++ b/Engine/Source/Runtime/Slate/Public/Framework/Commands/Contexts/UIContentContext.h
@@ -3,6 +3,7 @@
 #pragma once
 
 #include "CoreMinimal.h"
+#include "Framework/Commands/UIAction.h"
 #include "Layout/Visibility.h"
 #include "Styling/SlateTypes.h"
 #include "Framework/SlateDelegates.h"
diff --git a/Engine/Source/Runtime/Slate/Public/Framework/Commands/Contexts/UIIdentifierContext.h b/Engine/Source/Runtime/Slate/Public/Framework/Commands/Contexts/UIIdentifierContext.h
--- a/Engine/Source/Runtime/Slate/Public/Framework/Commands/Contexts/UIIdentifierContext.h
+++ b/Engine/Source/Runtime/Slate/Public/Framework/Commands/Contexts/UIIdentifierContext.h
@@ -3,6 +3,7 @@
 #pragma once
 
 #include "CoreMinimal.h"
+#include "Framework/Commands/UIAction.h"
 #include "Layout/Visibility.h"
 #include "Styling/SlateTypes.h"
 #include "Framework/SlateDelegates.h"
diff --git a/Engine/Source/Runtime/SlateCore/Private/Fonts/FontCacheFreeType.h b/Engine/Source/Runtime/SlateCore/Private/Fonts/FontCacheFreeType.h
--- a/Engine/Source/Runtime/SlateCore/Private/Fonts/FontCacheFreeType.h
+++ b/Engine/Source/Runtime/SlateCore/Private/Fonts/FontCacheFreeType.h
@@ -48,7 +48,7 @@
 namespace FreeTypeConstants
 {
 	/** The horizontal DPI we render at (horizontal and vertical) */
-	const uint32 RenderDPI = 96;
+	inline const uint32 RenderDPI = 96;
 } // namespace FreeTypeConstants
 
 
diff --git a/Engine/Source/Runtime/SlateCore/Private/Fonts/SlateFontRenderer.h b/Engine/Source/Runtime/SlateCore/Private/Fonts/SlateFontRenderer.h
--- a/Engine/Source/Runtime/SlateCore/Private/Fonts/SlateFontRenderer.h
+++ b/Engine/Source/Runtime/SlateCore/Private/Fonts/SlateFontRenderer.h
@@ -38,7 +38,7 @@ namespace SlateFontRendererUtils
 {
 
 /** Character used to substitute invalid font characters */
-const TCHAR InvalidSubChar = TEXT('\uFFFD');
+inline const TCHAR InvalidSubChar = TEXT('\uFFFD');
 
 #if WITH_FREETYPE
 
diff --git a/Engine/Source/Runtime/SlateCore/Private/Rendering/ElementBatcher.cpp b/Engine/Source/Runtime/SlateCore/Private/Rendering/ElementBatcher.cpp
--- a/Engine/Source/Runtime/SlateCore/Private/Rendering/ElementBatcher.cpp
+++ b/Engine/Source/Runtime/SlateCore/Private/Rendering/ElementBatcher.cpp
@@ -67,6 +67,15 @@ void FSlateElementBatch::SaveClippingState(const TArray<FSlateClippingState>& Pr
 	}*/
 }
 
+FSlateBatchData::FSlateBatchData()
+	: FirstRenderBatchIndex(INDEX_NONE)
+	, NumLayers(0)
+	, NumBatches(0)
+	, bIsStencilBufferRequired(false)
+{}
+
+FSlateBatchData::~FSlateBatchData() = default;
+
 void FSlateBatchData::ResetData()
 {
 	RenderBatches.Reset();
diff --git a/Engine/Source/Runtime/SlateCore/Public/Rendering/ElementBatcher.h b/Engine/Source/Runtime/SlateCore/Public/Rendering/ElementBatcher.h
--- a/Engine/Source/Runtime/SlateCore/Public/Rendering/ElementBatcher.h
+++ b/Engine/Source/Runtime/SlateCore/Public/Rendering/ElementBatcher.h
@@ -11,6 +11,7 @@
 
 class FSlateBatchData;
 class FSlateDrawElement;
+class FSlateRenderBatch;
 class FSlateRenderingPolicy;
 class FSlateShaderResource;
 class FSlateWindowElementList;
@@ -152,12 +153,8 @@ public:
 class FSlateBatchData
 {
 public:
-	FSlateBatchData()
-		: FirstRenderBatchIndex(INDEX_NONE)
-		, NumLayers(0)
-		, NumBatches(0)
-		, bIsStencilBufferRequired(false)
-	{}
+	SLATECORE_API FSlateBatchData();
+	SLATECORE_API ~FSlateBatchData();
 
 	SLATECORE_API void ResetData();
 
diff --git a/Engine/Source/Runtime/SlateCore/Public/Rendering/SlateObjectReferenceCollector.h b/Engine/Source/Runtime/SlateCore/Public/Rendering/SlateObjectReferenceCollector.h
--- a/Engine/Source/Runtime/SlateCore/Public/Rendering/SlateObjectReferenceCollector.h
+++ b/Engine/Source/Runtime/SlateCore/Public/Rendering/SlateObjectReferenceCollector.h
@@ -3,6 +3,7 @@
 #pragma once
 
 #include "CoreMinimal.h"
+#include "UObject/UObjectGlobals.h"
 
 /**
  * Used to gather all the referenced UObjects used by Slate when rendering a frame.
diff --git a/Engine/Source/Runtime/SlateCore/Public/Rendering/SlateResourceHandle.h b/Engine/Source/Runtime/SlateCore/Public/Rendering/SlateResourceHandle.h
--- a/Engine/Source/Runtime/SlateCore/Public/Rendering/SlateResourceHandle.h
+++ b/Engine/Source/Runtime/SlateCore/Public/Rendering/SlateResourceHandle.h
@@ -2,6 +2,9 @@
 
 #pragma once
 
+#include "HAL/Platform.h"
+#include "Templates/SharedPointer.h"
+
 class FSlateSharedHandleData;
 class FSlateShaderResourceProxy;
 
diff --git a/Engine/Source/Runtime/SlateCore/Public/Types/ReflectionMetadata.h b/Engine/Source/Runtime/SlateCore/Public/Types/ReflectionMetadata.h
--- a/Engine/Source/Runtime/SlateCore/Public/Types/ReflectionMetadata.h
+++ b/Engine/Source/Runtime/SlateCore/Public/Types/ReflectionMetadata.h
@@ -3,6 +3,7 @@
 #pragma once
 
 #include "CoreMinimal.h"
+#include "UObject/Class.h"
 #include "UObject/Object.h"
 #include "UObject/WeakObjectPtr.h"
 #include "Types/ISlateMetaData.h"
diff --git a/Engine/Source/Runtime/SlateCore/Public/Types/WidgetMouseEventsDelegate.h b/Engine/Source/Runtime/SlateCore/Public/Types/WidgetMouseEventsDelegate.h
--- a/Engine/Source/Runtime/SlateCore/Public/Types/WidgetMouseEventsDelegate.h
+++ b/Engine/Source/Runtime/SlateCore/Public/Types/WidgetMouseEventsDelegate.h
@@ -3,7 +3,9 @@
 #pragma once
 
 #include "CoreMinimal.h"
+#include "Delegates/DelegateCombinations.h"
 #include "Input/Events.h"
+#include "Input/Reply.h"
 #include "Layout/Geometry.h"
 
 /** Delegate type for handling mouse events */
diff --git a/Engine/Source/Runtime/SlateCore/Public/Widgets/Accessibility/SlateWidgetAccessibleTypes.h b/Engine/Source/Runtime/SlateCore/Public/Widgets/Accessibility/SlateWidgetAccessibleTypes.h
--- a/Engine/Source/Runtime/SlateCore/Public/Widgets/Accessibility/SlateWidgetAccessibleTypes.h
+++ b/Engine/Source/Runtime/SlateCore/Public/Widgets/Accessibility/SlateWidgetAccessibleTypes.h
@@ -3,6 +3,7 @@
 #pragma once
 
 #include "GenericPlatform/Accessibility/GenericAccessibleInterfaces.h"
+#include "Misc/Attribute.h"
 
 struct FAccessibleWidgetData
 {
