IPropertyHandle::GetMetaData and related functions (GetBoolMetaData, etc.) now always check the instance metadata as well. This is what you would expect to happen when calling SetInstanceMetaData, but was only being applied selectively by certain widgets.

Added IPropertyHandle::GetDoubleMetaData().

Reported on UDN.

diff --git a/Engine/Source/Editor/PropertyEditor/Private/PropertyHandleImpl.cpp b/Engine/Source/Editor/PropertyEditor/Private/PropertyHandleImpl.cpp
--- a/Engine/Source/Editor/PropertyEditor/Private/PropertyHandleImpl.cpp
+++ b/Engine/Source/Editor/PropertyEditor/Private/PropertyHandleImpl.cpp
@@ -2623,8 +2623,13 @@ FProperty* FPropertyHandleBase::GetMetaDataProperty() const
 
 bool FPropertyHandleBase::HasMetaData(const FName& Key) const
 {
-	FProperty* const MetaDataProperty = GetMetaDataProperty();
-	return (MetaDataProperty) ? MetaDataProperty->HasMetaData(Key) : false;
+	if (const FString* InstanceValue = GetInstanceMetaData(Key))
+	{
+		return true;
+	}
+
+	const FProperty* MetaDataProperty = GetMetaDataProperty();
+	return MetaDataProperty ? MetaDataProperty->HasMetaData(Key) : false;
 }
 
 const FString& FPropertyHandleBase::GetMetaData(const FName& Key) const
@@ -2632,34 +2637,50 @@ const FString& FPropertyHandleBase::GetMetaData(const FName& Key) const
 	// if not found, return a static empty string
 	static const FString EmptyString = TEXT("");
 
-	FProperty* const MetaDataProperty = GetMetaDataProperty();
-	return (MetaDataProperty) ? MetaDataProperty->GetMetaData(Key) : EmptyString;
+	if (const FString* InstanceValue = GetInstanceMetaData(Key))
+	{
+		return *InstanceValue;
+	}
+
+	const FProperty* MetaDataProperty = GetMetaDataProperty();
+	return MetaDataProperty ? MetaDataProperty->GetMetaData(Key) : EmptyString;
 }
 
 bool FPropertyHandleBase::GetBoolMetaData(const FName& Key) const
 {
-	FProperty* const MetaDataProperty = GetMetaDataProperty();
-	return (MetaDataProperty) ? MetaDataProperty->GetBoolMetaData(Key) : false;
+	if (!HasMetaData(Key))
+	{
+		return false;
+	}
+
+	const FString& StringValue = GetMetaData(Key);
+	return StringValue.IsEmpty() || StringValue == TEXT("true");
 }
 
 int32 FPropertyHandleBase::GetIntMetaData(const FName& Key) const
 {
-	FProperty* const MetaDataProperty = GetMetaDataProperty();
-	return (MetaDataProperty) ? MetaDataProperty->GetIntMetaData(Key) : 0;
+	const FString& StringValue = GetMetaData(Key);
+	return FCString::Atoi(*StringValue);
 }
 
 float FPropertyHandleBase::GetFloatMetaData(const FName& Key) const
 {
-	FProperty* const MetaDataProperty = GetMetaDataProperty();
-	return (MetaDataProperty) ? MetaDataProperty->GetFloatMetaData(Key) : 0.0f;
+	const FString& StringValue = GetMetaData(Key);
+	return FCString::Atof(*StringValue);
 }
 
-UClass* FPropertyHandleBase::GetClassMetaData(const FName& Key) const
+double FPropertyHandleBase::GetDoubleMetaData(const FName& Key) const
 {
-	FProperty* const MetaDataProperty = GetMetaDataProperty();
-	return (MetaDataProperty) ? MetaDataProperty->GetClassMetaData(Key) : nullptr;
+	const FString& StringValue = GetMetaData(Key);
+	return FCString::Atod(*StringValue);
 }
 
+UClass* FPropertyHandleBase::GetClassMetaData(const FName& Key) const
+{
+	const FString& ClassName = GetMetaData(Key);
+	UClass* FoundClass = UClass::TryFindTypeSlow<UClass>(ClassName);
+	return FoundClass;
+}
 
 void FPropertyHandleBase::SetInstanceMetaData(const FName& Key, const FString& Value)
 {
diff --git a/Engine/Source/Editor/PropertyEditor/Private/PropertyHandleImpl.h b/Engine/Source/Editor/PropertyEditor/Private/PropertyHandleImpl.h
--- a/Engine/Source/Editor/PropertyEditor/Private/PropertyHandleImpl.h
+++ b/Engine/Source/Editor/PropertyEditor/Private/PropertyHandleImpl.h
@@ -459,6 +459,7 @@ public:
 	virtual bool GetBoolMetaData(const FName& Key) const override;
 	virtual int32 GetIntMetaData(const FName& Key) const override;
 	virtual float GetFloatMetaData(const FName& Key) const override;
+	virtual double GetDoubleMetaData(const FName& Key) const override;
 	virtual UClass* GetClassMetaData(const FName& Key) const override;
 	virtual void SetInstanceMetaData(const FName& Key, const FString& Value) override;
 	virtual const FString* GetInstanceMetaData(const FName& Key) const override;
diff --git a/Engine/Source/Editor/PropertyEditor/Public/PropertyHandle.h b/Engine/Source/Editor/PropertyEditor/Public/PropertyHandle.h
--- a/Engine/Source/Editor/PropertyEditor/Public/PropertyHandle.h
+++ b/Engine/Source/Editor/PropertyEditor/Public/PropertyHandle.h
@@ -114,6 +114,14 @@ public:
 	 */
 	virtual float GetFloatMetaData(const FName& Key) const = 0;
 
+	/**
+	 * Find the metadata value associated with the key and return double
+	 *
+	 * @param Key The key to lookup in the metadata
+	 * @return the double value stored in the metadata.
+	 */
+	virtual double GetDoubleMetaData(const FName& Key) const = 0;
+
 	/**
 	 * Find the metadata value associated with the key and return UClass*
 	 *
