It is now possible to add external objects and hide the object root by specifying an FAddPropertyParams struct.

Reported on UDN.

diff --git a/Engine/Source/Editor/PropertyEditor/Private/CustomChildBuilder.cpp b/Engine/Source/Editor/PropertyEditor/Private/CustomChildBuilder.cpp
--- a/Engine/Source/Editor/PropertyEditor/Private/CustomChildBuilder.cpp
+++ b/Engine/Source/Editor/PropertyEditor/Private/CustomChildBuilder.cpp
@@ -56,7 +56,6 @@ IDetailPropertyRow& FCustomChildrenBuilder::AddProperty( TSharedRef<IPropertyHan
 	return *NewCustomization.PropertyRow;
 }
 
-
 IDetailPropertyRow* FCustomChildrenBuilder::AddExternalStructure(TSharedRef<FStructOnScope> ChildStructure, FName UniqueIdName)
 {
 	return AddExternalStructureProperty(ChildStructure, NAME_None, FAddPropertyParams().UniqueId(UniqueIdName));
@@ -85,33 +84,42 @@ IDetailPropertyRow* FCustomChildrenBuilder::AddExternalStructureProperty(TShared
 	return NewRow.Get();
 }
 
-IDetailPropertyRow* FCustomChildrenBuilder::AddExternalObjects(const TArray<UObject *>& Objects, FName UniqueIdName)
+TArray<TSharedPtr<IPropertyHandle>> FCustomChildrenBuilder::AddAllExternalStructureProperties(TSharedRef<FStructOnScope> ChildStructure)
 {
-	return AddExternalObjectProperty(Objects, NAME_None, FAddPropertyParams().UniqueId(UniqueIdName));
+	return ParentCategory.Pin()->AddAllExternalStructureProperties(ChildStructure);
 }
 
-TArray<TSharedPtr<IPropertyHandle>> FCustomChildrenBuilder::AddAllExternalStructureProperties(TSharedRef<FStructOnScope> ChildStructure)
+IDetailPropertyRow* FCustomChildrenBuilder::AddExternalObjects(const TArray<UObject*>& Objects, FName UniqueIdName)
 {
-	return ParentCategory.Pin()->AddAllExternalStructureProperties(ChildStructure);
+	FAddPropertyParams Params = FAddPropertyParams()
+		.UniqueId(UniqueIdName)
+		.AllowChildren(true);
+
+	return AddExternalObjects(Objects, Params);
 }
 
-IDetailPropertyRow* FCustomChildrenBuilder::AddExternalObjectProperty(const TArray<UObject*>& Objects, FName PropertyName, const FAddPropertyParams& Params)
+IDetailPropertyRow* FCustomChildrenBuilder::AddExternalObjects(const TArray<UObject*>& Objects, const FAddPropertyParams& Params)
 {
-	FDetailLayoutCustomization NewCustomization;
+	return AddExternalObjectProperty(Objects, NAME_None, Params);
+}
 
+IDetailPropertyRow* FCustomChildrenBuilder::AddExternalObjectProperty(const TArray<UObject*>& Objects, FName PropertyName, const FAddPropertyParams& Params)
+{
 	TSharedRef<FDetailCategoryImpl> ParentCategoryRef = ParentCategory.Pin().ToSharedRef();
 
+	FDetailLayoutCustomization NewCustomization;
 	FDetailPropertyRow::MakeExternalPropertyRowCustomization(Objects, PropertyName, ParentCategoryRef, NewCustomization, Params);
 
-	TSharedPtr<FDetailPropertyRow> NewRow = NewCustomization.PropertyRow;
+	if (Params.ShouldHideRootObjectNode() && NewCustomization.HasPropertyNode() && NewCustomization.GetPropertyNode()->AsObjectNode())
+	{
+		NewCustomization.PropertyRow->SetForceShowOnlyChildren(true);
+	}
 
+	TSharedPtr<FDetailPropertyRow> NewRow = NewCustomization.PropertyRow;
 	if (NewRow.IsValid())
 	{
 		NewRow->SetCustomExpansionId(Params.GetUniqueId());
 
-		TSharedPtr<FPropertyNode> PropertyNode = NewRow->GetPropertyNode();
-		TSharedPtr<FObjectPropertyNode> RootNode = StaticCastSharedRef<FObjectPropertyNode>(PropertyNode->FindObjectItemParent()->AsShared());
-
 		ChildCustomizations.Add(NewCustomization);
 	}
 
diff --git a/Engine/Source/Editor/PropertyEditor/Private/CustomChildBuilder.h b/Engine/Source/Editor/PropertyEditor/Private/CustomChildBuilder.h
--- a/Engine/Source/Editor/PropertyEditor/Private/CustomChildBuilder.h
+++ b/Engine/Source/Editor/PropertyEditor/Private/CustomChildBuilder.h
@@ -24,9 +24,11 @@ public:
 	virtual IDetailGroup& AddGroup(FName GroupName, const FText& LocalizedDisplayName) override;
 	virtual FDetailWidgetRow& AddCustomRow(const FText& SearchString) override;
 	virtual IDetailPropertyRow& AddProperty(TSharedRef<IPropertyHandle> PropertyHandle) override;
-	virtual IDetailPropertyRow* AddExternalObjects(const TArray<UObject*>& Objects, FName UniqueIdName = NAME_None) override;
+	UE_DEPRECATED(5.1, "Please use the overload that takes an FAddPropertyParams instead.")
+	virtual IDetailPropertyRow* AddExternalObjects(const TArray<UObject*>& Objects, FName UniqueIdName) override;
+	virtual IDetailPropertyRow* AddExternalObjects(const TArray<UObject*>& Objects, const FAddPropertyParams& Params) override;
 	virtual IDetailPropertyRow* AddExternalObjectProperty(const TArray<UObject*>& Objects, FName PropertyName, const FAddPropertyParams& Params) override;
-	virtual IDetailPropertyRow* AddExternalStructure(TSharedRef<FStructOnScope> ChildStructure, FName UniqueIdName = NAME_None) override;
+	virtual IDetailPropertyRow* AddExternalStructure(TSharedRef<FStructOnScope> ChildStructure, FName UniqueIdName) override;
 	virtual IDetailPropertyRow* AddExternalStructureProperty(TSharedRef<FStructOnScope> ChildStructure, FName PropertyName, const FAddPropertyParams& Params) override;
 	virtual TArray<TSharedPtr<IPropertyHandle>> AddAllExternalStructureProperties(TSharedRef<FStructOnScope> ChildStructure) override;
 	virtual TSharedRef<SWidget> GenerateStructValueWidget(TSharedRef<IPropertyHandle> StructPropertyHandle) override;
diff --git a/Engine/Source/Editor/PropertyEditor/Public/IDetailChildrenBuilder.h b/Engine/Source/Editor/PropertyEditor/Public/IDetailChildrenBuilder.h
--- a/Engine/Source/Editor/PropertyEditor/Public/IDetailChildrenBuilder.h
+++ b/Engine/Source/Editor/PropertyEditor/Public/IDetailChildrenBuilder.h
@@ -65,14 +65,21 @@ public:
 	 */
 	virtual IDetailPropertyRow* AddExternalObjects(const TArray<UObject*>& Objects, FName UniqueIdName = NAME_None) = 0;
 
+	/**
+	 * Adds a set of objects to as a child.  Similar to details panels, all objects will be visible in the details panel as set of properties from the common base class from the list of objects
+	 *
+	 * @param  Objects			The objects to add
+	 * @param  Params			Optional parameters for customizing the display of the property rows.
+	 * @return The header row generated for this set of objects by the details panel
+	 */
+	virtual IDetailPropertyRow* AddExternalObjects(const TArray<UObject*>& Objects, const FAddPropertyParams& Params = FAddPropertyParams()) = 0;
+
 	/**
 	 * Adds a set of objects to as a child.  Similar to details panels, all objects will be visible in the details panel as set of properties from the common base class from the list of objects
 	 *
 	 * @param  Objects			The objects to add
 	 * @param  PropertyName		Name of a property inside the object(s) to add.
-	 * @param  UniqueIdName		Optional identifier that uniquely identifies this object among other objects of the same type.  If this is empty, saving and restoring expansion state of this object may not work
-	 * @param  bAllowChildrenOverride Allows customization of how the new root property node is expanded when this is added.
-	 * @param  bCreateCategoryNodesOverride Allows customization of how the new root node's category is displayed (or not).
+	 * @param  Params			Optional parameters for customizing the display of the property rows.
 	 * @return The header row generated for this set of objects by the details panel
 	 */
 	virtual IDetailPropertyRow* AddExternalObjectProperty(const TArray<UObject*>& Objects, FName PropertyName, const FAddPropertyParams& Params = FAddPropertyParams()) = 0;
