Clear keyboard focus when setting an empty set of selected objects on an SDetailsView. This commits any pending changes before the widgets are destroyed.

This is necessary because although we register an OnRowReleased handler, we also register a Visibility handler in SDetailsViewBase::GetTreeVisibility(). The OnRowReleased delegate is called when the tree next Ticks(), but when the tree is cleared, it also gets its visibility set to Collapsed, which causes the next Tick() to not occur. Since there's no way to force the tree to Tick(), this is the next best thing.

Added a few trace macros.

#rnx

diff --git a/Engine/Source/Editor/PropertyEditor/Private/SDetailsView.cpp b/Engine/Source/Editor/PropertyEditor/Private/SDetailsView.cpp
--- a/Engine/Source/Editor/PropertyEditor/Private/SDetailsView.cpp
+++ b/Engine/Source/Editor/PropertyEditor/Private/SDetailsView.cpp
@@ -696,6 +696,8 @@ void SDetailsView::SetObjectArrayPrivate(const TArray<UObject*>& InObjects)
 		return;
 	}
 
+	TRACE_CPUPROFILER_EVENT_SCOPE(SDetailsView::SetObjectArrayPrivate);
+
 	TGuardValue<bool> RefreshGuard(bIsRefreshing, true);
 
 	double StartTime = FPlatformTime::Seconds();
@@ -861,6 +863,8 @@ void SDetailsView::RemoveDeletedObjects(const TArray<UObject*>& DeletedObjects)
 /** Called before during SetObjectArray before we change the objects being observed */
 void SDetailsView::PreSetObject(int32 InNewNumObjects)
 {
+	TRACE_CPUPROFILER_EVENT_SCOPE(SDetailsView::PreSetObject);
+
 	TSharedPtr<SColorPicker> ExistingColorPicker = GetColorPicker();
 	if (ExistingColorPicker.IsValid()
 		&& ExistingColorPicker->GetOptionalOwningDetailsView().IsValid()
@@ -869,6 +873,8 @@ void SDetailsView::PreSetObject(int32 InNewNumObjects)
 		DestroyColorPicker();
 	}
 
+	ClearKeyboardFocusIfWithin(AsShared());
+
 	// Save existing expanded items first
 	for(TSharedPtr<FComplexPropertyNode>& RootNode : RootPropertyNodes)
 	{
@@ -911,6 +917,8 @@ void SDetailsView::PreSetObject(int32 InNewNumObjects)
 /** Called at the end of SetObjectArray after we change the objects being observed */
 void SDetailsView::PostSetObject(const TArray<FDetailsViewObjectRoot>& Roots)
 {
+	TRACE_CPUPROFILER_EVENT_SCOPE(SDetailsView::PostSetObject);
+
 	TSharedPtr<SColorPicker> ExistingColorPicker = GetColorPicker();
 	if (ExistingColorPicker.IsValid()
 		&& (!ExistingColorPicker->GetOptionalOwningDetailsView().IsValid()
@@ -968,7 +976,7 @@ void SDetailsView::PostSetObject(const TArray<FDetailsViewObjectRoot>& Roots)
 		InitParams.bCreateDisableEditOnInstanceNodes = HasClassDefaultObject();
 		break;
 	default:
-		check(false);
+		checkNoEntry();
 	}
 
 	for( TSharedPtr<FComplexPropertyNode>& ComplexRootNode : RootPropertyNodes )
diff --git a/Engine/Source/Editor/PropertyEditor/Private/SDetailsViewBase.cpp b/Engine/Source/Editor/PropertyEditor/Private/SDetailsViewBase.cpp
--- a/Engine/Source/Editor/PropertyEditor/Private/SDetailsViewBase.cpp
+++ b/Engine/Source/Editor/PropertyEditor/Private/SDetailsViewBase.cpp
@@ -70,23 +70,27 @@ TSharedRef<ITableRow> SDetailsViewBase::OnGenerateRowForDetailTree(TSharedRef<FD
 	return InTreeNode->GenerateWidgetForTableView(OwnerTable, DetailsViewArgs.bAllowFavoriteSystem);
 }
 
-void SDetailsViewBase::OnRowReleasedForDetailTree(const TSharedRef<ITableRow>& TableRow)
+void SDetailsViewBase::ClearKeyboardFocusIfWithin(const TSharedRef<SWidget>& Widget) const
 {
 	// search upwards from the current keyboard-focused widget to see if it's contained in our row
 	TSharedPtr<SWidget> CurrentWidget = FSlateApplication::Get().GetKeyboardFocusedWidget();
 	while (CurrentWidget.IsValid())
 	{
-		if (CurrentWidget == TableRow->AsWidget())
+		if (CurrentWidget == Widget)
 		{
 			// if so, clear focus so that any pending value changes are committed
 			FSlateApplication::Get().ClearKeyboardFocus();
 			return;
 		}
-
 		CurrentWidget = CurrentWidget->GetParentWidget();
 	}
 }
 
+void SDetailsViewBase::OnRowReleasedForDetailTree(const TSharedRef<ITableRow>& TableRow)
+{
+	ClearKeyboardFocusIfWithin(TableRow->AsWidget());
+}
+
 void SDetailsViewBase::SetRootExpansionStates(const bool bExpand, const bool bRecurse)
 {
 	if(ContainsMultipleTopLevelObjects())
@@ -406,6 +410,8 @@ void SDetailsViewBase::SetColorPropertyFromColorPicker(FLinearColor NewColor)
 
 void SDetailsViewBase::UpdatePropertyMaps()
 {
+	TRACE_CPUPROFILER_EVENT_SCOPE(SDetailsView::UpdatePropertyMaps);
+
 	RootTreeNodes.Empty();
 
 	for(FDetailLayoutData& LayoutData : DetailLayouts)
@@ -938,6 +944,8 @@ void SDetailsViewBase::HandlePendingCleanup()
 /** Ticks the property view.  This function performs a data consistency check */
 void SDetailsViewBase::Tick( const FGeometry& AllottedGeometry, const double InCurrentTime, const float InDeltaTime )
 {
+	TRACE_CPUPROFILER_EVENT_SCOPE(SDetailsViewBase::Tick);
+
 	HandlePendingCleanup();
 
 	FDetailsViewConfig* ViewConfig = GetMutableViewConfig();
@@ -1254,6 +1262,8 @@ void SDetailsViewBase::RestorePreSearchExpandedItems()
 
 void SDetailsViewBase::SaveExpandedItems(TSharedRef<FPropertyNode> StartNode)
 {
+	TRACE_CPUPROFILER_EVENT_SCOPE(SDetailsView::SaveExpandedItems);
+
 	if (bRunningDeferredActions)
 	{
 		// Deferred actions can manipulate the tree, eg. Add Item
@@ -1323,6 +1333,8 @@ void SDetailsViewBase::SaveExpandedItems(TSharedRef<FPropertyNode> StartNode)
 
 void SDetailsViewBase::RestoreAllExpandedItems()
 {
+	TRACE_CPUPROFILER_EVENT_SCOPE(SDetailsViewBase::RestoreAllExpandedItems);
+
 	if (bRunningDeferredActions)
 	{
 		// Deferred actions can manipulate the tree, eg. Add Item
@@ -1425,6 +1437,8 @@ void SDetailsViewBase::FilterRootNode(const TSharedPtr<FComplexPropertyNode>& Ro
 
 void SDetailsViewBase::UpdateFilteredDetails()
 {
+	TRACE_CPUPROFILER_EVENT_SCOPE(SDetailsViewBase::UpdateFilteredDetails);
+
 	RootTreeNodes.Reset();
 
 	FDetailNodeList InitialRootNodeList;
diff --git a/Engine/Source/Editor/PropertyEditor/Private/SDetailsViewBase.h b/Engine/Source/Editor/PropertyEditor/Private/SDetailsViewBase.h
--- a/Engine/Source/Editor/PropertyEditor/Private/SDetailsViewBase.h
+++ b/Engine/Source/Editor/PropertyEditor/Private/SDetailsViewBase.h
@@ -355,6 +355,9 @@ protected:
 	const FDetailsViewConfig* GetConstViewConfig() const;
 	void SaveViewConfig();
 
+	/** Clear the keyboard focus if it's inside the given widget. */
+	void ClearKeyboardFocusIfWithin(const TSharedRef<SWidget>& Widget) const;
+
 protected:
 	/** The user defined args for the details view */
 	FDetailsViewArgs DetailsViewArgs;
