Removed unnecessary CoreMinimal.h include. Removed inline forward declarations.

#rnx

diff --git a/Engine/Source/Editor/PropertyEditor/Public/PropertyEditorDelegates.h b/Engine/Source/Editor/PropertyEditor/Public/PropertyEditorDelegates.h
--- a/Engine/Source/Editor/PropertyEditor/Public/PropertyEditorDelegates.h
+++ b/Engine/Source/Editor/PropertyEditor/Public/PropertyEditorDelegates.h
@@ -5,7 +5,6 @@
 #include "Containers/Array.h"
 #include "Containers/Map.h"
 #include "Containers/UnrealString.h"
-#include "CoreMinimal.h"
 #include "Delegates/Delegate.h"
 #include "Framework/Commands/UIAction.h"
 #include "HAL/Platform.h"
@@ -23,16 +22,18 @@ class FPropertyPath;
 class IDetailCustomization;
 class IDetailTreeNode;
 class IPropertyHandle;
+class IPropertyTreeRow;
 class IPropertyTypeCustomization;
 class IPropertyTypeIdentifier;
+class SHeaderRow;
 class SWidget;
 class UObject;
 struct FPropertyChangedEvent;
 
 struct FPropertyAndParent
 {
-	explicit FPropertyAndParent(const TSharedRef<class IPropertyHandle>& InPropertyHandle);
-	explicit FPropertyAndParent(const TSharedRef<class FPropertyNode>& InPropertyNode);
+	explicit FPropertyAndParent(const TSharedRef<IPropertyHandle>& InPropertyHandle);
+	explicit FPropertyAndParent(const TSharedRef<FPropertyNode>& InPropertyNode);
 
 	/** The property always exists */
 	const FProperty& Property;
@@ -57,10 +58,10 @@ private:
 };
 
 /** Delegate called to see if a property should be visible */
-DECLARE_DELEGATE_RetVal_OneParam( bool, FIsPropertyVisible, const FPropertyAndParent& );
+DECLARE_DELEGATE_RetVal_OneParam(bool, FIsPropertyVisible, const FPropertyAndParent&);
 
 /** Delegate called to see if a property should be read-only */
-DECLARE_DELEGATE_RetVal_OneParam( bool, FIsPropertyReadOnly, const FPropertyAndParent& );
+DECLARE_DELEGATE_RetVal_OneParam(bool, FIsPropertyReadOnly, const FPropertyAndParent&);
 
 /** Delegate called to determine if a custom row should be visible. */
 DECLARE_DELEGATE_RetVal_TwoParams(bool, FIsCustomRowVisible, FName /*InRowName*/, FName /*InParentName*/);
@@ -69,33 +70,33 @@ DECLARE_DELEGATE_RetVal_TwoParams(bool, FIsCustomRowVisible, FName /*InRowName*/
 DECLARE_DELEGATE_RetVal_TwoParams(bool, FIsCustomRowReadOnly, FName /*InRowName*/, FName /*InParentName*/);
 
 /** Delegate called to get a detail layout for a specific object class */
-DECLARE_DELEGATE_RetVal( TSharedRef<IDetailCustomization>, FOnGetDetailCustomizationInstance );
+DECLARE_DELEGATE_RetVal(TSharedRef<IDetailCustomization>, FOnGetDetailCustomizationInstance);
 
 /** Delegate called to get a property layout for a specific property type */
-DECLARE_DELEGATE_RetVal( TSharedRef<IPropertyTypeCustomization>, FOnGetPropertyTypeCustomizationInstance );
+DECLARE_DELEGATE_RetVal(TSharedRef<IPropertyTypeCustomization>, FOnGetPropertyTypeCustomizationInstance);
 
 /** Notification for when a property view changes */
-DECLARE_DELEGATE_TwoParams( FOnObjectArrayChanged, const FString&, const TArray<UObject*>& );
+DECLARE_DELEGATE_TwoParams(FOnObjectArrayChanged, const FString&, const TArray<UObject*>&);
 
 /** Notification for when displayed properties changes (for instance, because the user has filtered some properties */
-DECLARE_DELEGATE( FOnDisplayedPropertiesChanged );
+DECLARE_DELEGATE(FOnDisplayedPropertiesChanged);
 
 /** Notification for when a property selection changes. */
-DECLARE_DELEGATE_OneParam( FOnPropertySelectionChanged, FProperty* )
+DECLARE_DELEGATE_OneParam(FOnPropertySelectionChanged, FProperty*)
 
 /** Notification for when a property is double clicked by the user*/
-DECLARE_DELEGATE_OneParam( FOnPropertyDoubleClicked, FProperty* )
+DECLARE_DELEGATE_OneParam(FOnPropertyDoubleClicked, FProperty*)
 
 /** Notification for when a property is clicked by the user*/
-DECLARE_DELEGATE_OneParam( FOnPropertyClicked, const TSharedPtr<FPropertyPath >& )
+DECLARE_DELEGATE_OneParam(FOnPropertyClicked, const TSharedPtr<FPropertyPath >&)
 
 /** */
-DECLARE_DELEGATE_OneParam( FConstructExternalColumnHeaders, const TSharedRef< class SHeaderRow >& )
+DECLARE_DELEGATE_OneParam(FConstructExternalColumnHeaders, const TSharedRef< SHeaderRow >&)
 
-DECLARE_DELEGATE_RetVal_TwoParams( TSharedRef< SWidget >, FConstructExternalColumnCell,  const FName& /*ColumnName*/, const TSharedRef< class IPropertyTreeRow >& /*RowWidget*/)
+DECLARE_DELEGATE_RetVal_TwoParams(TSharedRef< SWidget >, FConstructExternalColumnCell, const FName& /*ColumnName*/, const TSharedRef< IPropertyTreeRow >& /*RowWidget*/)
 
 /** Delegate called to see if a property editing is enabled */
-DECLARE_DELEGATE_RetVal(bool, FIsPropertyEditingEnabled );
+DECLARE_DELEGATE_RetVal(bool, FIsPropertyEditingEnabled);
 
 /**
  * A delegate which is called after properties have been edited and PostEditChange has been called on all objects.
@@ -108,7 +109,7 @@ DECLARE_MULTICAST_DELEGATE_OneParam(FOnFinishedChangingProperties, const FProper
  * A property row extension arguments is displayed at the end of a property row, either inline or as a button.
  */
 struct FOnGenerateGlobalRowExtensionArgs
-{	
+{
 	/** The detail row's property handle. */
 	TSharedPtr<IPropertyHandle> PropertyHandle;
 
@@ -125,8 +126,8 @@ struct FOnGenerateGlobalRowExtensionArgs
 	FProperty* Property = nullptr;
 };
 
-/** 
- * A property row extension button is displayed at the end of a property row, either inline as a button, 
+/**
+ * A property row extension button is displayed at the end of a property row, either inline as a button,
  * or in a dropdown when not all buttons can fit.
  */
 struct FPropertyRowExtensionButton
