Added PropertyCustomizationHelpers::GetClassesFromMetadataString() to parse and load classes that are stored in a metadata string.

This is used for AllowedClasses and DisallowedClasses metadata in SPropertyEditorAsset.

#rnx

diff --git a/Engine/Source/Editor/PropertyEditor/Private/PropertyCustomizationHelpers.cpp b/Engine/Source/Editor/PropertyEditor/Private/PropertyCustomizationHelpers.cpp
--- a/Engine/Source/Editor/PropertyEditor/Private/PropertyCustomizationHelpers.cpp
+++ b/Engine/Source/Editor/PropertyEditor/Private/PropertyCustomizationHelpers.cpp
@@ -797,6 +797,58 @@ void PropertyCustomizationHelpers::MakeInstancedPropertyCustomUI(TMap<FName, IDe
 	}
 }
 
+TArray<const UClass*> PropertyCustomizationHelpers::GetClassesFromMetadataString(const FString& MetadataString)
+{
+	if (MetadataString.IsEmpty())
+	{
+		return TArray<const UClass*>();
+	}
+
+	auto FindClass = [](const FString& InClassName) -> const UClass*
+	{
+		const UClass* Class = UClass::TryFindTypeSlow<const UClass>(InClassName, EFindFirstObjectOptions::EnsureIfAmbiguous);
+		if (!Class)
+		{
+			Class = LoadObject<const UClass>(nullptr, *InClassName);
+		}
+		return Class;
+	};
+		
+	TArray<FString> ClassNames;
+	MetadataString.ParseIntoArrayWS(ClassNames, TEXT(","), true);
+
+	TArray<const UClass*> Classes;
+	Classes.Reserve(ClassNames.Num());
+
+	for (const FString& ClassName : ClassNames)
+	{
+		const UClass* Class = FindClass(ClassName);
+		if (!Class)
+		{
+			continue;
+		}
+
+		// If the class is an interface, expand it to be all classes in memory that implement the class.
+		if (Class->HasAnyClassFlags(CLASS_Interface))
+		{
+			for (TObjectIterator<UClass> ClassIt; ClassIt; ++ClassIt)
+			{
+				const UClass* ClassWithInterface = (*ClassIt);
+				if (ClassWithInterface->ImplementsInterface(Class))
+				{
+					Classes.Add(ClassWithInterface);
+				}
+			}
+		}
+		else
+		{
+			Classes.Add(Class);
+		}
+	}
+
+	return Classes;
+}
+
 //////////////////////////////////////////////////////////////////////////
 //
 // Sections list
diff --git a/Engine/Source/Editor/PropertyEditor/Private/UserInterface/PropertyEditor/SPropertyEditorAsset.cpp b/Engine/Source/Editor/PropertyEditor/Private/UserInterface/PropertyEditor/SPropertyEditorAsset.cpp
--- a/Engine/Source/Editor/PropertyEditor/Private/UserInterface/PropertyEditor/SPropertyEditorAsset.cpp
+++ b/Engine/Source/Editor/PropertyEditor/Private/UserInterface/PropertyEditor/SPropertyEditorAsset.cpp
@@ -57,7 +57,7 @@ static const FProperty* GetActualMetadataProperty(const FProperty* Property)
 }
 
 // Helper to support both meta=(TagName) and meta=(TagName=true) syntaxes
-static bool GetTagOrBoolMetadata(const FProperty* Property, const TCHAR* TagName, bool bDefault)
+static bool GetTagOrBoolMetadata(const FProperty* Property, FName TagName, bool bDefault)
 {
 	bool bResult = bDefault;
 
@@ -131,87 +131,16 @@ void SPropertyEditorAsset::InitializeClassFilters(const FProperty* Property)
 	// Account for the allowed classes specified in the property metadata
 	const FProperty* MetadataProperty = GetActualMetadataProperty(Property);
 
-	bExactClass = GetTagOrBoolMetadata(MetadataProperty, TEXT("ExactClass"), false);
-
-	auto FindClass = [](const FString& InClassName)
-	{
-		UClass* Class = UClass::TryFindTypeSlow<UClass>(InClassName, EFindFirstObjectOptions::EnsureIfAmbiguous);
-		if (!Class)
-		{
-			Class = LoadObject<UClass>(nullptr, *InClassName);
-		}
-		return Class;
-	};
-
-	const FString AllowedClassesFilterString = MetadataProperty->GetMetaData(TEXT("AllowedClasses"));
-	if (!AllowedClassesFilterString.IsEmpty())
-	{
-		TArray<FString> AllowedClassFilterNames;
-		AllowedClassesFilterString.ParseIntoArrayWS(AllowedClassFilterNames, TEXT(","), true);
-
-		for (const FString& ClassName : AllowedClassFilterNames)
-		{
-			UClass* Class = FindClass(ClassName);
-
-			if (Class)
-			{
-				// If the class is an interface, expand it to be all classes in memory that implement the class.
-				if (Class->HasAnyClassFlags(CLASS_Interface))
-				{
-					for (TObjectIterator<UClass> ClassIt; ClassIt; ++ClassIt)
-					{
-						UClass* const ClassWithInterface = (*ClassIt);
-						if (ClassWithInterface->ImplementsInterface(Class))
-						{
-							AllowedClassFilters.Add(ClassWithInterface);
-						}
-					}
-				}
-				else
-				{
-					AllowedClassFilters.Add(Class);
-				}
-			}
-		}
-	}
+	bExactClass = GetTagOrBoolMetadata(MetadataProperty, "ExactClass", false);
 	
+	AllowedClassFilters = PropertyCustomizationHelpers::GetClassesFromMetadataString(MetadataProperty->GetMetaData("AllowedClasses"));
 	if (AllowedClassFilters.Num() == 0)
 	{
 		// always add the object class to the filters
 		AllowedClassFilters.Add(ObjectClass);
 	}
 
-	const FString DisallowedClassesFilterString = MetadataProperty->GetMetaData(TEXT("DisallowedClasses"));
-	if (!DisallowedClassesFilterString.IsEmpty())
-	{
-		TArray<FString> DisallowedClassFilterNames;
-		DisallowedClassesFilterString.ParseIntoArrayWS(DisallowedClassFilterNames, TEXT(","), true);
-
-		for (const FString& ClassName : DisallowedClassFilterNames)
-		{
-			UClass* Class = FindClass(ClassName);
-
-			if (Class)
-			{
-				// If the class is an interface, expand it to be all classes in memory that implement the class.
-				if (Class->HasAnyClassFlags(CLASS_Interface))
-				{
-					for (TObjectIterator<UClass> ClassIt; ClassIt; ++ClassIt)
-					{
-						UClass* const ClassWithInterface = (*ClassIt);
-						if (ClassWithInterface->ImplementsInterface(Class))
-						{
-							DisallowedClassFilters.Add(ClassWithInterface);
-						}
-					}
-				}
-				else
-				{
-					DisallowedClassFilters.Add(Class);
-				}
-			}
-		}
-	}
+	DisallowedClassFilters = PropertyCustomizationHelpers::GetClassesFromMetadataString(MetadataProperty->GetMetaData("DisallowedClasses"));
 }
 
 void SPropertyEditorAsset::InitializeAssetDataTags(const FProperty* Property)
@@ -222,7 +151,7 @@ void SPropertyEditorAsset::InitializeAssetDataTags(const FProperty* Property)
 	}
 
 	const FProperty* MetadataProperty = GetActualMetadataProperty(Property);
-	const FString DisallowedAssetDataTagsFilterString = MetadataProperty->GetMetaData(TEXT("DisallowedAssetDataTags"));
+	const FString DisallowedAssetDataTagsFilterString = MetadataProperty->GetMetaData("DisallowedAssetDataTags");
 	if (!DisallowedAssetDataTagsFilterString.IsEmpty())
 	{
 		TArray<FString> DisallowedAssetDataTagsAndValues;
@@ -243,7 +172,7 @@ void SPropertyEditorAsset::InitializeAssetDataTags(const FProperty* Property)
 		}
 	}
 
-	const FString RequiredAssetDataTagsFilterString = MetadataProperty->GetMetaData(TEXT("RequiredAssetDataTags"));
+	const FString RequiredAssetDataTagsFilterString = MetadataProperty->GetMetaData("RequiredAssetDataTags");
 	if (!RequiredAssetDataTagsFilterString.IsEmpty())
 	{
 		TArray<FString> RequiredAssetDataTagsAndValues;
diff --git a/Engine/Source/Editor/PropertyEditor/Public/PropertyCustomizationHelpers.h b/Engine/Source/Editor/PropertyEditor/Public/PropertyCustomizationHelpers.h
--- a/Engine/Source/Editor/PropertyEditor/Public/PropertyCustomizationHelpers.h
+++ b/Engine/Source/Editor/PropertyEditor/Public/PropertyCustomizationHelpers.h
@@ -150,6 +150,12 @@ namespace PropertyCustomizationHelpers
 	 * @param AddRowDelegate The delegate that will be called on each child property.
 	 */
 	PROPERTYEDITOR_API void MakeInstancedPropertyCustomUI(TMap<FName, IDetailGroup*>& ExistingGroup, IDetailCategoryBuilder& BaseCategory, TSharedRef<IPropertyHandle>& BaseProperty, FOnInstancedPropertyIteration AddRowDelegate);
+	
+	/**
+	 * Parse and load the given metadata string into a list of allowed classes.
+	 * The metadata string is likely from something like AllowedClasses or DisallowedClasses.
+	 */
+	PROPERTYEDITOR_API TArray<const UClass*> GetClassesFromMetadataString(const FString& MetadataString);
 }
 
 
