Switched SPropertyEditorClass to use PropertyCustomizationHelpers::GetClassesFromMetadataString() helper as well.

Added consts and removed inline forward declarations. Set default values in FPropertyEditorClassFilter.

#rnx

diff --git a/Engine/Source/Editor/PropertyEditor/Private/UserInterface/PropertyEditor/SPropertyEditorClass.cpp b/Engine/Source/Editor/PropertyEditor/Private/UserInterface/PropertyEditor/SPropertyEditorClass.cpp
--- a/Engine/Source/Editor/PropertyEditor/Private/UserInterface/PropertyEditor/SPropertyEditorClass.cpp
+++ b/Engine/Source/Editor/PropertyEditor/Private/UserInterface/PropertyEditor/SPropertyEditorClass.cpp
@@ -19,13 +19,13 @@ class FPropertyEditorClassFilter : public IClassViewerFilter
 {
 public:
 	/** The meta class for the property that classes must be a child-of. */
-	const UClass* ClassPropertyMetaClass;
+	const UClass* ClassPropertyMetaClass = nullptr;
 
 	/** The interface that must be implemented. */
-	const UClass* InterfaceThatMustBeImplemented;
+	const UClass* InterfaceThatMustBeImplemented = nullptr;
 
 	/** Whether or not abstract classes are allowed. */
-	bool bAllowAbstract;
+	bool bAllowAbstract = false;
 
 	/** Classes that can be picked */
 	TArray<const UClass*> AllowedClassFilters;
@@ -88,7 +88,7 @@ void SPropertyEditorClass::GetDesiredWidth(float& OutMinDesiredWidth, float& Out
 	OutMaxDesiredWidth = 400.0f;
 }
 
-bool SPropertyEditorClass::Supports(const TSharedRef< class FPropertyEditor >& InPropertyEditor)
+bool SPropertyEditorClass::Supports(const TSharedRef< FPropertyEditor >& InPropertyEditor)
 {
 	const TSharedRef< FPropertyNode > PropertyNode = InPropertyEditor->GetPropertyNode();
 	const FProperty* Property = InPropertyEditor->GetProperty();
@@ -109,19 +109,19 @@ bool SPropertyEditorClass::CanEdit() const
 	return PropertyEditor.IsValid() ? !PropertyEditor->IsEditConst() : true;
 }
 
-void SPropertyEditorClass::Construct(const FArguments& InArgs, const TSharedPtr< class FPropertyEditor >& InPropertyEditor)
+void SPropertyEditorClass::Construct(const FArguments& InArgs, const TSharedPtr< FPropertyEditor >& InPropertyEditor)
 {
 	PropertyEditor = InPropertyEditor;
 	
 	if (PropertyEditor.IsValid())
 	{
 		const TSharedRef<FPropertyNode> PropertyNode = PropertyEditor->GetPropertyNode();
-		FProperty* const Property = PropertyNode->GetProperty();
-		if (FClassProperty* const ClassProp = CastField<FClassProperty>(Property))
+		const FProperty* const Property = PropertyNode->GetProperty();
+		if (const FClassProperty* const ClassProp = CastField<FClassProperty>(Property))
 		{
 			MetaClass = ClassProp->MetaClass;
 		}
-		else if (FSoftClassProperty* const SoftClassProperty = CastField<FSoftClassProperty>(Property))
+		else if (const FSoftClassProperty* const SoftClassProperty = CastField<FSoftClassProperty>(Property))
 		{
 			MetaClass = SoftClassProperty->MetaClass;
 		}
@@ -141,54 +141,16 @@ void SPropertyEditorClass::Construct(const FArguments& InArgs, const TSharedPtr<
 		bShowTree = Property->GetOwnerProperty()->HasMetaData(TEXT("ShowTreeView"));
 		bShowDisplayNames = Property->GetOwnerProperty()->HasMetaData(TEXT("ShowDisplayNames"));
 
-		auto FillClassFilters = [&](TArray<const UClass*> &ClassFilters, FName TagName)
-		{
-			const FString* ClassesFilterString = &Property->GetMetaData(TagName);
-
-			ClassFilters.Empty();
-
-			if (!ClassesFilterString->IsEmpty())
-			{
-				TArray<FString> ClassFilterNames;
-				ClassesFilterString->ParseIntoArrayWS(ClassFilterNames, TEXT(","), true);
-
-				for (const FString& ClassName : ClassFilterNames)
-				{
-					UClass* Class = FindOrLoadClass(ClassName);
-
-					if (Class)
-					{
-						// If the class is an interface, expand it to be all classes in memory that implement the class.
-						if (Class->HasAnyClassFlags(CLASS_Interface))
-						{
-							for (TObjectIterator<UClass> ClassIt; ClassIt; ++ClassIt)
-							{
-								UClass* const ClassWithInterface = (*ClassIt);
-								if (ClassWithInterface->ImplementsInterface(Class))
-								{
-									ClassFilters.Add(ClassWithInterface);
-								}
-							}
-						}
-						else
-						{
-							ClassFilters.Add(Class);
-						}
-					}
-				}
-			}
-		};
-
 		// Filter based on UPROPERTY meta data
-		FillClassFilters(AllowedClassFilters, "AllowedClasses");
-		FillClassFilters(DisallowedClassFilters, "DisallowedClasses");
+		AllowedClassFilters = PropertyCustomizationHelpers::GetClassesFromMetadataString(Property->GetOwnerProperty()->GetMetaData("AllowedClasses"));
+		DisallowedClassFilters = PropertyCustomizationHelpers::GetClassesFromMetadataString(Property->GetOwnerProperty()->GetMetaData("DisallowedClasses"));
 
 		// Filter based on restrictions
 		for (const TSharedRef<const FPropertyRestriction>& ClassRestriction : PropertyNode->GetRestrictions())
 		{
-			for (TArray<FString>::TConstIterator Iter= ClassRestriction.Get().GetHiddenValuesIterator(); Iter; ++Iter)
+			for (TArray<FString>::TConstIterator Iter = ClassRestriction.Get().GetHiddenValuesIterator(); Iter; ++Iter)
 			{
-				if (UClass* HiddenClass = FindOrLoadClass(*Iter))
+				if (const UClass* HiddenClass = FindOrLoadClass(*Iter))
 				{
 					DisallowedClassFilters.Add(HiddenClass);
 				}
@@ -196,7 +158,7 @@ void SPropertyEditorClass::Construct(const FArguments& InArgs, const TSharedPtr<
 
 			for (TArray<FString>::TConstIterator Iter = ClassRestriction.Get().GetDisabledValuesIterator(); Iter; ++Iter)
 			{
-				if (UClass* DisabledClass = FindOrLoadClass(*Iter))
+				if (const UClass* DisabledClass = FindOrLoadClass(*Iter))
 				{
 					DisallowedClassFilters.Add(DisabledClass);
 				}
@@ -248,10 +210,10 @@ void SPropertyEditorClass::Construct(const FArguments& InArgs, const TSharedPtr<
 static FString GetClassDisplayName(const UObject* Object, bool bShowDisplayNames)
 {
 	const UClass* Class = Cast<UClass>(Object);
-	if (Class != NULL)
+	if (Class != nullptr)
 	{
-		UBlueprint* BP = UBlueprint::GetBlueprintFromClass(Class);
-		if(BP != NULL)
+		const UBlueprint* BP = UBlueprint::GetBlueprintFromClass(Class);
+		if (BP != nullptr)
 		{
 			return BP->GetName();
 		}
@@ -273,10 +235,10 @@ FText SPropertyEditorClass::GetDisplayValueAsString() const
 		TGuardValue<bool> Guard( bIsReentrant, true );
 		if(PropertyEditor.IsValid())
 		{
-			UObject* ObjectValue = NULL;
+			UObject* ObjectValue = nullptr;
 			FPropertyAccess::Result Result = PropertyEditor->GetPropertyHandle()->GetValue(ObjectValue);
 
-			if(Result == FPropertyAccess::Success && ObjectValue != NULL)
+			if(Result == FPropertyAccess::Success && ObjectValue != nullptr)
 			{
 				return FText::FromString(GetClassDisplayName(ObjectValue, bShowDisplayNames));
 			}
@@ -312,14 +274,14 @@ void SPropertyEditorClass::CreateClassFilter(const TArray<TSharedRef<IClassViewe
 	ClassViewerOptions.ClassFilters.Append(InClassFilters);
 
 	TSharedRef<FPropertyEditorClassFilter> PropEdClassFilter = MakeShared<FPropertyEditorClassFilter>();
-	ClassViewerOptions.ClassFilters.Add(PropEdClassFilter);
-
 	PropEdClassFilter->ClassPropertyMetaClass = MetaClass;
 	PropEdClassFilter->InterfaceThatMustBeImplemented = RequiredInterface;
 	PropEdClassFilter->bAllowAbstract = bAllowAbstract;
 	PropEdClassFilter->AllowedClassFilters = AllowedClassFilters;
 	PropEdClassFilter->DisallowedClassFilters = DisallowedClassFilters;
 
+	ClassViewerOptions.ClassFilters.Add(PropEdClassFilter);
+
 	ClassFilter = FModuleManager::LoadModuleChecked<FClassViewerModule>("ClassViewer").CreateClassFilter(ClassViewerOptions);
 	ClassFilterFuncs = FModuleManager::LoadModuleChecked<FClassViewerModule>("ClassViewer").CreateFilterFuncs();
 }
@@ -364,7 +326,7 @@ void SPropertyEditorClass::SendToObjects(const FString& NewValue)
 	}
 	else if (!NewValue.IsEmpty() && NewValue != TEXT("None"))
 	{
-		UClass* NewClass = FindOrLoadClass(NewValue);
+		const UClass* NewClass = FindOrLoadClass(NewValue);
 		OnSetClass.Execute(NewClass);
 	}
 	else
@@ -407,15 +369,15 @@ void SPropertyEditorClass::OnDragEnter(const FGeometry& MyGeometry, const FDragD
 	TSharedPtr<FAssetDragDropOp> UnloadedClassOp = DragDropEvent.GetOperationAs<FAssetDragDropOp>();
 	if (UnloadedClassOp.IsValid())
 	{
-		UObject* Object = LoadDragDropObject(UnloadedClassOp);
+		const UObject* Object = LoadDragDropObject(UnloadedClassOp);
 
 		bool bOK = false;
 
-		if (UClass* Class = Cast<UClass>(Object))
+		if (const UClass* Class = Cast<UClass>(Object))
 		{
 			bOK = ClassFilter->IsClassAllowed(ClassViewerOptions, Class, ClassFilterFuncs.ToSharedRef());
 		}
-		else if (UBlueprint* Blueprint = Cast<UBlueprint>(Object))
+		else if (const UBlueprint* Blueprint = Cast<UBlueprint>(Object))
 		{
 			if (Blueprint->GeneratedClass)
 			{
@@ -449,7 +411,7 @@ FReply SPropertyEditorClass::OnDrop(const FGeometry& MyGeometry, const FDragDrop
 	if (ClassOperation.IsValid())
 	{
 		// We can only drop one item into the combo box, so drop the first one.
-		FString ClassPath = ClassOperation->ClassesToDrop[0]->GetPathName();
+		const FString ClassPath = ClassOperation->ClassesToDrop[0]->GetPathName();
 
 		// Set the property, it will be verified as valid.
 		SendToObjects(ClassPath);
@@ -486,7 +448,7 @@ FReply SPropertyEditorClass::OnDrop(const FGeometry& MyGeometry, const FDragDrop
 			GWarn->EndSlowTask();
 		}
 
-		if (UClass* Class = Cast<UClass>(Object))
+		if (const UClass* Class = Cast<UClass>(Object))
 		{
 			if (ClassFilter->IsClassAllowed(ClassViewerOptions, Class, ClassFilterFuncs.ToSharedRef()))
 			{
@@ -494,7 +456,7 @@ FReply SPropertyEditorClass::OnDrop(const FGeometry& MyGeometry, const FDragDrop
 				SendToObjects(Class->GetPathName());
 			}
 		}
-		else if (UBlueprint* Blueprint = Cast<UBlueprint>(Object))
+		else if (const UBlueprint* Blueprint = Cast<UBlueprint>(Object))
 		{
 			if (Blueprint->GeneratedClass)
 			{
