If a subobject in the details view was marked pending kill, the details view would keep rebuilding itself every frame until the object was GC'd. We now make sure that the object is valid before running the check.

Reported on UDN.

diff --git a/Engine/Source/Editor/PropertyEditor/Private/PropertyNode.cpp b/Engine/Source/Editor/PropertyEditor/Private/PropertyNode.cpp
--- a/Engine/Source/Editor/PropertyEditor/Private/PropertyNode.cpp
+++ b/Engine/Source/Editor/PropertyEditor/Private/PropertyNode.cpp
@@ -514,7 +514,7 @@ EPropertyDataValidationResult FPropertyNode::EnsureDataIsValid()
 			bool bObjectPropertyNull = true;
 
 			//Edit inline properties can change underneath the window
-			bool bIgnoreChangingChildren = !(HasNodeFlags(EPropertyNodeFlags::EditInlineNew) || HasNodeFlags(EPropertyNodeFlags::ShowInnerObjectProperties)) ;
+			bool bIgnoreChangingChildren = !(HasNodeFlags(EPropertyNodeFlags::EditInlineNew) || HasNodeFlags(EPropertyNodeFlags::ShowInnerObjectProperties));
 			//ignore this node if the consistency check should happen for the children
 			bool bIgnoreStaticArray = (Property->ArrayDim > 1) && (ArrayIndex == -1);
 
@@ -522,7 +522,7 @@ EPropertyDataValidationResult FPropertyNode::EnsureDataIsValid()
 			if (bIgnoreChangingChildren || bIgnoreStaticArray || HasNodeFlags(EPropertyNodeFlags::NoChildrenDueToCircularReference))
 			{
 				//this will bypass object property consistency checks
-				ObjectProperty = NULL;
+				ObjectProperty = nullptr;
 			}
 
 			FReadAddressList ReadAddresses;
@@ -598,32 +598,31 @@ EPropertyDataValidationResult FPropertyNode::EnsureDataIsValid()
 				if (ObjectProperty && !bIgnoreAllMismatch)
 				{
 					UObject* Obj = ObjectProperty->GetObjectPropertyValue(Addr);
-
-					if (!bShowInnerObjectPropertiesObjectChanged && HasNodeFlags(EPropertyNodeFlags::ShowInnerObjectProperties|EPropertyNodeFlags::EditInlineNew) && ChildNodes.Num() == 1)
+					if (IsValid(Obj))
 					{
-						bool bChildObjectFound = false;
-						// should never have more than one node (0 is ok if the object property is null)
-						check(ChildNodes.Num() == 1);
-						bool bNeedRebuild = false;
-						FObjectPropertyNode* ChildObjectNode = ChildNodes[0]->AsObjectNode();
-						for(int32 ObjectIndex = 0; ObjectIndex < ChildObjectNode->GetNumObjects(); ++ObjectIndex)
+						if (!bShowInnerObjectPropertiesObjectChanged && 
+							HasNodeFlags(EPropertyNodeFlags::ShowInnerObjectProperties|EPropertyNodeFlags::EditInlineNew) && 
+							ChildNodes.Num() == 1)
 						{
-							if(Obj == ChildObjectNode->GetUObject(ObjectIndex))
+							bool bChildObjectFound = false;
+							FObjectPropertyNode* ChildObjectNode = ChildNodes[0]->AsObjectNode();
+							for (int32 ObjectIndex = 0; ObjectIndex < ChildObjectNode->GetNumObjects(); ++ObjectIndex)
 							{
-								bChildObjectFound = true;
-								break;
+								if (Obj == ChildObjectNode->GetUObject(ObjectIndex))
+								{
+									bChildObjectFound = true;
+									break;
+								}
 							}
+							bShowInnerObjectPropertiesObjectChanged = !bChildObjectFound;
 						}
-						bShowInnerObjectPropertiesObjectChanged = !bChildObjectFound;
 					}
 
-					if (Obj != NULL)
+					if (Obj != nullptr)
 					{
 						bObjectPropertyNull = false;
 						break;
 					}
-
-					
 				}
 			}
 
@@ -646,7 +645,7 @@ EPropertyDataValidationResult FPropertyNode::EnsureDataIsValid()
 				return EPropertyDataValidationResult::ArraySizeChanged;
 			}
 
-			if(bShowInnerObjectPropertiesObjectChanged)
+			if (bShowInnerObjectPropertiesObjectChanged)
 			{
 				RebuildChildren();
 				return EPropertyDataValidationResult::EditInlineNewValueChanged;
@@ -655,7 +654,9 @@ EPropertyDataValidationResult FPropertyNode::EnsureDataIsValid()
 			const bool bHasChildren = (GetNumChildNodes() > 0);
 			// If the object property is not null and has no children, its children need to be rebuilt
 			// If the object property is null and this node has children, the node needs to be rebuilt
-			if (!HasNodeFlags(EPropertyNodeFlags::ShowInnerObjectProperties) && ObjectProperty && ((!bObjectPropertyNull && !bHasChildren) || (bObjectPropertyNull && bHasChildren)))
+			if (!HasNodeFlags(EPropertyNodeFlags::ShowInnerObjectProperties) && 
+				ObjectProperty != nullptr && 
+				((!bObjectPropertyNull && !bHasChildren) || (bObjectPropertyNull && bHasChildren)))
 			{
 				RebuildChildren();
 				return EPropertyDataValidationResult::PropertiesChanged;
