PR #8883: Groups added in details customizations can now be formatted as though they are subcategories (Contributed by [REDACTED_PERSON])

diff --git a/Engine/Source/Editor/PropertyEditor/Private/DetailGroup.cpp b/Engine/Source/Editor/PropertyEditor/Private/DetailGroup.cpp
--- a/Engine/Source/Editor/PropertyEditor/Private/DetailGroup.cpp
+++ b/Engine/Source/Editor/PropertyEditor/Private/DetailGroup.cpp
@@ -13,7 +13,8 @@ FDetailGroup::FDetailGroup( const FName InGroupName, TSharedRef<FDetailCategoryI
 	, LocalizedDisplayName( InLocalizedDisplayName )
 	, GroupName( InGroupName )
 	, bStartExpanded( bInStartExpanded )
-	, ResetEnabled(false)
+	, bResetEnabled( false )
+	, DisplayMode( EDetailGroupDisplayMode::Group )
 {
 
 }
@@ -111,6 +112,11 @@ TOptional<FResetToDefaultOverride> FDetailGroup::GetCustomResetToDefault() const
 	return TOptional<FResetToDefaultOverride>();
 }
 
+void FDetailGroup::SetDisplayMode(EDetailGroupDisplayMode Mode)
+{
+	DisplayMode = Mode;
+}
+
 TSharedPtr<FDetailPropertyRow> FDetailGroup::GetHeaderPropertyRow() const
 {
 	return HeaderCustomization.IsValid() ? HeaderCustomization->PropertyRow : nullptr;
@@ -237,7 +243,7 @@ TSharedRef<SWidget> FDetailGroup::MakeNameWidget()
 
 void FDetailGroup::OnResetClicked()
 {
-	if (ResetEnabled)
+	if (bResetEnabled)
 	{
 		TArray<TSharedPtr<IPropertyHandle>> PropertyHandles;
 
@@ -255,7 +261,7 @@ void FDetailGroup::OnResetClicked()
 
 bool FDetailGroup::IsResetVisible() const
 {
-	if (ResetEnabled)
+	if (bResetEnabled)
 	{
 		TArray<TSharedPtr<IPropertyHandle>> PropertyHandles;
 
@@ -276,7 +282,7 @@ bool FDetailGroup::IsResetVisible() const
 
 void FDetailGroup::EnableReset(bool InValue)
 {
-	ResetEnabled = InValue;
+	bResetEnabled = InValue;
 }
 
 bool FDetailGroup::GetAllChildrenPropertyHandles(TArray<TSharedPtr<IPropertyHandle>>& PropertyHandles) const
diff --git a/Engine/Source/Editor/PropertyEditor/Private/DetailGroup.h b/Engine/Source/Editor/PropertyEditor/Private/DetailGroup.h
--- a/Engine/Source/Editor/PropertyEditor/Private/DetailGroup.h
+++ b/Engine/Source/Editor/PropertyEditor/Private/DetailGroup.h
@@ -15,10 +15,10 @@
 
 class IDetailPropertyRow;
 
-class FDetailGroup : public IDetailGroup, public IDetailLayoutRow, public TSharedFromThis<FDetailGroup>
+class FDetailGroup : public IDetailGroup, public TSharedFromThis<FDetailGroup>
 {
 public:
-	FDetailGroup( const FName InGroupName, TSharedRef<FDetailCategoryImpl> InParentCategory, const FText& InLocalizedDisplayName, const bool bStartExpanded = false );
+	FDetailGroup( const FName InGroupName, TSharedRef<FDetailCategoryImpl> InParentCategory, const FText& InLocalizedDisplayName, const bool bInStartExpanded = false );
 
 	/** IDetailGroup interface */     
 	virtual FDetailWidgetRow& HeaderRow() override;
@@ -29,26 +29,33 @@ public:
 
 	virtual void ToggleExpansion( bool bExpand ) override;
 	virtual bool GetExpansionState() const override;
+	virtual void SetDisplayMode(EDetailGroupDisplayMode Mode) override;
 
 	/** IDetailLayoutRow interface */
-	virtual FName GetRowName() const override { return GetGroupName(); }
+	virtual FName GetRowName() const override { return GroupName; }
 	virtual TOptional<FResetToDefaultOverride> GetCustomResetToDefault() const override;
 
 	TSharedPtr<FDetailPropertyRow> GetHeaderPropertyRow() const;
 	TSharedPtr<FPropertyNode> GetHeaderPropertyNode() const;
 
 	/** @return The name of the group */
-	virtual FName GetGroupName() const override { return GroupName; }
+	virtual FName GetGroupName() const override { return GetRowName(); }
 
+	/** @return The localized display name of the group */
+	const FText& GetGroupDisplayName() const { return LocalizedDisplayName; }
+	
 	/** Whether or not the group has columns */
 	bool HasColumns() const;
 
 	/** @return true if this row should be ticked */
 	bool RequiresTick() const;
 
-	/** @return true is this row should start expanded */
+	/** @return true if this row should start expanded */
 	bool ShouldStartExpanded() const { return bStartExpanded; }
 
+	/** @return the display mode that this group should use */
+	EDetailGroupDisplayMode GetDisplayMode() const { return DisplayMode; }
+	
 	/** 
 	 * @return The visibility of this group
 	 */
@@ -122,9 +129,11 @@ private:
 	/** Name identifier of this group */
 	FName GroupName;
 	/** Whether the detail group should start expanded or not */
-	bool bStartExpanded;
+	bool bStartExpanded : 1;
 	/** Permit resetting all the properties in the group */
-	bool ResetEnabled;
+	bool bResetEnabled : 1;
+	/** Whether the detail group should appear like it's a subcategory or not */
+	EDetailGroupDisplayMode DisplayMode;
 	/**	Delegate called when user press the Group Reset ui */
 	FDetailGroupReset OnDetailGroupReset;
 };
diff --git a/Engine/Source/Editor/PropertyEditor/Private/DetailItemNode.cpp b/Engine/Source/Editor/PropertyEditor/Private/DetailItemNode.cpp
--- a/Engine/Source/Editor/PropertyEditor/Private/DetailItemNode.cpp
+++ b/Engine/Source/Editor/PropertyEditor/Private/DetailItemNode.cpp
@@ -246,6 +246,14 @@ TSharedRef< ITableRow > FDetailItemNode::GenerateWidgetForTableView( const TShar
 			.AddMetaData<FTagMetaData>(TagMeta)
 			.InnerCategory(true);
 	}
+	else if (Customization.HasGroup() && Customization.DetailGroup->GetDisplayMode() == EDetailGroupDisplayMode::Category)
+	{
+		return
+			SNew(SDetailCategoryTableRow, AsShared(), OwnerTable)
+			.DisplayName(Customization.DetailGroup->GetGroupDisplayName())
+			.AddMetaData<FTagMetaData>(TagMeta)
+			.InnerCategory(true);
+	}
 	else
 	{
 		return
diff --git a/Engine/Source/Editor/PropertyEditor/Public/IDetailGroup.h b/Engine/Source/Editor/PropertyEditor/Public/IDetailGroup.h
--- a/Engine/Source/Editor/PropertyEditor/Public/IDetailGroup.h
+++ b/Engine/Source/Editor/PropertyEditor/Public/IDetailGroup.h
@@ -2,22 +2,29 @@
 
 #pragma once
 
-#include "CoreMinimal.h"
-#include "PropertyHandle.h"
+#include "DetailCategoryBuilder.h"
+#include "Internationalization/Text.h"
+#include "Templates/SharedPointer.h"
 
 class FDetailWidgetRow;
 class IDetailPropertyRow;
+class IPropertyHandle;
 
 /**
- * A group in the details panel that can have children                                                              
+ * How the group should be formatted - either as a minor group, or identically to a major subcategory.
  */
-class IDetailGroup
+enum class EDetailGroupDisplayMode
+{
+	Group,
+	Category
+};
+
+/**
+ * A group in the details panel that can have children
+ */
+class IDetailGroup : public IDetailLayoutRow
 {
 public:
-	/**
-	* Delegate called when user press the Group Reset ui
-	*/
-	DECLARE_MULTICAST_DELEGATE(FDetailGroupReset);
 
 	virtual ~IDetailGroup(){}
 
@@ -38,14 +45,14 @@ public:
 	 * 
 	 * @return a new row for adding widgets
 	 */
-	virtual class FDetailWidgetRow& AddWidgetRow() = 0;
+	virtual FDetailWidgetRow& AddWidgetRow() = 0;
 
 	/**
 	 * Adds a new row for a property
 	 *
 	 * @return an interface for customizing the appearance of the property row
 	 */
-	virtual class IDetailPropertyRow& AddPropertyRow( TSharedRef<IPropertyHandle> PropertyHandle ) = 0;
+	virtual IDetailPropertyRow& AddPropertyRow(TSharedRef<IPropertyHandle> PropertyHandle) = 0;
 
 	/**
 	 * Adds a group to the category
@@ -68,12 +75,16 @@ public:
 	 */
 	virtual bool GetExpansionState() const = 0;
 
-
 	/**
 	* Permit resetting the properties in this group
 	*/
 	virtual void EnableReset(bool InValue) = 0;
 
+	/**
+	* Delegate called when user press the Group Reset ui
+	*/
+	DECLARE_MULTICAST_DELEGATE(FDetailGroupReset);
+
 	/**
 	* Return the delegate called when user press the Group Reset ui
 	*/
@@ -88,4 +99,7 @@ public:
 	* Return the property row associated with the specified property handle
 	*/
 	virtual TSharedPtr<IDetailPropertyRow> FindPropertyRow(TSharedRef<IPropertyHandle> PropertyHandle) const = 0;
+
+	/** Specify whether this group should be formatted to display like a subcategory */
+	virtual void SetDisplayMode(EDetailGroupDisplayMode Mode) = 0;
 };
