Updating invalid characters checked for in property menu actor picker

diff --git a/Engine/Source/Editor/PropertyEditor/Private/UserInterface/PropertyEditor/SPropertyMenuActorPicker.cpp b/Engine/Source/Editor/PropertyEditor/Private/UserInterface/PropertyEditor/SPropertyMenuActorPicker.cpp
--- a/Engine/Source/Editor/PropertyEditor/Private/UserInterface/PropertyEditor/SPropertyMenuActorPicker.cpp
+++ b/Engine/Source/Editor/PropertyEditor/Private/UserInterface/PropertyEditor/SPropertyMenuActorPicker.cpp
@@ -192,7 +192,7 @@ void SPropertyMenuActorPicker::OnActorSelected( AActor* InActor )
 {
 	FText OutErrorMsg;
 
-	if (IsValid(InActor) && !FName::IsValidXName(InActor->GetName(), "aeoiu",  &OutErrorMsg))
+	if (IsValid(InActor) && !FName::IsValidXName(InActor->GetName(), INVALID_NAME_CHARACTERS,  &OutErrorMsg))
 	{
 		FFormatNamedArguments Args;
 		Args.Add("ActorName", FText::FromString(FName::SanitizeWhitespace(InActor->GetName())));
