Move internal property name display from tooltip to contextual menu.

diff --git a/Engine/Documentation/Source/Shared/MenuEntries/PropertyView/PropertyView.INT.udn b/Engine/Documentation/Source/Shared/MenuEntries/PropertyView/PropertyView.INT.udn
new file mode 100644
--- /dev/null
+++ b/Engine/Documentation/Source/Shared/MenuEntries/PropertyView/PropertyView.INT.udn
@@ -0,0 +1,18 @@
+Availability:NoPublish
+Title:
+Crumbs:
+Description:
+
+[EXCERPT:CopyPropertyDisplayName_ToolTip]
+The display name of a property is the name you see in the Unreal Editor UI. For example, in Details panels and settings tabs, this is the name that appears in the left hand column.
+
+Copying this value to the clipboard can be useful any time you need to refer to a property by the name that people see in the editor.
+[/EXCERPT:CopyPropertyDisplayName_ToolTip]
+
+[EXCERPT:CopyPropertyInternalName_ToolTip]
+The internal name of a property is the name that appears in the property's declaration — for example, in C++ code or in a Blueprint asset.
+
+Copying this value to the clipboard can be useful if you want to search for a property in the engine's source code, or if you are writing an editor script and need to refer to the property by name.
+
+If this option is disabled, it may mean that this property or setting doesn't correspond exactly with a single property in C++ or Blueprint.
+[/EXCERPT:CopyPropertyInternalName_ToolTip]
diff --git a/Engine/Source/Editor/PropertyEditor/Private/SDetailSingleItemRow.cpp b/Engine/Source/Editor/PropertyEditor/Private/SDetailSingleItemRow.cpp
--- a/Engine/Source/Editor/PropertyEditor/Private/SDetailSingleItemRow.cpp
+++ b/Engine/Source/Editor/PropertyEditor/Private/SDetailSingleItemRow.cpp
@@ -798,16 +798,17 @@ FSlateColor SDetailSingleItemRow::GetInnerBackgroundColor() const
 bool SDetailSingleItemRow::OnContextMenuOpening(FMenuBuilder& MenuBuilder)
 {
 	FUIAction CopyDisplayNameAction = FExecuteAction::CreateSP(this, &SDetailSingleItemRow::OnCopyPropertyDisplayName);
+	FUIAction CopyInternalNameAction = FExecuteAction::CreateSP(this, &SDetailSingleItemRow::OnCopyPropertyInternalName);
+	CopyInternalNameAction.CanExecuteAction = FCanExecuteAction::CreateSP(this, &SDetailSingleItemRow::CanCopyPropertyInternalName);
 
-	bool bAddedMenuEntry = false;
-	if (CopyAction.IsBound() && PasteAction.IsBound())
+	// Show a separator line if there are items before this in the menu
+	if (MenuBuilder.GetMultiBox()->GetBlocks().Num() > 1)
 	{
-		// Hide separator line if it only contains the SearchWidget, making the next 2 elements the top of the list
-		if (MenuBuilder.GetMultiBox()->GetBlocks().Num() > 1)
-		{
-			MenuBuilder.AddMenuSeparator();
-		}
+		MenuBuilder.AddMenuSeparator();
+	}
 
+	if (CopyAction.IsBound() && PasteAction.IsBound())
+	{
 		bool bLongDisplayName = false;
 
 		FMenuEntryParams CopyContentParams;
@@ -825,14 +826,20 @@ bool SDetailSingleItemRow::OnContextMenuOpening(FMenuBuilder& MenuBuilder)
 		PasteContentParams.IconOverride = FSlateIcon(FCoreStyle::Get().GetStyleSetName(), "GenericCommands.Paste");
 		PasteContentParams.DirectActions = PasteAction;
 		MenuBuilder.AddMenuEntry(PasteContentParams);
-
-		MenuBuilder.AddMenuEntry(
-			NSLOCTEXT("PropertyView", "CopyPropertyDisplayName", "Copy Display Name"),
-			NSLOCTEXT("PropertyView", "CopyPropertyDisplayName_ToolTip", "Copy this property display name"),
-			FSlateIcon(FCoreStyle::Get().GetStyleSetName(), "GenericCommands.Copy"),
-			CopyDisplayNameAction);
 	}
 
+	MenuBuilder.AddMenuEntry(
+		NSLOCTEXT("PropertyView", "CopyPropertyDisplayName", "Copy Display Name"),
+		NSLOCTEXT("PropertyView", "CopyPropertyDisplayName_ToolTip", "Copy the display name of this property to the system clipboard."),
+		FSlateIcon(FCoreStyle::Get().GetStyleSetName(), "GenericCommands.Copy"),
+		CopyDisplayNameAction);
+
+	MenuBuilder.AddMenuEntry(
+		NSLOCTEXT("PropertyView", "CopyPropertyInternalName", "Copy Internal Name"),
+		NSLOCTEXT("PropertyView", "CopyPropertyInternalName_ToolTip", "Copy the internal name of this property to the system clipboard."),
+		FSlateIcon(FCoreStyle::Get().GetStyleSetName(), "GenericCommands.Copy"),
+		CopyInternalNameAction);
+
 	if (OwnerTreeNode.Pin()->GetDetailsView()->IsFavoritingEnabled())
 	{
 		FUIAction FavoriteAction;
@@ -888,8 +895,6 @@ bool SDetailSingleItemRow::OnContextMenuOpening(FMenuBuilder& MenuBuilder)
 					  FIsActionChecked::CreateSP(this, &SDetailSingleItemRow::IsDenyListChecked)),
 			NAME_None,
 			EUserInterfaceActionType::Check);
-			
-		bAddedMenuEntry = true;
 	}
 
 	if (WidgetRow.CustomMenuItems.Num() > 0)
@@ -946,6 +951,46 @@ void SDetailSingleItemRow::OnCopyPropertyDisplayName()
 	}
 }
 
+void SDetailSingleItemRow::OnCopyPropertyInternalName()
+{
+	if (OwnerTreeNode.IsValid())
+	{
+		TSharedPtr<FPropertyNode> PropertyNode = GetPropertyNode();
+		if (PropertyNode.IsValid())
+		{
+			if (const FProperty* Property = PropertyNode->GetProperty())
+			{
+				if (const UStruct* OwnerStruct = Property->GetOwnerStruct())
+				{
+					FPlatformApplicationMisc::ClipboardCopy(*OwnerStruct->GetAuthoredNameForField(Property));
+				}
+			}
+		}
+	}
+}
+
+bool SDetailSingleItemRow::CanCopyPropertyInternalName()
+{
+	if (OwnerTreeNode.IsValid())
+	{
+		TSharedPtr<FPropertyNode> PropertyNode = GetPropertyNode();
+		if (PropertyNode.IsValid())
+		{
+			if (const FProperty* Property = PropertyNode->GetProperty())
+			{
+				if (const UStruct* OwnerStruct = Property->GetOwnerStruct())
+				{
+					if (!OwnerStruct->GetAuthoredNameForField(Property).IsEmpty())
+					{
+						return true;
+					}
+				}
+			}
+		}
+	}
+	return false;
+}
+
 void SDetailSingleItemRow::OnPasteProperty()
 {
 	FString ClipboardContent;
diff --git a/Engine/Source/Editor/PropertyEditor/Private/SDetailSingleItemRow.h b/Engine/Source/Editor/PropertyEditor/Private/SDetailSingleItemRow.h
--- a/Engine/Source/Editor/PropertyEditor/Private/SDetailSingleItemRow.h
+++ b/Engine/Source/Editor/PropertyEditor/Private/SDetailSingleItemRow.h
@@ -72,6 +72,8 @@ protected:
 private:
 	void OnCopyProperty();
 	void OnCopyPropertyDisplayName();
+	void OnCopyPropertyInternalName();
+	bool CanCopyPropertyInternalName();
 	void OnPasteProperty();
 	bool CanPasteProperty() const;
 	FSlateColor GetOuterBackgroundColor() const;
diff --git a/Engine/Source/Runtime/CoreUObject/Private/UObject/Field.cpp b/Engine/Source/Runtime/CoreUObject/Private/UObject/Field.cpp
--- a/Engine/Source/Runtime/CoreUObject/Private/UObject/Field.cpp
+++ b/Engine/Source/Runtime/CoreUObject/Private/UObject/Field.cpp
@@ -763,16 +763,6 @@ FText FField::GetToolTipText(bool bShortTooltip) const
 		LocalizedToolTip = FText::FromString(NativeToolTip);
 	}
 
-	const FText DisplayName = FText::FromString(FName::NameToDisplayString(FFieldDisplayNameHelper::Get(*this), IsA<FBoolProperty>()));
-	if (LocalizedToolTip.IsEmpty())
-	{
-		LocalizedToolTip = DisplayName;
-	}
-	else
-	{
-		LocalizedToolTip = FText::Join(FText::FromString(TEXT(":" LINE_TERMINATOR_ANSI)), DisplayName, LocalizedToolTip);
-	}
-
 	return LocalizedToolTip;
 }
 
