Fix large stall in editor when PIE due to widget invalidation
Only clear cache widgets after bulk property changes.

diff --git a/Engine/Source/Editor/PropertyEditor/Private/PropertyEditorPermissionList.cpp b/Engine/Source/Editor/PropertyEditor/Private/PropertyEditorPermissionList.cpp
--- a/Engine/Source/Editor/PropertyEditor/Private/PropertyEditorPermissionList.cpp
+++ b/Engine/Source/Editor/PropertyEditor/Private/PropertyEditorPermissionList.cpp
@@ -114,6 +114,24 @@ void FPropertyEditorPermissionList::AddToAllowList(TSoftObjectPtr<UStruct> Struc
 	}
 }
 
+void FPropertyEditorPermissionList::AddToAllowList(TSoftObjectPtr<UStruct> Struct, const TArray<FName>& PropertyNames, const FName Owner)
+{
+	FPropertyEditorPermissionListEntry& Entry = RawPropertyEditorPermissionList.FindOrAdd(Struct);
+	bool bAddedItem = false;
+	for (const FName& PropertyName : PropertyNames)
+	{
+		if (Entry.PermissionList.AddAllowListItem(Owner, PropertyName))
+		{
+			bAddedItem = true;
+		}
+	}
+
+	if (bAddedItem)
+	{
+		ClearCacheAndBroadcast(Struct, Owner);
+	}
+}
+
 void FPropertyEditorPermissionList::RemoveFromAllowList(TSoftObjectPtr<UStruct> Struct, const FName PropertyName, const FName Owner)
 {
 	FPropertyEditorPermissionListEntry& Entry = RawPropertyEditorPermissionList.FindOrAdd(Struct);
diff --git a/Engine/Source/Editor/PropertyEditor/Public/PropertyEditorPermissionList.h b/Engine/Source/Editor/PropertyEditor/Public/PropertyEditorPermissionList.h
--- a/Engine/Source/Editor/PropertyEditor/Public/PropertyEditorPermissionList.h
+++ b/Engine/Source/Editor/PropertyEditor/Public/PropertyEditorPermissionList.h
@@ -64,6 +64,8 @@ public:
 
 	/** Add a specific property to a UStruct's AllowList */
 	void AddToAllowList(TSoftObjectPtr<UStruct> Struct, const FName PropertyName, const FName Owner = NAME_None);
+	/** Add a list of properties to a UStruct's AllowList */
+	void AddToAllowList(TSoftObjectPtr<UStruct> Struct, const TArray<FName>& PropertyNames, const FName Owner = NAME_None);
 	/** Remove a specific property from a UStruct's AllowList */
 	void RemoveFromAllowList(TSoftObjectPtr<UStruct> Struct, const FName PropertyName, const FName Owner = NAME_None);
 	/** Add a specific property to a UStruct's DenyList */
