Deleting or re-ordering indexes in data assets causes all indexes to collapse.

For this, as the children of a property node were built, the value of the Expanded flag was not passed from the old child to the rebuilt one. The fix is to, in rebuildChildren in PropertyNode.cpp, save off the property paths of the expanded children prior to the rebuild, and post rebuild of the children, if the property path of a new child is in the saved expansion paths, set the child's expanded flag to true, as well.

diff --git a/Engine/Source/Editor/PropertyEditor/Private/PropertyNode.cpp b/Engine/Source/Editor/PropertyEditor/Private/PropertyNode.cpp
--- a/Engine/Source/Editor/PropertyEditor/Private/PropertyNode.cpp
+++ b/Engine/Source/Editor/PropertyEditor/Private/PropertyNode.cpp
@@ -267,6 +267,16 @@ void FPropertyNode::RebuildChildren()
 	CachedReadAddresses.Reset();
 
 	bool bDestroySelf = false;
+	TSet<FString> OutExpandedChildPropertyPaths;
+
+	for (const TSharedPtr<FPropertyNode>& ChildNode : ChildNodes)
+	{
+		if (ChildNode->HasNodeFlags(EPropertyNodeFlags::Expanded))
+		{
+			OutExpandedChildPropertyPaths.Add(ChildNode->GetPropertyPath());
+		}
+	}
+	
 	DestroyTree(bDestroySelf);
 
 	if (MaxChildDepthAllowed != 0)
@@ -277,6 +287,13 @@ void FPropertyNode::RebuildChildren()
 		if (HasNodeFlags(EPropertyNodeFlags::CanBeExpanded) && (ChildNodes.Num() == 0))
 		{
 			InitChildNodes();
+			for (const TSharedPtr<FPropertyNode>& ChildNode : ChildNodes)
+			{
+				if ( OutExpandedChildPropertyPaths.Contains(ChildNode->GetPropertyPath()) )
+				{
+					ChildNode->SetNodeFlags(EPropertyNodeFlags::Expanded, true);
+				}
+			}
 		}
 	}
 
