Property editor fix to allow users to set an editinline-instanced subobject value when its outer object path contains a whitespace token.

diff --git a/Engine/Source/Editor/PropertyEditor/Private/UserInterface/PropertyEditor/SPropertyEditorEditInline.cpp b/Engine/Source/Editor/PropertyEditor/Private/UserInterface/PropertyEditor/SPropertyEditorEditInline.cpp
--- a/Engine/Source/Editor/PropertyEditor/Private/UserInterface/PropertyEditor/SPropertyEditorEditInline.cpp
+++ b/Engine/Source/Editor/PropertyEditor/Private/UserInterface/PropertyEditor/SPropertyEditorEditInline.cpp
@@ -311,7 +311,12 @@ void SPropertyEditorEditInline::OnClassPicked(UClass* InClass)
 
 					UObject* NewUObject = NewObject<UObject>(UseOuter, InClass, *NewObjectName, MaskedOuterFlags, NewObjectTemplate);
 
-					NewValue = NewUObject->GetPathName();
+					// Wrap the value in quotes before setting - in some cases for editinline-instanced values, the outer object path
+					// can potentially contain a space token (e.g. if the outer object was instanced as a Blueprint template object
+					// based on a user-facing variable name). While technically such characters should not be allowed, historically there
+					// has been no issue with most tokens in the INVALID_OBJECTNAME_CHARACTERS set being present in in-memory object names,
+					// other than some systems failing to resolve the object's path if the string representation of the value is not quoted.
+					NewValue = FString::Printf(TEXT("\"%s\""), *NewUObject->GetPathName().ReplaceQuotesWithEscapedQuotes());
 				}
 			}
 			else
