Exposed SetDisplayName() in IDetailCategoryBuilder because there's not really any reason not to allow customizations to override the display name.

diff --git a/Engine/Source/Editor/PropertyEditor/Private/DetailCategoryBuilderImpl.cpp b/Engine/Source/Editor/PropertyEditor/Private/DetailCategoryBuilderImpl.cpp
--- a/Engine/Source/Editor/PropertyEditor/Private/DetailCategoryBuilderImpl.cpp
+++ b/Engine/Source/Editor/PropertyEditor/Private/DetailCategoryBuilderImpl.cpp
@@ -693,6 +693,11 @@ bool FDetailCategoryImpl::ContainsOnlyAdvanced() const
 	return !bFavoriteCategory && SimpleChildNodes.Num() == 0 && AdvancedChildNodes.Num() > 0;
 }
 
+void FDetailCategoryImpl::SetDisplayName(const FText& InDisplayName)
+{
+	SetDisplayName(CategoryName, InDisplayName);
+}
+
 void FDetailCategoryImpl::SetDisplayName(FName InCategoryName, const FText& LocalizedNameOverride)
 {
 	if (!LocalizedNameOverride.IsEmpty())
diff --git a/Engine/Source/Editor/PropertyEditor/Private/DetailCategoryBuilderImpl.h b/Engine/Source/Editor/PropertyEditor/Private/DetailCategoryBuilderImpl.h
--- a/Engine/Source/Editor/PropertyEditor/Private/DetailCategoryBuilderImpl.h
+++ b/Engine/Source/Editor/PropertyEditor/Private/DetailCategoryBuilderImpl.h
@@ -180,6 +180,7 @@ public:
 	virtual IDetailGroup& AddGroup(FName GroupName, const FText& LocalizedDisplayName, bool bForAdvanced = false, bool bStartExpanded = false) override;
 	virtual void GetDefaultProperties(TArray<TSharedRef<IPropertyHandle> >& OutAllProperties, bool bSimpleProperties = true, bool bAdvancedProperties = true) override;
 	virtual const FText& GetDisplayName() const override { return DisplayName; }
+	virtual void SetDisplayName(const FText& InDisplayName) override;
 	virtual void SetCategoryVisibility(bool bIsVisible) override;
 	virtual void SetShowAdvanced(bool bShowAdvanced) override;
 	virtual int32 GetSortOrder() const override;
diff --git a/Engine/Source/Editor/PropertyEditor/Public/DetailCategoryBuilder.h b/Engine/Source/Editor/PropertyEditor/Public/DetailCategoryBuilder.h
--- a/Engine/Source/Editor/PropertyEditor/Public/DetailCategoryBuilder.h
+++ b/Engine/Source/Editor/PropertyEditor/Public/DetailCategoryBuilder.h
@@ -187,12 +187,16 @@ public:
 	 */
 	virtual bool IsParentLayoutValid() const = 0;
 
-
 	/**
 	 * @return The localized display name of the category
 	 */
 	virtual const FText& GetDisplayName() const = 0;
 
+	/**
+	 * Set a new display name for this category.
+	 */
+	virtual void SetDisplayName(const FText& DisplayName) = 0;
+
 	/**
 	 * Sets whether or not this category is hidden or shown
 	 * This is designed to be used for dynamic category visibility after construction of the category
