SPropertyEditorNumeric now supports MaxFractionalDigits metadata, which allows properties to specify that they don't need 6 digits of fractional precision which can be useful for properties with large ranges.

diff --git a/Engine/Source/Editor/PropertyEditor/Private/UserInterface/PropertyEditor/SPropertyEditorNumeric.h b/Engine/Source/Editor/PropertyEditor/Private/UserInterface/PropertyEditor/SPropertyEditorNumeric.h
--- a/Engine/Source/Editor/PropertyEditor/Private/UserInterface/PropertyEditor/SPropertyEditorNumeric.h
+++ b/Engine/Source/Editor/PropertyEditor/Private/UserInterface/PropertyEditor/SPropertyEditorNumeric.h
@@ -349,7 +349,7 @@ public:
 			// Set up the correct type interface if we want to display units on the property editor
 
 			// First off, check for ForceUnits= meta data. This meta tag tells us to interpret, and always display the value in these units. FUnitConversion::Settings().ShouldDisplayUnits does not apply to such properties
-			auto PropertyUnits = FUnitConversion::UnitFromString(*ForcedUnits);
+			TOptional<EUnit> PropertyUnits = FUnitConversion::UnitFromString(*ForcedUnits);
 			if (PropertyUnits.IsSet())
 			{
 				// Create the type interface and set up the default input units if they are compatible
@@ -379,8 +379,18 @@ public:
 
 				// Create the type interface and set up the default input units if they are compatible
 				TypeInterface = MakeShareable(new TNumericUnitTypeInterface<NumericType>(PropertyUnits.GetValue()));
-				auto Value = OnGetValue();
 
+				const FString& MaxFractionalDigitsString = GetMetaDataFromKey("MaxFractionalDigits");
+				if (!MaxFractionalDigitsString.IsEmpty())
+				{
+					int32 MaxFractionalDigits = 0;
+					if (LexTryParseString(MaxFractionalDigits, *MaxFractionalDigitsString) && MaxFractionalDigits > 0)
+					{
+						TypeInterface->SetMaxFractionalDigits(MaxFractionalDigits);
+					}
+				}
+
+				TOptional<NumericType> Value = OnGetValue();
 				if (Value.IsSet())
 				{
 					TypeInterface->SetupFixedDisplay(Value.GetValue());
