Property Editor: Generalize property editor meta data so that other systems can use it in the same way.

Related to this jira so that Sequence's key editors can have the same property meta data behavior

diff --git a/Engine/Source/Editor/PropertyEditor/Private/UserInterface/PropertyEditor/SPropertyEditorNumeric.h b/Engine/Source/Editor/PropertyEditor/Private/UserInterface/PropertyEditor/SPropertyEditorNumeric.h
--- a/Engine/Source/Editor/PropertyEditor/Private/UserInterface/PropertyEditor/SPropertyEditorNumeric.h
+++ b/Engine/Source/Editor/PropertyEditor/Private/UserInterface/PropertyEditor/SPropertyEditorNumeric.h
@@ -21,6 +21,7 @@
 #include "PropertyHandle.h"
 #include "Presentation/PropertyEditor/PropertyEditor.h"
 #include "PropertyEditorHelpers.h"
+#include "NumericPropertyParams.h"
 #include "UserInterface/PropertyEditor/PropertyEditorConstants.h"
 #include "Framework/Commands/UIAction.h"
 #include "Framework/MultiBox/MultiBoxBuilder.h"
@@ -257,90 +258,13 @@ public:
 		else
 		{
 			// Instance metadata overrides per-class metadata.
-			auto GetMetaDataFromKey = [PropertyHandle](const FName& Key) -> const FString&
+			const typename TNumericPropertyParams<NumericType>::FMetaDataGetter MetaDataGetter = TNumericPropertyParams<NumericType>::FMetaDataGetter::CreateLambda([&](const FName& Key)
 			{
 				const FString* InstanceValue = PropertyHandle->GetInstanceMetaData(Key);
 				return (InstanceValue != nullptr) ? *InstanceValue : PropertyHandle->GetMetaData(Key);
-			};
-
-			const FString& MetaUIMinString = GetMetaDataFromKey("UIMin");
-			const FString& MetaUIMaxString = GetMetaDataFromKey("UIMax");
-			const FString& SliderExponentString = GetMetaDataFromKey("SliderExponent");
-			const FString& LinearDeltaSensitivityString = GetMetaDataFromKey("LinearDeltaSensitivity");
-			const FString& DeltaString = GetMetaDataFromKey("Delta");
-			const FString& ClampMinString = GetMetaDataFromKey("ClampMin");
-			const FString& ClampMaxString = GetMetaDataFromKey("ClampMax");
-			const FString& ForcedUnits = GetMetaDataFromKey("ForceUnits");
-
-			// If no UIMin/Max was specified then use the clamp string
-			const FString& UIMinString = MetaUIMinString.Len() ? MetaUIMinString : ClampMinString;
-			const FString& UIMaxString = MetaUIMaxString.Len() ? MetaUIMaxString : ClampMaxString;
-
-			NumericType ClampMin = TNumericLimits<NumericType>::Lowest();
-			NumericType ClampMax = TNumericLimits<NumericType>::Max();
-
-			if (!ClampMinString.IsEmpty())
-			{
-				TTypeFromString<NumericType>::FromString(ClampMin, *ClampMinString);
-			}
-
-			if (!ClampMaxString.IsEmpty())
-			{
-				TTypeFromString<NumericType>::FromString(ClampMax, *ClampMaxString);
-			}
-
-			NumericType UIMin = TNumericLimits<NumericType>::Lowest();
-			NumericType UIMax = TNumericLimits<NumericType>::Max();
-			TTypeFromString<NumericType>::FromString(UIMin, *UIMinString);
-			TTypeFromString<NumericType>::FromString(UIMax, *UIMaxString);
-
-			NumericType SliderExponent = NumericType(1);
-			if (SliderExponentString.Len())
-			{
-				TTypeFromString<NumericType>::FromString(SliderExponent, *SliderExponentString);
-			}
-
-			NumericType Delta = NumericType(0);
-			if (DeltaString.Len())
-			{
-				TTypeFromString<NumericType>::FromString(Delta, *DeltaString);
-			}
-
-			int32 LinearDeltaSensitivity = 0;
-			if (LinearDeltaSensitivityString.Len())
-			{
-				TTypeFromString<int32>::FromString(LinearDeltaSensitivity, *LinearDeltaSensitivityString);
-			}
-			// LinearDeltaSensitivity only works in SSpinBox if delta is provided, so add it in if it wasn't.
-			Delta = (LinearDeltaSensitivity != 0 && Delta == NumericType(0)) ? NumericType(1) : Delta;
-
-			if (ClampMin >= ClampMax && (ClampMinString.Len() || ClampMaxString.Len()))
-			{
-				UE_LOG(LogPropertyNode, Warning, TEXT("Clamp Min (%s) >= Clamp Max (%s) for Ranged Numeric property %s"), *ClampMinString, *ClampMaxString, *Property->GetPathName());
-			}
-
-			const NumericType ActualUIMin = FMath::Max(UIMin, ClampMin);
-			const NumericType ActualUIMax = FMath::Min(UIMax, ClampMax);
-
-			TOptional<NumericType> MinValue = ClampMinString.Len() ? ClampMin : TOptional<NumericType>();
-			TOptional<NumericType> MaxValue = ClampMaxString.Len() ? ClampMax : TOptional<NumericType>();
-			TOptional<NumericType> SliderMinValue = (UIMinString.Len()) ? ActualUIMin : TOptional<NumericType>();
-			TOptional<NumericType> SliderMaxValue = (UIMaxString.Len()) ? ActualUIMax : TOptional<NumericType>();
-
-			if ((ActualUIMin >= ActualUIMax) && (SliderMinValue.IsSet() && SliderMaxValue.IsSet()))
-			{
-				UE_LOG(LogPropertyNode, Warning, TEXT("UI Min (%s) >= UI Max (%s) for Ranged Numeric property %s"), *UIMinString, *UIMaxString, *Property->GetPathName());
-			}
-
-			const FString& WheelStepString = GetMetaDataFromKey("WheelStep");
-			NumericType WheelStepValue = NumericType(0);
-
-			if (!WheelStepString.IsEmpty())
-			{
-				TTypeFromString<NumericType>::FromString(WheelStepValue, *WheelStepString);
-			}
+			});
 
-			TOptional<NumericType> WheelStep = WheelStepString.Len() ? WheelStepValue : TOptional<NumericType>();
+			TNumericPropertyParams<NumericType> NumericPropertyParams(Property, MetaDataGetter);
 
 			FObjectPropertyNode* ObjectPropertyNode = PropertyNode->FindObjectItemParent();
 			const bool bAllowSpin = (!ObjectPropertyNode || (1 == ObjectPropertyNode->GetNumObjects()))
@@ -349,6 +273,7 @@ public:
 			// Set up the correct type interface if we want to display units on the property editor
 
 			// First off, check for ForceUnits= meta data. This meta tag tells us to interpret, and always display the value in these units. FUnitConversion::Settings().ShouldDisplayUnits does not apply to such properties
+			const FString& ForcedUnits = MetaDataGetter.Execute("ForceUnits");
 			TOptional<EUnit> PropertyUnits = FUnitConversion::UnitFromString(*ForcedUnits);
 			if (PropertyUnits.IsSet())
 			{
@@ -368,7 +293,7 @@ public:
 					}
 					else
 					{
-						PropertyUnits = FUnitConversion::UnitFromString(*GetMetaDataFromKey("Units"));
+						PropertyUnits = FUnitConversion::UnitFromString(*MetaDataGetter.Execute("Units"));
 					}
 				}
 
@@ -380,7 +305,7 @@ public:
 				// Create the type interface and set up the default input units if they are compatible
 				TypeInterface = MakeShareable(new TNumericUnitTypeInterface<NumericType>(PropertyUnits.GetValue()));
 
-				const FString& MaxFractionalDigitsString = GetMetaDataFromKey("MaxFractionalDigits");
+				const FString& MaxFractionalDigitsString = MetaDataGetter.Execute("MaxFractionalDigits");
 				if (!MaxFractionalDigitsString.IsEmpty())
 				{
 					int32 MaxFractionalDigits = 0;
@@ -404,16 +329,16 @@ public:
 				.AllowSpin(bAllowSpin)
 				.Value(this, &SPropertyEditorNumeric<NumericType>::OnGetValue)
 				.Font(InArgs._Font)
-				.MinValue(MinValue)
-				.MaxValue(MaxValue)
-				.MinSliderValue(SliderMinValue)
-				.MaxSliderValue(SliderMaxValue)
-				.SliderExponent(SliderExponent)
-				.Delta(Delta)
+				.MinValue(NumericPropertyParams.MinValue)
+				.MaxValue(NumericPropertyParams.MaxValue)
+				.MinSliderValue(NumericPropertyParams.MinSliderValue)
+				.MaxSliderValue(NumericPropertyParams.MaxSliderValue)
+				.SliderExponent(NumericPropertyParams.SliderExponent)
+				.Delta(NumericPropertyParams.Delta)
 				// LinearDeltaSensitivity needs to be left unset if not provided, rather than being set to some default
-				.LinearDeltaSensitivity(LinearDeltaSensitivity != 0 ? LinearDeltaSensitivity : TAttribute<int32>())
+				.LinearDeltaSensitivity(NumericPropertyParams.GetLinearDeltaSensitivityAttribute())
 				.AllowWheel(bAllowSpin)
-				.WheelStep(WheelStep)
+				.WheelStep(NumericPropertyParams.WheelStep)
 				.UndeterminedString(LOCTEXT("MultipleValues", "Multiple Values"))
 				.OnValueChanged(this, &SPropertyEditorNumeric<NumericType>::OnValueChanged)
 				.OnValueCommitted(this, &SPropertyEditorNumeric<NumericType>::OnValueCommitted)
diff --git a/Engine/Source/Editor/PropertyEditor/Public/NumericPropertyParams.h b/Engine/Source/Editor/PropertyEditor/Public/NumericPropertyParams.h
new file mode 100644
--- /dev/null
+++ b/Engine/Source/Editor/PropertyEditor/Public/NumericPropertyParams.h
@@ -0,0 +1,128 @@
+// Copyright Epic Games, Inc. All Rights Reserved.
+
+#pragma once
+
+#include "CoreMinimal.h"
+
+template<typename NumericType>
+struct PROPERTYEDITOR_API TNumericPropertyParams
+{
+public:
+	DECLARE_DELEGATE_RetVal_OneParam(FString, FMetaDataGetter, FName);
+
+	TNumericPropertyParams(const FProperty* Property, FMetaDataGetter MetaDataGetter)	
+	{
+		FString UIMinString;
+		FString UIMaxString;
+		FString SliderExponentString;
+		FString LinearDeltaSensitivityString;
+		FString DeltaString;
+		FString ClampMinString;
+		FString ClampMaxString;
+		FString ForcedUnits;
+		FString WheelStepString;
+
+		if (!MetaDataGetter.IsBound() && Property != nullptr)
+		{
+			MetaDataGetter = FMetaDataGetter::CreateLambda([Property](const FName& Key) 
+			{ 
+				return Property->GetMetaData(Key); 
+			});
+		}
+
+		if (MetaDataGetter.IsBound())
+		{
+			UIMinString = MetaDataGetter.Execute("UIMin");
+			UIMaxString = MetaDataGetter.Execute("UIMax");
+			SliderExponentString = MetaDataGetter.Execute("SliderExponent");
+			LinearDeltaSensitivityString = MetaDataGetter.Execute("LinearDeltaSensitivity");
+			DeltaString = MetaDataGetter.Execute("Delta");
+			ClampMinString = MetaDataGetter.Execute("ClampMin");
+			ClampMaxString = MetaDataGetter.Execute("ClampMax");
+			ForcedUnits = MetaDataGetter.Execute("ForceUnits");
+			WheelStepString = MetaDataGetter.Execute("WheelStep");
+		}
+
+		// If no UIMin/Max was specified then use the clamp string
+		const FString& ActualUIMinString = UIMinString.Len() ? UIMinString : ClampMinString;
+		const FString& ActualUIMaxString = UIMaxString.Len() ? UIMaxString : ClampMaxString;
+
+		NumericType ClampMin = TNumericLimits<NumericType>::Lowest();
+		NumericType ClampMax = TNumericLimits<NumericType>::Max();
+
+		if (!ClampMinString.IsEmpty())
+		{
+			TTypeFromString<NumericType>::FromString(ClampMin, *ClampMinString);
+		}
+
+		if (!ClampMaxString.IsEmpty())
+		{
+			TTypeFromString<NumericType>::FromString(ClampMax, *ClampMaxString);
+		}
+
+		NumericType UIMin = TNumericLimits<NumericType>::Lowest();
+		NumericType UIMax = TNumericLimits<NumericType>::Max();
+		TTypeFromString<NumericType>::FromString(UIMin, *ActualUIMinString);
+		TTypeFromString<NumericType>::FromString(UIMax, *ActualUIMaxString);
+
+		if (ClampMin >= ClampMax && (ClampMinString.Len() || ClampMaxString.Len()))
+		{
+			UE_LOG(LogTemp, Warning, TEXT("Clamp Min (%s) >= Clamp Max (%s) for Ranged Numeric property %s"), *ClampMinString, *ClampMaxString, Property ? *Property->GetPathName() : TEXT(""));
+		}
+
+		const NumericType ActualUIMin = FMath::Max(UIMin, ClampMin);
+		const NumericType ActualUIMax = FMath::Min(UIMax, ClampMax);
+
+		MinValue = ClampMinString.Len() ? ClampMin : TOptional<NumericType>();
+		MaxValue = ClampMaxString.Len() ? ClampMax : TOptional<NumericType>();
+		MinSliderValue = (ActualUIMinString.Len()) ? ActualUIMin : TOptional<NumericType>();
+		MaxSliderValue = (ActualUIMaxString.Len()) ? ActualUIMax : TOptional<NumericType>();
+
+		if ((ActualUIMin >= ActualUIMax) && (MinSliderValue.IsSet() && MaxSliderValue.IsSet()))
+		{
+			UE_LOG(LogTemp, Warning, TEXT("UI Min (%s) >= UI Max (%s) for Ranged Numeric property %s"), *ActualUIMinString, *ActualUIMaxString, Property ? *Property->GetPathName() : TEXT(""));
+		}
+		
+		SliderExponent = NumericType(1);
+		if (SliderExponentString.Len())
+		{
+			TTypeFromString<NumericType>::FromString(SliderExponent, *SliderExponentString);
+		}
+
+		Delta = NumericType(0);
+		if (DeltaString.Len())
+		{
+			TTypeFromString<NumericType>::FromString(Delta, *DeltaString);
+		}
+
+		LinearDeltaSensitivity = 0;
+		if (LinearDeltaSensitivityString.Len())
+		{
+			TTypeFromString<int32>::FromString(LinearDeltaSensitivity, *LinearDeltaSensitivityString);
+		}
+		// LinearDeltaSensitivity only works in SSpinBox if delta is provided, so add it in if it wasn't.
+		Delta = (LinearDeltaSensitivity != 0 && Delta == NumericType(0)) ? NumericType(1) : Delta;
+
+		NumericType WheelStepValue = NumericType(0);
+
+		if (!WheelStepString.IsEmpty())
+		{
+			TTypeFromString<NumericType>::FromString(WheelStepValue, *WheelStepString);
+			WheelStep = WheelStepString.Len() ? WheelStepValue : TOptional<NumericType>();
+		}
+	}
+
+	TAttribute<int32> GetLinearDeltaSensitivityAttribute() const
+	{
+		return LinearDeltaSensitivity != 0 ? TAttribute<int32>(LinearDeltaSensitivity) : TAttribute<int32>();
+	}
+
+	TOptional<NumericType> MinValue;
+	TOptional<NumericType> MaxValue;
+	TOptional<NumericType> MinSliderValue;
+	TOptional<NumericType> MaxSliderValue;
+	NumericType SliderExponent;
+	NumericType Delta;
+	int32 LinearDeltaSensitivity;
+	TOptional<NumericType> WheelStep;
+};
