Fix crash resetting property when IPropertyTypeCustomization::CustomizeChildren used on instanced TArray<UObject*> called AddExternalObjects on array elements and  supplied a OverrideResetToDefault implementation which simply called ResetToDefault on the property handle that was passed in by CustomizeChildren

diff --git a/Engine/Source/Editor/PropertyEditor/Private/PropertyNode.cpp b/Engine/Source/Editor/PropertyEditor/Private/PropertyNode.cpp
--- a/Engine/Source/Editor/PropertyEditor/Private/PropertyNode.cpp
+++ b/Engine/Source/Editor/PropertyEditor/Private/PropertyNode.cpp
@@ -2581,7 +2581,7 @@ void FPropertyNode::NotifyPostChange( FPropertyChangedEvent& InPropertyChangedEv
 	TSharedRef<FEditPropertyChain> PropertyChain = BuildPropertyChain( InPropertyChangedEvent.Property );
 	
 	// remember the property that was the chain's original active property; this will correspond to the outermost property of struct/array that was modified
-	FProperty* const OriginalActiveProperty = PropertyChain->GetActiveMemberNode()->GetValue();
+	FProperty* const OriginalActiveProperty = PropertyChain->GetActiveMemberNode() ? PropertyChain->GetActiveMemberNode()->GetValue() : nullptr;
 
 	// invalidate the entire chain of objects in the hierarchy
 	{
@@ -2926,7 +2926,7 @@ TSharedRef<FEditPropertyChain> FPropertyNode::BuildPropertyChain( FProperty* InP
 
 	do
 	{
-		if (ItemNode == ComplexNode)
+		if (ItemNode == ComplexNode && PropertyChain->GetHead())
 		{
 			MemberProperty = PropertyChain->GetHead()->GetValue();
 		}
@@ -2946,7 +2946,7 @@ TSharedRef<FEditPropertyChain> FPropertyNode::BuildPropertyChain( FProperty* InP
 	while( ItemNode != NULL );
 
 	// If the modified property was a property of the object at the root of this property window, the member property will not have been set correctly
-	if (ItemNode == ComplexNode)
+	if (ItemNode == ComplexNode && PropertyChain->GetHead())
 	{
 		MemberProperty = PropertyChain->GetHead()->GetValue();
 	}
