SPropertyEditorNumeric now allows spinning when more than one object is selected.

diff --git a/Engine/Source/Editor/PropertyEditor/Private/UserInterface/PropertyEditor/SPropertyEditorNumeric.h b/Engine/Source/Editor/PropertyEditor/Private/UserInterface/PropertyEditor/SPropertyEditorNumeric.h
--- a/Engine/Source/Editor/PropertyEditor/Private/UserInterface/PropertyEditor/SPropertyEditorNumeric.h
+++ b/Engine/Source/Editor/PropertyEditor/Private/UserInterface/PropertyEditor/SPropertyEditorNumeric.h
@@ -267,8 +267,6 @@ public:
 			TNumericPropertyParams<NumericType> NumericPropertyParams(Property, MetaDataGetter);
 
 			FObjectPropertyNode* ObjectPropertyNode = PropertyNode->FindObjectItemParent();
-			const bool bAllowSpin = (!ObjectPropertyNode || (1 == ObjectPropertyNode->GetNumObjects()))
-				&& !PropertyHandle->GetBoolMetaData("NoSpinbox");
 
 			// Set up the correct type interface if we want to display units on the property editor
 
@@ -322,6 +320,8 @@ public:
 				}
 			}
 
+			const bool bAllowSpin = !PropertyHandle->GetBoolMetaData("NoSpinbox");
+
 			ChildSlot
 			[
 				SAssignNew(PrimaryWidget, SNumericEntryBox<NumericType>)
