Send a Toggle Editable event type instead of Value Set when the edit condition state of a property changes (i.e. the inline checkbox is toggled), so that we can listen to properties becoming (un)editable

diff --git a/Engine/Source/Editor/PropertyEditor/Private/Presentation/PropertyEditor/PropertyEditor.cpp b/Engine/Source/Editor/PropertyEditor/Private/Presentation/PropertyEditor/PropertyEditor.cpp
--- a/Engine/Source/Editor/PropertyEditor/Private/Presentation/PropertyEditor/PropertyEditor.cpp
+++ b/Engine/Source/Editor/PropertyEditor/Private/Presentation/PropertyEditor/PropertyEditor.cpp
@@ -520,7 +520,7 @@ void FPropertyEditor::ToggleEditConditionState()
 		FPropertyValueImpl::GenerateArrayIndexMapToObjectNode(ArrayIndicesPerObject[ObjectIndex], &PropertyNode.Get());
 	}
 
-	FPropertyChangedEvent ChangeEvent(PropertyNode->GetProperty(), EPropertyChangeType::ValueSet);
+	FPropertyChangedEvent ChangeEvent(PropertyNode->GetProperty(), EPropertyChangeType::ToggleEditable);
 	ChangeEvent.SetArrayIndexPerObject(ArrayIndicesPerObject);
 	PropertyNode->NotifyPostChange( ChangeEvent, PropertyUtilities->GetNotifyHook() );
 }
diff --git a/Engine/Source/Runtime/CoreUObject/Public/UObject/UnrealType.h b/Engine/Source/Runtime/CoreUObject/Public/UObject/UnrealType.h
--- a/Engine/Source/Runtime/CoreUObject/Public/UObject/UnrealType.h
+++ b/Engine/Source/Runtime/CoreUObject/Public/UObject/UnrealType.h
@@ -6363,6 +6363,8 @@ namespace EPropertyChangeType
 	inline const Type Redirected		= 1 << 7;
 	// Array Item Moved Within the Array
 	inline const Type ArrayMove			= 1 << 8;
+	// Edit Condition State has changed
+	inline const Type ToggleEditable	= 1 << 9;
 };
 
 /**
