Property panels will now properly use C++ setters and getters

diff --git a/Engine/Source/Editor/PropertyEditor/Private/PropertyHandleImpl.cpp b/Engine/Source/Editor/PropertyEditor/Private/PropertyHandleImpl.cpp
--- a/Engine/Source/Editor/PropertyEditor/Private/PropertyHandleImpl.cpp
+++ b/Engine/Source/Editor/PropertyEditor/Private/PropertyHandleImpl.cpp
@@ -124,6 +124,19 @@ FPropertyAccess::Result FPropertyValueImpl::GetValueData( void*& OutAddress ) co
 	return Res;
 }
 
+FPropertyAccess::Result FPropertyValueImpl::GetValueEditStack(FPropertyNodeEditStack& OutStack) const
+{
+	FPropertyAccess::Result Res = FPropertyAccess::Fail;
+
+	TSharedPtr<FPropertyNode> PropertyNodePin = PropertyNode.Pin();
+	if (PropertyNodePin.IsValid())
+	{
+		Res = PropertyNodePin->GetSingleEditStack(OutStack);
+	}
+
+	return Res;
+}
+
 FPropertyAccess::Result FPropertyValueImpl::ImportText( const FString& InValue, EPropertyValueSetFlags::Type Flags )
 {
 	TSharedPtr<FPropertyNode> PropertyNodePin = PropertyNode.Pin();
@@ -394,12 +407,19 @@ FPropertyAccess::Result FPropertyValueImpl::ImportText( const TArray<FObjectBase
 			}
 
 			// Set the new value.
-			EPropertyPortFlags PortFlags = (Flags & EPropertyValueSetFlags::InstanceObjects) != 0 ? PPF_InstanceSubobjects : PPF_None;
-			FPropertyTextUtilities::TextToPropertyHelper(*NewValue, InPropertyNode, NodeProperty, Cur, PortFlags);
+			{
+				EPropertyPortFlags PortFlags = (Flags & EPropertyValueSetFlags::InstanceObjects) != 0 ? PPF_InstanceSubobjects : PPF_None;
+				FPropertyNodeEditStack PropStack(InPropertyNode, CurObject);
+				FPropertyTextUtilities::TextToPropertyHelper(*NewValue, InPropertyNode, NodeProperty, PropStack.GetDirectPropertyAddress(), CurObject, PortFlags);
+				PropStack.CommitChanges();
+			}
 
 			// Cache the value of the property after having modified it.
 			FString ValueAfterImport;
-			FPropertyTextUtilities::PropertyToTextHelper(ValueAfterImport, InPropertyNode, NodeProperty, Cur, PPF_None);
+			{
+				FPropertyNodeEditStack PropStack(InPropertyNode, CurObject);
+				FPropertyTextUtilities::PropertyToTextHelper(ValueAfterImport, InPropertyNode, NodeProperty, PropStack.GetDirectPropertyAddress(), CurObject, PPF_None);
+			}
 
 			if (CurObject)
 			{
@@ -1025,58 +1045,67 @@ void FPropertyValueImpl::AddChild()
 
 						if (Array)
 						{
-							FScriptArrayHelper	ArrayHelper(Array, Addr);
-							Index = ArrayHelper.AddValue();
-
-							// check whether the inner type is flagged as a non-nullable. if so, create it.
-							FObjectProperty* InnerObjectProperty = CastField<FObjectProperty>(Array->Inner);
-							if (InnerObjectProperty && InnerObjectProperty->HasAnyPropertyFlags(CPF_NonNullable))
-							{
-								UObject* NewItem = NewObject<UObject>(Obj, InnerObjectProperty->PropertyClass);
-								InnerObjectProperty->SetObjectPropertyValue(ArrayHelper.GetRawPtr(Index), NewItem);
-							}
+							Array->PerformOperationWithSetter(Obj, Addr, [Obj, Array, &Index](void* DirectAddress)
+							{								
+								FScriptArrayHelper	ArrayHelper(Array, DirectAddress);
+								Index = ArrayHelper.AddValue();
+
+								// check whether the inner type is flagged as a non-nullable. if so, create it.
+								FObjectProperty* InnerObjectProperty = CastField<FObjectProperty>(Array->Inner);
+								if (InnerObjectProperty && InnerObjectProperty->HasAnyPropertyFlags(CPF_NonNullable))
+								{
+									UObject* NewItem = NewObject<UObject>(Obj, InnerObjectProperty->PropertyClass);
+									InnerObjectProperty->SetObjectPropertyValue(ArrayHelper.GetRawPtr(Index), NewItem);
+								}
+							});
 						}
 						else if (Set)
 						{
-							FScriptSetHelper	SetHelper(Set, Addr);
-							Index = SetHelper.AddDefaultValue_Invalid_NeedsRehash();
-
-							// check whether the element type is flagged as a non-nullable. if so, create it.
-							FObjectProperty* ElementObjectProperty = CastField<FObjectProperty>(Set->ElementProp);
-							if (ElementObjectProperty && ElementObjectProperty->HasAnyPropertyFlags(CPF_NonNullable))
+							Set->PerformOperationWithSetter(Obj, Addr, [Obj, Set, &Index](void* DirectAddress)
 							{
-								UObject* NewItem = NewObject<UObject>(Obj, ElementObjectProperty->PropertyClass);
-								ElementObjectProperty->SetObjectPropertyValue(SetHelper.GetElementPtr(Index), NewItem);
-							}
+								FScriptSetHelper	SetHelper(Set, DirectAddress);
+								Index = SetHelper.AddDefaultValue_Invalid_NeedsRehash();
+
+								// check whether the element type is flagged as a non-nullable. if so, create it.
+								FObjectProperty* ElementObjectProperty = CastField<FObjectProperty>(Set->ElementProp);
+								if (ElementObjectProperty && ElementObjectProperty->HasAnyPropertyFlags(CPF_NonNullable))
+								{
+									UObject* NewItem = NewObject<UObject>(Obj, ElementObjectProperty->PropertyClass);
+									ElementObjectProperty->SetObjectPropertyValue(SetHelper.GetElementPtr(Index), NewItem);
+								}
 
-							SetHelper.Rehash();
+								SetHelper.Rehash();
+							});
 						}
 						else if (Map)
 						{
-							FScriptMapHelper	MapHelper(Map, Addr);
-							Index = MapHelper.AddDefaultValue_Invalid_NeedsRehash();
-
-							// check whether the key or value type is flagged as a non-nullable. if so, create it.
+							Map->PerformOperationWithSetter(Obj, Addr, [Obj, Map, &Index, &bAddedMapEntry](void* DirectAddress)
 							{
-								FObjectProperty* KeyObjectProperty = CastField<FObjectProperty>(Map->KeyProp);
-								if (KeyObjectProperty && KeyObjectProperty->HasAnyPropertyFlags(CPF_NonNullable))
+								FScriptMapHelper	MapHelper(Map, DirectAddress);
+								Index = MapHelper.AddDefaultValue_Invalid_NeedsRehash();
+
+								// check whether the key or value type is flagged as a non-nullable. if so, create it.
 								{
-									UObject* NewItem = NewObject<UObject>(Obj, KeyObjectProperty->PropertyClass);
-									KeyObjectProperty->SetObjectPropertyValue(MapHelper.GetKeyPtr(Index), NewItem);
+									FObjectProperty* KeyObjectProperty = CastField<FObjectProperty>(Map->KeyProp);
+									if (KeyObjectProperty && KeyObjectProperty->HasAnyPropertyFlags(CPF_NonNullable))
+									{
+										UObject* NewItem = NewObject<UObject>(Obj, KeyObjectProperty->PropertyClass);
+										KeyObjectProperty->SetObjectPropertyValue(MapHelper.GetKeyPtr(Index), NewItem);
+									}
 								}
-							}
 
-							{
-								FObjectProperty* ValueObjectProperty = CastField<FObjectProperty>(Map->ValueProp);
-								if (ValueObjectProperty && ValueObjectProperty->HasAnyPropertyFlags(CPF_NonNullable))
 								{
-									UObject* NewItem = NewObject<UObject>(Obj, ValueObjectProperty->PropertyClass);
-									ValueObjectProperty->SetObjectPropertyValue(MapHelper.GetValuePtr(Index), NewItem);
+									FObjectProperty* ValueObjectProperty = CastField<FObjectProperty>(Map->ValueProp);
+									if (ValueObjectProperty && ValueObjectProperty->HasAnyPropertyFlags(CPF_NonNullable))
+									{
+										UObject* NewItem = NewObject<UObject>(Obj, ValueObjectProperty->PropertyClass);
+										ValueObjectProperty->SetObjectPropertyValue(MapHelper.GetValuePtr(Index), NewItem);
+									}
 								}
-							}
 
-							MapHelper.Rehash();
-							bAddedMapEntry = true;
+								MapHelper.Rehash();
+								bAddedMapEntry = true;
+							});
 						}
 
 						ArrayIndicesPerObject[i].Add(NodeProperty->GetName(), Index);
@@ -3367,20 +3396,20 @@ bool FPropertyHandleInt::Supports( TSharedRef<FPropertyNode> PropertyNode )
 }
 
 template <typename PropertyClass, typename ValueType>
-ValueType GetIntegerValue(void* PropValue, FPropertyValueImpl& Implementation )
+ValueType GetIntegerValue(FPropertyNodeEditStack& PropStack, FPropertyValueImpl& Implementation)
 {
-	check( Implementation.IsPropertyTypeOf( PropertyClass::StaticClass() ) );
-	return Implementation.GetPropertyValue<PropertyClass>(PropValue);
+	check(Implementation.IsPropertyTypeOf(PropertyClass::StaticClass()));
+	return Implementation.GetPropertyValue<PropertyClass>(PropStack.GetDirectPropertyAddress());
 }
 
 FPropertyAccess::Result FPropertyHandleInt::GetValue(int8& OutValue) const
 {
-	void* PropValue = nullptr;
-	FPropertyAccess::Result Res = Implementation->GetValueData(PropValue);
+	FPropertyNodeEditStack PropStack;
+	FPropertyAccess::Result Res = Implementation->GetValueEditStack(PropStack);
 
 	if (Res == FPropertyAccess::Success)
 	{
-		OutValue = GetIntegerValue<FInt8Property, int8>(PropValue, *Implementation);
+		OutValue = GetIntegerValue<FInt8Property, int8>(PropStack, *Implementation);
 	}
 
 	return Res;
@@ -3388,12 +3417,12 @@ FPropertyAccess::Result FPropertyHandleInt::GetValue(int8& OutValue) const
 
 FPropertyAccess::Result FPropertyHandleInt::GetValue(int16& OutValue) const
 {
-	void* PropValue = nullptr;
-	FPropertyAccess::Result Res = Implementation->GetValueData(PropValue);
+	FPropertyNodeEditStack PropStack;
+	FPropertyAccess::Result Res = Implementation->GetValueEditStack(PropStack);
 
 	if (Res == FPropertyAccess::Success)
 	{
-		OutValue = GetIntegerValue<FInt16Property, int16>(PropValue, *Implementation);
+		OutValue = GetIntegerValue<FInt16Property, int16>(PropStack, *Implementation);
 	}
 
 	return Res;
@@ -3401,24 +3430,24 @@ FPropertyAccess::Result FPropertyHandleInt::GetValue(int16& OutValue) const
 
 FPropertyAccess::Result FPropertyHandleInt::GetValue(int32& OutValue) const
 {
-	void* PropValue = nullptr;
-	FPropertyAccess::Result Res = Implementation->GetValueData(PropValue);
+	FPropertyNodeEditStack PropStack;
+	FPropertyAccess::Result Res = Implementation->GetValueEditStack(PropStack);
 
 	if (Res == FPropertyAccess::Success)
 	{
-		OutValue = GetIntegerValue<FIntProperty, int32>(PropValue, *Implementation);
+		OutValue = GetIntegerValue<FIntProperty, int32>(PropStack, *Implementation);
 	}
 
 	return Res;
 }
 FPropertyAccess::Result FPropertyHandleInt::GetValue( int64& OutValue ) const
 {
-	void* PropValue = nullptr;
-	FPropertyAccess::Result Res = Implementation->GetValueData( PropValue );
+	FPropertyNodeEditStack PropStack;
+	FPropertyAccess::Result Res = Implementation->GetValueEditStack(PropStack);
 
 	if( Res == FPropertyAccess::Success )
 	{
-		OutValue = GetIntegerValue<FInt64Property,int64>(PropValue, *Implementation);
+		OutValue = GetIntegerValue<FInt64Property,int64>(PropStack, *Implementation);
 	}
 
 	return Res;
@@ -3426,12 +3455,12 @@ FPropertyAccess::Result FPropertyHandleInt::GetValue( int64& OutValue ) const
 
 FPropertyAccess::Result FPropertyHandleInt::GetValue(uint16& OutValue) const
 {
-	void* PropValue = nullptr;
-	FPropertyAccess::Result Res = Implementation->GetValueData(PropValue);
+	FPropertyNodeEditStack PropStack;
+	FPropertyAccess::Result Res = Implementation->GetValueEditStack(PropStack);
 
 	if (Res == FPropertyAccess::Success)
 	{
-		OutValue = GetIntegerValue<FUInt16Property, uint16>(PropValue, *Implementation);
+		OutValue = GetIntegerValue<FUInt16Property, uint16>(PropStack, *Implementation);
 	}
 
 	return Res;
@@ -3439,12 +3468,12 @@ FPropertyAccess::Result FPropertyHandleInt::GetValue(uint16& OutValue) const
 
 FPropertyAccess::Result FPropertyHandleInt::GetValue(uint32& OutValue) const
 {
-	void* PropValue = nullptr;
-	FPropertyAccess::Result Res = Implementation->GetValueData(PropValue);
+	FPropertyNodeEditStack PropStack;
+	FPropertyAccess::Result Res = Implementation->GetValueEditStack(PropStack);
 
 	if (Res == FPropertyAccess::Success)
 	{
-		OutValue = GetIntegerValue<FUInt32Property, uint32>(PropValue, *Implementation);
+		OutValue = GetIntegerValue<FUInt32Property, uint32>(PropStack, *Implementation);
 	}
 
 	return Res;
@@ -3452,12 +3481,12 @@ FPropertyAccess::Result FPropertyHandleInt::GetValue(uint32& OutValue) const
 
 FPropertyAccess::Result FPropertyHandleInt::GetValue(uint64& OutValue) const
 {
-	void* PropValue = nullptr;
-	FPropertyAccess::Result Res = Implementation->GetValueData(PropValue);
+	FPropertyNodeEditStack PropStack;
+	FPropertyAccess::Result Res = Implementation->GetValueEditStack(PropStack);
 
 	if (Res == FPropertyAccess::Success)
 	{
-		OutValue = GetIntegerValue<FUInt64Property, uint64>(PropValue, *Implementation);
+		OutValue = GetIntegerValue<FUInt64Property, uint64>(PropStack, *Implementation);
 	}
 
 	return Res;
@@ -3564,8 +3593,9 @@ bool FPropertyHandleFloat::Supports( TSharedRef<FPropertyNode> PropertyNode )
 
 FPropertyAccess::Result FPropertyHandleFloat::GetValue( float& OutValue ) const
 {
-	void* PropValue = nullptr;
-	FPropertyAccess::Result Res = Implementation->GetValueData( PropValue );
+	FPropertyNodeEditStack PropStack;
+	FPropertyAccess::Result Res = Implementation->GetValueEditStack(PropStack);
+	void* PropValue = PropStack.GetDirectPropertyAddress();
 
 	if( Res == FPropertyAccess::Success )
 	{
@@ -3602,8 +3632,9 @@ bool FPropertyHandleDouble::Supports( TSharedRef<FPropertyNode> PropertyNode )
 
 FPropertyAccess::Result FPropertyHandleDouble::GetValue( double& OutValue ) const
 {
-	void* PropValue = nullptr;
-	FPropertyAccess::Result Res = Implementation->GetValueData( PropValue );
+	FPropertyNodeEditStack PropStack;
+	FPropertyAccess::Result Res = Implementation->GetValueEditStack(PropStack);
+	void* PropValue = PropStack.GetDirectPropertyAddress();
 
 	if (Res == FPropertyAccess::Success)
 	{
@@ -3642,8 +3673,9 @@ bool FPropertyHandleBool::Supports( TSharedRef<FPropertyNode> PropertyNode )
 
 FPropertyAccess::Result FPropertyHandleBool::GetValue( bool& OutValue ) const
 {
-	void* PropValue = nullptr;
-	FPropertyAccess::Result Res = Implementation->GetValueData( PropValue );
+	FPropertyNodeEditStack PropStack;
+	FPropertyAccess::Result Res = Implementation->GetValueEditStack(PropStack);
+	void* PropValue = PropStack.GetDirectPropertyAddress();
 
 	if( Res == FPropertyAccess::Success )
 	{
@@ -3687,8 +3719,9 @@ bool FPropertyHandleByte::Supports( TSharedRef<FPropertyNode> PropertyNode )
 
 FPropertyAccess::Result FPropertyHandleByte::GetValue( uint8& OutValue ) const
 {
-	void* PropValue = nullptr;
-	FPropertyAccess::Result Res = Implementation->GetValueData( PropValue );
+	FPropertyNodeEditStack PropStack;
+	FPropertyAccess::Result Res = Implementation->GetValueEditStack(PropStack);
+	void* PropValue = PropStack.GetDirectPropertyAddress();
 
 	if( Res == FPropertyAccess::Success )
 	{
@@ -3776,8 +3809,9 @@ FPropertyAccess::Result FPropertyHandleString::SetValue( const TCHAR* NewValue,
 
 FPropertyAccess::Result FPropertyHandleString::GetValue( FName& OutValue ) const
 {
-	void* PropValue = nullptr;
-	FPropertyAccess::Result Res = Implementation->GetValueData( PropValue );
+	FPropertyNodeEditStack PropStack;
+	FPropertyAccess::Result Res = Implementation->GetValueEditStack(PropStack);
+	void* PropValue = PropStack.GetDirectPropertyAddress();
 
 	if( Res == FPropertyAccess::Success )
 	{
@@ -3814,8 +3848,9 @@ FPropertyAccess::Result FPropertyHandleObject::GetValue( UObject*& OutValue ) co
 
 FPropertyAccess::Result FPropertyHandleObject::GetValue( const UObject*& OutValue ) const
 {
-	void* PropValue = nullptr;
-	FPropertyAccess::Result Res = Implementation->GetValueData( PropValue );
+	FPropertyNodeEditStack PropStack;
+	FPropertyAccess::Result Res = Implementation->GetValueEditStack(PropStack);
+	void* PropValue = PropStack.GetDirectPropertyAddress();
 
 	if( Res == FPropertyAccess::Success )
 	{
@@ -4211,8 +4246,9 @@ bool FPropertyHandleMixed::Supports(TSharedRef<FPropertyNode> PropertyNode)
 
 FPropertyAccess::Result FPropertyHandleMixed::GetValue(double& OutValue) const
 {
-	void* PropValue = nullptr;
-	FPropertyAccess::Result Res = Implementation->GetValueData(PropValue);
+	FPropertyNodeEditStack PropStack;
+	FPropertyAccess::Result Res = Implementation->GetValueEditStack(PropStack);
+	void* PropValue = PropStack.GetDirectPropertyAddress();
 
 	if (Res == FPropertyAccess::Success)
 	{
diff --git a/Engine/Source/Editor/PropertyEditor/Private/PropertyHandleImpl.h b/Engine/Source/Editor/PropertyEditor/Private/PropertyHandleImpl.h
--- a/Engine/Source/Editor/PropertyEditor/Private/PropertyHandleImpl.h
+++ b/Engine/Source/Editor/PropertyEditor/Private/PropertyHandleImpl.h
@@ -81,6 +81,14 @@ public:
 	 */
 	FPropertyAccess::Result GetValueData( void*& OutAddress ) const;
 
+	/**
+	 * Gets edit stack for this property
+	 *
+	 * @param OutStack	Edit stack
+	 * @return The result of the query
+	 */
+	FPropertyAccess::Result GetValueEditStack(FPropertyNodeEditStack& OutStack) const;
+
 	/**
 	 * Given an address and a property type, get the actual value out
 	 *
diff --git a/Engine/Source/Editor/PropertyEditor/Private/PropertyNode.cpp b/Engine/Source/Editor/PropertyEditor/Private/PropertyNode.cpp
--- a/Engine/Source/Editor/PropertyEditor/Private/PropertyNode.cpp
+++ b/Engine/Source/Editor/PropertyEditor/Private/PropertyNode.cpp
@@ -723,18 +723,132 @@ EPropertyDataValidationResult FPropertyNode::EnsureDataIsValid()
 	return FinalResult;
 }
 
+FPropertyNodeEditStack::FPropertyNodeEditStack(const FPropertyNode* InNode, const UObject* InObj)
+{
+	InitializeInternal(InNode, InObj);
+}
+
+void FPropertyNodeEditStack::Initialize(const FPropertyNode* InNode, const UObject* InObj)
+{
+	Cleanup();
+	InitializeInternal(InNode, InObj);
+}
+
+void FPropertyNodeEditStack::InitializeInternal(const FPropertyNode* InNode, const UObject* InObj)
+{
+	const FPropertyNode* Parent = InNode->GetParentNode();
+	const FProperty* Property = InNode->GetProperty();
+	if (Parent && Parent->GetProperty())
+	{
+		// Recursively initialize the stack
+		InitializeInternal(Parent, InObj);
+
+		// Get the direct memory pointer for the current property
+		const FProperty* ParentProperty = Parent->GetProperty();
+		if (Property == ParentProperty) // Static array items
+		{
+			// Static array property node creates subnodes that point to individual array items
+			MemoryStack.Add(FMemoryFrame(Property, MemoryStack.Last().Memory + InNode->GetArrayIndex() * Property->ElementSize));
+		}
+		else if (const FStructProperty* StructProp = CastField<FStructProperty>(ParentProperty)) // structs
+		{
+			if (Property->HasSetterOrGetter())
+			{
+				// If a property has a setter or getter we allocate temp memory to hold its value so that we can
+				// change the value using direct memory pointer access. After we're done editing we will copy the memory back to the property in CommitChanges
+				FMemoryFrame PropertyFrame(Property, (uint8*)Property->AllocateAndInitializeValue());
+				int32 StackIndex = MemoryStack.Add(PropertyFrame);
+				Property->GetValue_InContainer(MemoryStack[StackIndex - 1].Memory, PropertyFrame.Memory);
+			}
+			else
+			{
+				MemoryStack.Add(FMemoryFrame(Property, Property->ContainerPtrToValuePtr<uint8>(MemoryStack.Last().Memory)));
+			}
+		}
+		else if (Property->GetOwner<FProperty>() == ParentProperty) // TArrays, TMaps and TSets
+		{
+			MemoryStack.Add(FMemoryFrame(Property, (uint8*)ParentProperty->GetValueAddressAtIndex_Direct(Property, (void*)MemoryStack.Last().Memory, InNode->GetArrayIndex())));
+		}
+		else
+		{
+			checkf(false, TEXT("Unsupported property chain: Current: %s, Parent: %s"), *Property->GetFullName(), *ParentProperty->GetFullName());
+		}
+	}
+	else
+	{
+		// Determine the root container object (UObject instance or sparse class data) for this property stack
+		uint8* Container = nullptr;
+		if (InNode->HasNodeFlags(EPropertyNodeFlags::IsSparseClassData))
+		{
+			Container = (uint8*)InObj->GetClass()->GetOrCreateSparseClassData();
+		}
+		else
+		{
+			Container = (uint8*)InObj;
+		}
+		MemoryStack.Add(FMemoryFrame(nullptr, (uint8*)Container));
+		
+		// Get the direct memory pointer for the root property
+		if (Property->HasSetterOrGetter())
+		{
+			FMemoryFrame PropertyFrame(Property, (uint8*)Property->AllocateAndInitializeValue());
+			int32 StackIndex = MemoryStack.Add(PropertyFrame);
+			Property->GetValue_InContainer(MemoryStack[StackIndex - 1].Memory, PropertyFrame.Memory);
+		}
+		else
+		{
+			MemoryStack.Add(FMemoryFrame(Property, (uint8*)Property->ContainerPtrToValuePtr<uint8>(Container)));
+		}
+	}
+}
+
+void FPropertyNodeEditStack::CommitChanges()
+{
+	for (int32 Index = MemoryStack.Num() - 1; Index > 0; --Index)
+	{
+		if (MemoryStack[Index].Property->HasSetterOrGetter() &&
+			MemoryStack[Index].Property != MemoryStack[Index - 1].Property) // If this property is identical to the one below then it represents an item from an array
+		{
+			// Set the actual property value with the temp allocated memory
+			MemoryStack[Index].Property->SetValue_InContainer(MemoryStack[Index - 1].Memory, MemoryStack[Index].Memory);
+		}
+	}
+}
+
+void FPropertyNodeEditStack::Cleanup()
+{
+	for (int32 Index = MemoryStack.Num() - 1; Index > 0; --Index)
+	{
+		if (MemoryStack[Index].Property->HasSetterOrGetter() &&
+			MemoryStack[Index].Property != MemoryStack[Index - 1].Property) // If this property is identical to the one below then it represents an item from an array
+		{
+			MemoryStack[Index].Property->DestroyAndFreeValue(MemoryStack[Index].Memory);
+			MemoryStack[Index].Memory = nullptr;
+		}
+	}
+	MemoryStack.Empty();
+}
+
+FPropertyNodeEditStack::~FPropertyNodeEditStack()
+{
+	Cleanup();
+}
+
 FPropertyAccess::Result FPropertyNode::GetPropertyValueString(FString& OutString, const bool bAllowAlternateDisplayValue, EPropertyPortFlags PortFlags) const
 {
-	uint8* ValueAddress = nullptr;
-	FPropertyAccess::Result Result = GetSingleReadAddress(ValueAddress);
+	UObject* ValueContainer = nullptr;
+	FPropertyAccess::Result Result = GetSingleObject(ValueContainer);
 
-	if (ValueAddress != nullptr)
+	if (ValueContainer != nullptr)
 	{
 		const FProperty* PropertyPtr = GetProperty();
 
 		// Check for bogus data
 		if (PropertyPtr != nullptr && GetParentNode() != nullptr)
 		{
+			FPropertyNodeEditStack PropStack(this, ValueContainer);
+			uint8* ValueAddress = PropStack.GetDirectPropertyAddress();
+
 			FPropertyTextUtilities::PropertyToTextHelper(OutString, this, PropertyPtr, ValueAddress, nullptr, PortFlags);
 
 			UEnum* Enum = nullptr;
@@ -781,14 +895,17 @@ FPropertyAccess::Result FPropertyNode::GetPropertyValueString(FString& OutString
 
 FPropertyAccess::Result FPropertyNode::GetPropertyValueText(FText& OutText, const bool bAllowAlternateDisplayValue) const
 {
-	uint8* ValueAddress = nullptr;
-	FPropertyAccess::Result Result = GetSingleReadAddress(ValueAddress);
+	UObject* ValueContainer = nullptr;
+	FPropertyAccess::Result Result = GetSingleObject(ValueContainer);
 
-	if (ValueAddress != nullptr)
+	if (ValueContainer != nullptr)
 	{
 		const FProperty* PropertyPtr = GetProperty();
 		if (PropertyPtr)
 		{
+			FPropertyNodeEditStack PropStack(this, ValueContainer);
+			uint8* ValueAddress = PropStack.GetDirectPropertyAddress();
+
 			if (PropertyPtr->IsA(FTextProperty::StaticClass()))
 			{
 				OutText = CastField<FTextProperty>(PropertyPtr)->GetPropertyValue(ValueAddress);
@@ -1311,6 +1428,37 @@ FPropertyAccess::Result FPropertyNode::GetSingleReadAddress(uint8*& OutValueAddr
 	return ReadAddresses.Num() > 1 ? FPropertyAccess::MultipleValues : FPropertyAccess::Fail;
 }
 
+FPropertyAccess::Result FPropertyNode::GetSingleObject(UObject*& OutObject) const
+{
+	OutObject = nullptr;
+	FReadAddressList ReadAddresses;
+	bool bAllValuesTheSame = GetReadAddress(HasNodeFlags(EPropertyNodeFlags::SingleSelectOnly), ReadAddresses, false, true);
+
+	if ((ReadAddresses.Num() > 0 && bAllValuesTheSame) || ReadAddresses.Num() == 1)
+	{
+		OutObject = (UObject*)ReadAddresses.GetObject(0);
+
+		return FPropertyAccess::Success;
+	}
+
+	return ReadAddresses.Num() > 1 ? FPropertyAccess::MultipleValues : FPropertyAccess::Fail;
+}
+
+FPropertyAccess::Result FPropertyNode::GetSingleEditStack(FPropertyNodeEditStack& OutStack) const
+{
+	UObject* Object = nullptr;
+	FPropertyAccess::Result Result = FPropertyAccess::Fail;
+	if (GetProperty())
+	{
+		Result = GetSingleObject(Object);
+		if (Result == FPropertyAccess::Success)
+		{
+			OutStack.Initialize(this, Object);
+		}
+	}
+	return Result;
+}
+
 uint8* FPropertyNode::GetStartAddressFromObject(const UObject* Obj) const
 {
 	if (!Obj)
diff --git a/Engine/Source/Editor/PropertyEditor/Private/PropertyNode.h b/Engine/Source/Editor/PropertyEditor/Private/PropertyNode.h
--- a/Engine/Source/Editor/PropertyEditor/Private/PropertyNode.h
+++ b/Engine/Source/Editor/PropertyEditor/Private/PropertyNode.h
@@ -136,6 +136,12 @@ public:
 		return (Pair.Object.IsValid() || Pair.bIsStruct) ? Pair.ReadAddress : 0;
 	}
 
+	const UObject* GetObject(int32 Index)
+	{
+		const FAddressPair& Pair = ReadAddresses[Index];
+		return Pair.Object.Get();
+	}
+
 	bool IsValidIndex( int32 Index ) const
 	{
 		return ReadAddresses.IsValidIndex(Index);
@@ -175,6 +181,11 @@ public:
 		return ReadAddressListData->GetAddress( Index );
 	}
 	
+	const UObject* GetObject(int32 Index)
+	{
+		return ReadAddressListData->GetObject(Index);
+	}
+
 	bool IsValidIndex( int32 Index ) const
 	{
 		return ReadAddressListData->IsValidIndex( Index );
@@ -275,6 +286,65 @@ enum EPropertyDataValidationResult : uint8
 	DataValid,
 };
 
+/** Helper class for modifying property values with setters and getters */
+class FPropertyNodeEditStack
+{
+	struct FMemoryFrame
+	{
+		FMemoryFrame() = default;
+		FMemoryFrame(const FProperty* InProperty, uint8* InMemory)
+			: Property(InProperty)
+			, Memory(InMemory)
+		{
+		}
+		/** Property that points to the memory in this frame */
+		const FProperty* Property = nullptr;
+		/** Property address */
+		uint8* Memory = nullptr;
+	};
+public:
+
+	/**
+	* Constructs property stack for the specified node
+	* InNode Property node to construct the stack for
+	* InObj Object instance that contains the property being modified
+	*/
+	FPropertyNodeEditStack(const FPropertyNode* InNode, const UObject* InObj);
+	FPropertyNodeEditStack() = default;
+	~FPropertyNodeEditStack();
+
+	FPropertyNodeEditStack& operator = (const FPropertyNodeEditStack& Other) = delete;
+	FPropertyNodeEditStack(const FPropertyNodeEditStack& Other) = delete;
+
+	/**
+	* Initializes property stack for the specified node
+	* InNode Property node to construct the stack for
+	* InObj Object instance that contains the property being modified
+	*/
+	void Initialize(const FPropertyNode* InNode, const UObject* InObj);
+
+	/**
+	* Returns the address of the property being modified.
+	* If anywhere in the property stack is a property with a setter or getter this will point to a temporarily allocated memory.
+	*/
+	uint8* GetDirectPropertyAddress()
+	{
+		return MemoryStack.Last().Memory;
+	}
+
+	/**
+	* Commits all modifications to temporarily allocated property values back to the actual member variables using setters and getters where available
+	*/
+	void CommitChanges();
+
+private:
+
+	void InitializeInternal(const FPropertyNode* InNode, const UObject* InObj);
+	void Cleanup();
+
+	TArray<FMemoryFrame> MemoryStack;
+};
+
 /**
  * The base class for all property nodes                                                              
  */
@@ -493,6 +563,19 @@ public:
 	 */
 	FPropertyAccess::Result GetSingleReadAddress(uint8*& OutValueAddress) const;
 
+	/**
+	 * Fills in the OutObject with the address of the object of all the available objects.
+	 * If multiple items are selected, this will return a null address unless they are all the same value.
+	 * @param OutObject	The address of the Object
+	 */
+	FPropertyAccess::Result GetSingleObject(UObject*& OutObject) const;
+
+	/**
+	 * Fills in the OutContainer with the address of the container (struct or UObject instance) that owns the property this node represents.
+	 * @param OutContainer	The address of the container instance
+	 */
+	FPropertyAccess::Result GetSingleEditStack(FPropertyNodeEditStack& OutStack) const;
+
 	/**
 	 * Gets read addresses without accessing cached data.  Is less efficient but gets the must up to date data
 	 */
