Fixing a crash when opening struct-only editor panel

diff --git a/Engine/Source/Editor/PropertyEditor/Private/PropertyNode.cpp b/Engine/Source/Editor/PropertyEditor/Private/PropertyNode.cpp
--- a/Engine/Source/Editor/PropertyEditor/Private/PropertyNode.cpp
+++ b/Engine/Source/Editor/PropertyEditor/Private/PropertyNode.cpp
@@ -776,17 +776,33 @@ void FPropertyNodeEditStack::InitializeInternal(const FPropertyNode* InNode, con
 	}
 	else
 	{
-		// Determine the root container object (UObject instance or sparse class data) for this property stack
+		const UObject* Object = InObj;
+		if (!Object)
+		{
+			UObject* NodeObject = nullptr;
+			FPropertyAccess::Result Result = InNode->GetSingleObject(NodeObject);
+			Object = NodeObject;
+			checkf(Result == FPropertyAccess::Success, TEXT("Unable to get object for property %s"), *GetNameSafe(Property));
+		}
+
+		// Determine the root container address (Struct address, UObject instance or sparse class data) for this property stack
 		uint8* Container = nullptr;
 		if (InNode->HasNodeFlags(EPropertyNodeFlags::IsSparseClassData))
 		{
-			Container = (uint8*)InObj->GetClass()->GetOrCreateSparseClassData();
+			checkf(Object != nullptr, TEXT("No object pointer for property %s"), *GetNameSafe(Property));
+			Container = (uint8*)Object->GetClass()->GetOrCreateSparseClassData();
+		}
+		else if (Object)
+		{
+			Container = (uint8*)Object;
 		}
 		else
 		{
-			Container = (uint8*)InObj;
+			FPropertyAccess::Result Result = InNode->GetSingleReadAddress(Container);
+			checkf(Result == FPropertyAccess::Success, TEXT("Unable to get read address for property %s"), *GetNameSafe(Property));
 		}
-		MemoryStack.Add(FMemoryFrame(nullptr, (uint8*)Container));
+		checkf(Container, TEXT("Container pointer can't be null. Creating edit stack for property %s"), *GetNameSafe(Property));
+		MemoryStack.Add(FMemoryFrame(nullptr, Container));
 		
 		// Get the direct memory pointer for the root property
 		if (Property->HasSetterOrGetter())
diff --git a/Engine/Source/Editor/PropertyEditor/Private/PropertyNode.h b/Engine/Source/Editor/PropertyEditor/Private/PropertyNode.h
--- a/Engine/Source/Editor/PropertyEditor/Private/PropertyNode.h
+++ b/Engine/Source/Editor/PropertyEditor/Private/PropertyNode.h
@@ -307,9 +307,9 @@ public:
 	/**
 	* Constructs property stack for the specified node
 	* InNode Property node to construct the stack for
-	* InObj Object instance that contains the property being modified
+	* InObj Optional Object instance that contains the property being modified (if not provided the root container pointer will be acquired from the provided node hierarchy)
 	*/
-	FPropertyNodeEditStack(const FPropertyNode* InNode, const UObject* InObj);
+	FPropertyNodeEditStack(const FPropertyNode* InNode, const UObject* InObj = nullptr);
 	FPropertyNodeEditStack() = default;
 	~FPropertyNodeEditStack();
 
