Fixing crashes when creating property panels: making FPropertyNodeEditStack respect error results when acquiring property read address.

diff --git a/Engine/Source/Editor/PropertyEditor/Private/PropertyNode.cpp b/Engine/Source/Editor/PropertyEditor/Private/PropertyNode.cpp
--- a/Engine/Source/Editor/PropertyEditor/Private/PropertyNode.cpp
+++ b/Engine/Source/Editor/PropertyEditor/Private/PropertyNode.cpp
@@ -725,23 +725,33 @@ EPropertyDataValidationResult FPropertyNode::EnsureDataIsValid()
 
 FPropertyNodeEditStack::FPropertyNodeEditStack(const FPropertyNode* InNode, const UObject* InObj)
 {
-	InitializeInternal(InNode, InObj);
+	Initialize(InNode, InObj);
 }
 
-void FPropertyNodeEditStack::Initialize(const FPropertyNode* InNode, const UObject* InObj)
+FPropertyAccess::Result FPropertyNodeEditStack::Initialize(const FPropertyNode* InNode, const UObject* InObj)
 {
 	Cleanup();
-	InitializeInternal(InNode, InObj);
+	FPropertyAccess::Result Result = InitializeInternal(InNode, InObj);
+	if (Result != FPropertyAccess::Success)
+	{
+		Cleanup();
+	}
+	return Result;
 }
 
-void FPropertyNodeEditStack::InitializeInternal(const FPropertyNode* InNode, const UObject* InObj)
+FPropertyAccess::Result FPropertyNodeEditStack::InitializeInternal(const FPropertyNode* InNode, const UObject* InObj)
 {
+	FPropertyAccess::Result Result = FPropertyAccess::Success;
 	const FPropertyNode* Parent = InNode->GetParentNode();
 	const FProperty* Property = InNode->GetProperty();
 	if (Parent && Parent->GetProperty())
 	{
 		// Recursively initialize the stack
-		InitializeInternal(Parent, InObj);
+		Result = InitializeInternal(Parent, InObj);
+		if (Result != FPropertyAccess::Success)
+		{
+			return Result;
+		}
 
 		// Get the direct memory pointer for the current property
 		const FProperty* ParentProperty = Parent->GetProperty();
@@ -780,9 +790,12 @@ void FPropertyNodeEditStack::InitializeInternal(const FPropertyNode* InNode, con
 		if (!Object)
 		{
 			UObject* NodeObject = nullptr;
-			FPropertyAccess::Result Result = InNode->GetSingleObject(NodeObject);
+			Result = InNode->GetSingleObject(NodeObject);
+			if (Result != FPropertyAccess::Success)
+			{
+				return Result;
+			}
 			Object = NodeObject;
-			checkf(Result == FPropertyAccess::Success, TEXT("Unable to get object for property %s"), *GetNameSafe(Property));
 		}
 
 		// Determine the root container address (Struct address, UObject instance or sparse class data) for this property stack
@@ -798,10 +811,17 @@ void FPropertyNodeEditStack::InitializeInternal(const FPropertyNode* InNode, con
 		}
 		else
 		{
-			FPropertyAccess::Result Result = InNode->GetSingleReadAddress(Container);
-			checkf(Result == FPropertyAccess::Success, TEXT("Unable to get read address for property %s"), *GetNameSafe(Property));
+			Result = InNode->GetSingleReadAddress(Container);
+			if (Result != FPropertyAccess::Success)
+			{
+				return Result;
+			}
+		}
+		if (!Container)
+		{
+			// This may happen when the node points at stale object
+			return FPropertyAccess::Fail;
 		}
-		checkf(Container, TEXT("Container pointer can't be null. Creating edit stack for property %s"), *GetNameSafe(Property));
 		MemoryStack.Add(FMemoryFrame(nullptr, Container));
 		
 		// Get the direct memory pointer for the root property
@@ -816,6 +836,7 @@ void FPropertyNodeEditStack::InitializeInternal(const FPropertyNode* InNode, con
 			MemoryStack.Add(FMemoryFrame(Property, (uint8*)Property->ContainerPtrToValuePtr<uint8>(Container)));
 		}
 	}
+	return Result;
 }
 
 void FPropertyNodeEditStack::CommitChanges()
@@ -1469,7 +1490,7 @@ FPropertyAccess::Result FPropertyNode::GetSingleEditStack(FPropertyNodeEditStack
 		Result = GetSingleObject(Object);
 		if (Result == FPropertyAccess::Success)
 		{
-			OutStack.Initialize(this, Object);
+			Result = OutStack.Initialize(this, Object);
 		}
 	}
 	return Result;
diff --git a/Engine/Source/Editor/PropertyEditor/Private/PropertyNode.h b/Engine/Source/Editor/PropertyEditor/Private/PropertyNode.h
--- a/Engine/Source/Editor/PropertyEditor/Private/PropertyNode.h
+++ b/Engine/Source/Editor/PropertyEditor/Private/PropertyNode.h
@@ -321,7 +321,7 @@ public:
 	* InNode Property node to construct the stack for
 	* InObj Object instance that contains the property being modified
 	*/
-	void Initialize(const FPropertyNode* InNode, const UObject* InObj);
+	FPropertyAccess::Result Initialize(const FPropertyNode* InNode, const UObject* InObj);
 
 	/**
 	* Returns the address of the property being modified.
@@ -337,9 +337,15 @@ public:
 	*/
 	void CommitChanges();
 
+	/** Checks if this edit stack is valid */
+	bool IsValid() const
+	{
+		return MemoryStack.Num() > 0;
+	}
+
 private:
 
-	void InitializeInternal(const FPropertyNode* InNode, const UObject* InObj);
+	FPropertyAccess::Result InitializeInternal(const FPropertyNode* InNode, const UObject* InObj);
 	void Cleanup();
 
 	TArray<FMemoryFrame> MemoryStack;
