Fixing FPropertyNodeEditStack assuming the read address for struct property nodes is the struct memory address whereas it's a direct property address.

+ Fixing logic errors when getting property values as string

Fixing crashes when creating property panels: making FPropertyNodeEditStack respect error results when acquiring property read address.

diff --git a/Engine/Source/Editor/PropertyEditor/Private/PropertyNode.cpp b/Engine/Source/Editor/PropertyEditor/Private/PropertyNode.cpp
--- a/Engine/Source/Editor/PropertyEditor/Private/PropertyNode.cpp
+++ b/Engine/Source/Editor/PropertyEditor/Private/PropertyNode.cpp
@@ -831,10 +831,16 @@ FPropertyAccess::Result FPropertyNodeEditStack::InitializeInternal(const FProper
 			int32 StackIndex = MemoryStack.Add(PropertyFrame);
 			Property->GetValue_InContainer(MemoryStack[StackIndex - 1].Memory, PropertyFrame.Memory);
 		}
-		else
+		else if ((uint8*)Object == Container)
 		{
 			MemoryStack.Add(FMemoryFrame(Property, (uint8*)Property->ContainerPtrToValuePtr<uint8>(Container)));
 		}
+		else
+		{
+			// This node represents a struct in which case the Container represents direct memory for the root property.
+			// todo: RobM: ideally we want Container to be the struct memory and not a property address
+			MemoryStack.Add(FMemoryFrame(Property, Container));
+		}
 	}
 	return Result;
 }
@@ -873,19 +879,22 @@ FPropertyNodeEditStack::~FPropertyNodeEditStack()
 
 FPropertyAccess::Result FPropertyNode::GetPropertyValueString(FString& OutString, const bool bAllowAlternateDisplayValue, EPropertyPortFlags PortFlags) const
 {
-	UObject* ValueContainer = nullptr;
-	FPropertyAccess::Result Result = GetSingleObject(ValueContainer);
-
-	if (ValueContainer != nullptr)
+	FPropertyAccess::Result Result = FPropertyAccess::Fail;
+	const FProperty* PropertyPtr = GetProperty();
+	// Check for bogus data
+	if (PropertyPtr != nullptr && GetParentNode() != nullptr)
 	{
-		const FProperty* PropertyPtr = GetProperty();
+		UObject* ValueContainer = nullptr;		
+		FPropertyNodeEditStack PropStack;
+		Result = GetSingleObject(ValueContainer);
+		if (Result == FPropertyAccess::Success)
+		{
+			Result = PropStack.Initialize(this, ValueContainer);
+		}
 
-		// Check for bogus data
-		if (PropertyPtr != nullptr && GetParentNode() != nullptr)
+		if (Result == FPropertyAccess::Success)
 		{
-			FPropertyNodeEditStack PropStack(this, ValueContainer);
 			uint8* ValueAddress = PropStack.GetDirectPropertyAddress();
-
 			FPropertyTextUtilities::PropertyToTextHelper(OutString, this, PropertyPtr, ValueAddress, nullptr, PortFlags);
 
 			UEnum* Enum = nullptr;
@@ -921,10 +930,6 @@ FPropertyAccess::Result FPropertyNode::GetPropertyValueString(FString& OutString
 				}
 			}
 		}
-		else
-		{
-			Result = FPropertyAccess::Fail;
-		}
 	}
 
 	return Result;
@@ -932,15 +937,20 @@ FPropertyAccess::Result FPropertyNode::GetPropertyValueString(FString& OutString
 
 FPropertyAccess::Result FPropertyNode::GetPropertyValueText(FText& OutText, const bool bAllowAlternateDisplayValue) const
 {
-	UObject* ValueContainer = nullptr;
-	FPropertyAccess::Result Result = GetSingleObject(ValueContainer);
-
-	if (ValueContainer != nullptr)
+	FPropertyAccess::Result Result = FPropertyAccess::Fail;
+	const FProperty* PropertyPtr = GetProperty();
+	if (PropertyPtr)
 	{
-		const FProperty* PropertyPtr = GetProperty();
-		if (PropertyPtr)
+		UObject* ValueContainer = nullptr;
+		FPropertyNodeEditStack PropStack;
+		Result = GetSingleObject(ValueContainer);
+		if (Result == FPropertyAccess::Success)
+		{
+			Result = PropStack.Initialize(this, ValueContainer);
+		}
+
+		if (Result == FPropertyAccess::Success)
 		{
-			FPropertyNodeEditStack PropStack(this, ValueContainer);
 			uint8* ValueAddress = PropStack.GetDirectPropertyAddress();
 
 			if (PropertyPtr->IsA(FTextProperty::StaticClass()))
@@ -983,10 +993,6 @@ FPropertyAccess::Result FPropertyNode::GetPropertyValueText(FText& OutText, cons
 				}
 			}
 		}
-		else
-		{
-			Result = FPropertyAccess::Fail;
-		}
 	}
 
 	return Result;
