Moved SPropertyEditor and SPropertyEditorArray implementations into .cpp files.

diff --git a/Engine/Source/Editor/PropertyEditor/Private/UserInterface/PropertyEditor/SPropertyEditor.cpp b/Engine/Source/Editor/PropertyEditor/Private/UserInterface/PropertyEditor/SPropertyEditor.cpp
new file mode 100644
--- /dev/null
+++ b/Engine/Source/Editor/PropertyEditor/Private/UserInterface/PropertyEditor/SPropertyEditor.cpp
@@ -0,0 +1,46 @@
+// Copyright Epic Games, Inc. All Rights Reserved.
+
+#include "SPropertyEditor.h"
+#include "Presentation/PropertyEditor/PropertyEditor.h"
+#include "UObject/UnrealType.h"
+#include "Widgets/Input/SEditableTextBox.h"
+
+void SPropertyEditor::Construct(const FArguments& InArgs, const TSharedRef<FPropertyEditor>& InPropertyEditor)
+{
+	PropertyEditor = InPropertyEditor;
+
+	if (ShouldShowValue(InPropertyEditor))
+	{
+		ChildSlot
+		[
+			// Make a read only text box so that copy still works
+			SNew(SEditableTextBox)
+			.Text(InPropertyEditor, &FPropertyEditor::GetValueAsText)
+			.ToolTipText(InPropertyEditor, &FPropertyEditor::GetValueAsText)
+			.Font(InArgs._Font)
+			.IsReadOnly(true)
+		];
+	}
+}
+
+void SPropertyEditor::GetDesiredWidth(float& OutMinDesiredWidth, float& OutMaxDesiredWidth) const
+{
+	OutMinDesiredWidth = 125.0f;
+	OutMaxDesiredWidth = 0.0f;
+
+	if (PropertyEditor.IsValid())
+	{
+		const FProperty* Property = PropertyEditor->GetProperty();
+		if (Property && Property->IsA<FStructProperty>())
+		{
+			// Struct headers with nothing in them have no min width
+			OutMinDesiredWidth = 0;
+			OutMaxDesiredWidth = 130.0f;
+		}
+	}
+}
+
+bool SPropertyEditor::ShouldShowValue(const TSharedRef<FPropertyEditor>& InPropertyEditor) const
+{
+	return PropertyEditor->GetProperty() && !PropertyEditor->GetProperty()->IsA<FStructProperty>();
+}
diff --git a/Engine/Source/Editor/PropertyEditor/Private/UserInterface/PropertyEditor/SPropertyEditor.h b/Engine/Source/Editor/PropertyEditor/Private/UserInterface/PropertyEditor/SPropertyEditor.h
--- a/Engine/Source/Editor/PropertyEditor/Private/UserInterface/PropertyEditor/SPropertyEditor.h
+++ b/Engine/Source/Editor/PropertyEditor/Private/UserInterface/PropertyEditor/SPropertyEditor.h
@@ -1,69 +1,29 @@
 // Copyright Epic Games, Inc. All Rights Reserved.
 #pragma once
 
-#include "CoreMinimal.h"
-#include "UObject/UnrealType.h"
 #include "Fonts/SlateFontInfo.h"
 #include "Styling/AppStyle.h"
+#include "UserInterface/PropertyEditor/PropertyEditorConstants.h"
 #include "Widgets/DeclarativeSyntaxSupport.h"
 #include "Widgets/SCompoundWidget.h"
-#include "Presentation/PropertyEditor/PropertyEditor.h"
-#include "UserInterface/PropertyEditor/PropertyEditorConstants.h"
-#include "Widgets/Input/SEditableTextBox.h"
 
-#define LOCTEXT_NAMESPACE "PropertyEditor"
+class FPropertyEditor;
 
 class SPropertyEditor : public SCompoundWidget
 {
 public:
-	SLATE_BEGIN_ARGS( SPropertyEditor )
-		: _Font( FAppStyle::GetFontStyle( PropertyEditorConstants::PropertyFontStyle ) ) 
+	SLATE_BEGIN_ARGS(SPropertyEditor)
+		: _Font(FAppStyle::GetFontStyle(PropertyEditorConstants::PropertyFontStyle))
 	{}
-	SLATE_ATTRIBUTE( FSlateFontInfo, Font )
-		SLATE_END_ARGS()
-
-	void Construct( const FArguments& InArgs, const TSharedRef< class FPropertyEditor >& InPropertyEditor )
-	{
-		PropertyEditor = InPropertyEditor;
+		SLATE_ATTRIBUTE(FSlateFontInfo, Font)
+	SLATE_END_ARGS()
 
-		if( ShouldShowValue( InPropertyEditor ) )
-		{
-			ChildSlot
-			[
-				// Make a read only text box so that copy still works
-				SNew( SEditableTextBox )
-				.Text( InPropertyEditor, &FPropertyEditor::GetValueAsText )
-				.ToolTipText( InPropertyEditor, &FPropertyEditor::GetValueAsText )
-				.Font( InArgs._Font )
-				.IsReadOnly( true )
-			];
-		}
-	}
+	void Construct(const FArguments& InArgs, const TSharedRef<FPropertyEditor>& InPropertyEditor);
+	void GetDesiredWidth(float& OutMinDesiredWidth, float& OutMaxDesiredWidth) const;
 
-	void GetDesiredWidth( float& OutMinDesiredWidth, float& OutMaxDesiredWidth )
-	{
-		OutMinDesiredWidth = 125.0f;
-		OutMaxDesiredWidth = 0.0f;
-
-		if( PropertyEditor.IsValid() )
-		{
-			const FProperty* Property = PropertyEditor->GetProperty();
-			if( Property && Property->IsA<FStructProperty>() )
-			{
-				// Struct headers with nothing in them have no min width
-				OutMinDesiredWidth = 0;
-				OutMaxDesiredWidth = 130.0f;
-			}
-			
-		}
-	}
 private:
-	bool ShouldShowValue( const TSharedRef< class FPropertyEditor >& InPropertyEditor ) const 
-	{
-		return PropertyEditor->GetProperty() && !PropertyEditor->GetProperty()->IsA<FStructProperty>();
-	}
+	bool ShouldShowValue(const TSharedRef<FPropertyEditor>& InPropertyEditor) const;
+
 private:
-	TSharedPtr< class FPropertyEditor > PropertyEditor;
+	TSharedPtr<FPropertyEditor> PropertyEditor;
 };
-
-#undef LOCTEXT_NAMESPACE
diff --git a/Engine/Source/Editor/PropertyEditor/Private/UserInterface/PropertyEditor/SPropertyEditorArray.cpp b/Engine/Source/Editor/PropertyEditor/Private/UserInterface/PropertyEditor/SPropertyEditorArray.cpp
new file mode 100644
--- /dev/null
+++ b/Engine/Source/Editor/PropertyEditor/Private/UserInterface/PropertyEditor/SPropertyEditorArray.cpp
@@ -0,0 +1,155 @@
+// Copyright Epic Games, Inc. All Rights Reserved.
+
+#include "SPropertyEditorArray.h"
+#include "DragAndDrop/AssetDragDropOp.h"
+#include "Presentation/PropertyEditor/PropertyEditor.h"
+#include "PropertyEditorHelpers.h"
+#include "SDropTarget.h"
+#include "Widgets/Text/STextBlock.h"
+
+#define LOCTEXT_NAMESPACE "PropertyEditor"
+
+void SPropertyEditorArray::Construct(const FArguments& InArgs, const TSharedRef<FPropertyEditor>& InPropertyEditor)
+{
+	PropertyEditor = InPropertyEditor;
+
+	TAttribute<FText> TextAttr;
+	if (PropertyEditorHelpers::IsStaticArray(*InPropertyEditor->GetPropertyNode()))
+	{
+		// Static arrays need special case handling for their values
+		TextAttr.Set(GetArrayTextValue());
+	}
+	else
+	{
+		TextAttr.Bind(this, &SPropertyEditorArray::GetArrayTextValue);
+	}
+
+	ChildSlot
+	.Padding(0.0f, 0.0f, 2.0f, 0.0f)
+	[
+		SNew(SDropTarget)
+		.OnDropped(this, &SPropertyEditorArray::OnDragDropTarget)
+		.OnAllowDrop(this, &SPropertyEditorArray::WillAddValidElements)
+		.OnIsRecognized(this, &SPropertyEditorArray::IsValidAssetDropOp)
+		.Content()
+		[
+			SNew(STextBlock)
+			.Text(TextAttr)
+			.Font(InArgs._Font)
+		]
+	];
+
+	SetEnabled(TAttribute<bool>(this, &SPropertyEditorArray::CanEdit));
+}
+
+bool SPropertyEditorArray::Supports(const TSharedRef<FPropertyEditor>& InPropertyEditor)
+{
+	const FProperty* NodeProperty = InPropertyEditor->GetProperty();
+
+	return PropertyEditorHelpers::IsStaticArray(*InPropertyEditor->GetPropertyNode())
+		|| PropertyEditorHelpers::IsDynamicArray(*InPropertyEditor->GetPropertyNode());
+}
+
+void SPropertyEditorArray::GetDesiredWidth(float& OutMinDesiredWidth, float& OutMaxDesiredWidth) const
+{
+	OutMinDesiredWidth = 170.0f;
+	OutMaxDesiredWidth = 170.0f;
+}
+
+FText SPropertyEditorArray::GetArrayTextValue() const
+{
+	FString ArrayString;
+	FPropertyAccess::Result GetValResult = PropertyEditor->GetPropertyHandle()->GetValueAsDisplayString(ArrayString);
+
+	if (GetValResult == FPropertyAccess::MultipleValues)
+	{
+		return LOCTEXT("MultipleValues", "Multiple Values");
+	}
+	return FText::Format(LOCTEXT("NumArrayItemsFmt", "{0} Array elements"), FText::AsNumber(PropertyEditor->GetPropertyNode()->GetNumChildNodes()));
+}
+
+bool SPropertyEditorArray::CanEdit() const
+{
+	return PropertyEditor.IsValid() ? !PropertyEditor->IsEditConst() : true;
+}
+
+FReply SPropertyEditorArray::OnDragDropTarget(const FGeometry& InGeometry, const FDragDropEvent& InDragDropEvent)
+{
+	FObjectProperty* ObjectProperty = nullptr;
+	if (const FArrayProperty* ArrayProperty = CastField<FArrayProperty>(PropertyEditor->GetProperty()))
+	{
+		ObjectProperty = CastField<FObjectProperty>(ArrayProperty->Inner);
+	}
+
+	// Only try to add entries if we are dropping on an asset array
+	TSharedPtr<FDragDropOperation> DragOperation = InDragDropEvent.GetOperation();
+	if (ObjectProperty && DragOperation && DragOperation->IsOfType<FAssetDragDropOp>())
+	{
+		TSharedPtr<FAssetDragDropOp> DragDropOp = StaticCastSharedPtr<FAssetDragDropOp>(DragOperation);
+		if (DragDropOp.IsValid())
+		{
+			for (FAssetData AssetData : DragDropOp->GetAssets())
+			{
+				// if the type matches
+				if (AssetData.IsInstanceOf(ObjectProperty->PropertyClass))
+				{
+					PropertyEditor->AddGivenItem(AssetData.GetObjectPathString());
+				}
+			}
+			// Let this bubble up to the rest of the row
+			return FReply::Unhandled();
+
+		}
+	}
+	return FReply::Unhandled();
+}
+
+bool SPropertyEditorArray::IsValidAssetDropOp(TSharedPtr<FDragDropOperation> InOperation)
+{
+	FObjectProperty* ObjectProperty = nullptr;
+	if (const FArrayProperty* ArrayProperty = CastField<FArrayProperty>(PropertyEditor->GetProperty()))
+	{
+		ObjectProperty = CastField<FObjectProperty>(ArrayProperty->Inner);
+	}
+
+	// Only try to add entries if we are dropping on an asset array
+	if (ObjectProperty && InOperation->IsOfType<FAssetDragDropOp>())
+	{
+		return true;
+	}
+	return false;
+}
+
+bool SPropertyEditorArray::WillAddValidElements(TSharedPtr<FDragDropOperation> InOperation)
+{
+	FObjectProperty* ObjectProperty = nullptr;
+
+	if (const FArrayProperty* ArrayProperty = CastField<FArrayProperty>(PropertyEditor->GetProperty()))
+	{
+		ObjectProperty = CastField<FObjectProperty>(ArrayProperty->Inner);
+	}
+
+	// Only try to add entries if we are dropping on an asset array
+	if (ObjectProperty && InOperation->IsOfType<FAssetDragDropOp>())
+	{
+		bool bHasOnlyValidElements = true;
+		TSharedPtr<FAssetDragDropOp> DragDropOp = StaticCastSharedPtr<FAssetDragDropOp>(InOperation);
+		if (DragDropOp.IsValid())
+		{
+			for (FAssetData AssetData : DragDropOp->GetAssets())
+			{
+				// if the type does not match
+				if (!AssetData.IsInstanceOf(ObjectProperty->PropertyClass))
+				{
+					bHasOnlyValidElements = false;
+					break;
+				}
+			}
+		}
+		return bHasOnlyValidElements;
+	}
+
+	return false;
+}
+
+#undef LOCTEXT_NAMESPACE
diff --git a/Engine/Source/Editor/PropertyEditor/Private/UserInterface/PropertyEditor/SPropertyEditorArray.h b/Engine/Source/Editor/PropertyEditor/Private/UserInterface/PropertyEditor/SPropertyEditorArray.h
--- a/Engine/Source/Editor/PropertyEditor/Private/UserInterface/PropertyEditor/SPropertyEditorArray.h
+++ b/Engine/Source/Editor/PropertyEditor/Private/UserInterface/PropertyEditor/SPropertyEditorArray.h
@@ -1,178 +1,42 @@
 // Copyright Epic Games, Inc. All Rights Reserved.
 #pragma once
 
-#include "CoreMinimal.h"
-#include "Misc/Attribute.h"
-#include "Widgets/SCompoundWidget.h"
-#include "Widgets/DeclarativeSyntaxSupport.h"
 #include "Fonts/SlateFontInfo.h"
+#include "Internationalization/Text.h"
 #include "Styling/AppStyle.h"
-#include "Presentation/PropertyEditor/PropertyEditor.h"
-#include "PropertyEditorHelpers.h"
 #include "UserInterface/PropertyEditor/PropertyEditorConstants.h"
-#include "Widgets/Text/STextBlock.h"
-#include "DragAndDrop/AssetDragDropOp.h"
-#include "SDropTarget.h"
-
+#include "Widgets/DeclarativeSyntaxSupport.h"
+#include "Widgets/SCompoundWidget.h"
 
-#define LOCTEXT_NAMESPACE "PropertyEditor"
+class FDragDropEvent;
+class FDragDropOperation;
+class FPropertyEditor;
 
 class SPropertyEditorArray : public SCompoundWidget
 {
 public:
 
-	SLATE_BEGIN_ARGS( SPropertyEditorArray )
-		: _Font( FAppStyle::GetFontStyle( PropertyEditorConstants::PropertyFontStyle ) ) 
-		{}
-		SLATE_ATTRIBUTE( FSlateFontInfo, Font )
+	SLATE_BEGIN_ARGS(SPropertyEditorArray)
+		: _Font(FAppStyle::GetFontStyle(PropertyEditorConstants::PropertyFontStyle))
+	{}
+		SLATE_ATTRIBUTE(FSlateFontInfo, Font)
 	SLATE_END_ARGS()
 
-	void Construct( const FArguments& InArgs, const TSharedRef<FPropertyEditor>& InPropertyEditor )
-	{
-		PropertyEditor = InPropertyEditor;
-
-		TAttribute<FText> TextAttr;
-		if( PropertyEditorHelpers::IsStaticArray( *InPropertyEditor->GetPropertyNode() ) )
-		{
-			// Static arrays need special case handling for their values
-			TextAttr.Set( GetArrayTextValue() );
-		}
-		else
-		{
-			TextAttr.Bind( this, &SPropertyEditorArray::GetArrayTextValue );
-		}
-
-		ChildSlot
-		.Padding( 0.0f, 0.0f, 2.0f, 0.0f )
-		[
-			SNew(SDropTarget)
-			.OnDropped(this, &SPropertyEditorArray::OnDragDropTarget)
-			.OnAllowDrop(this, &SPropertyEditorArray::WillAddValidElements)
-			.OnIsRecognized(this, &SPropertyEditorArray::IsValidAssetDropOp)
-			.Content()
-			[
-				SNew(STextBlock)
-				.Text(TextAttr)
-				.Font(InArgs._Font)
-			]
-		];
-
-		SetEnabled( TAttribute<bool>( this, &SPropertyEditorArray::CanEdit ) );
-	}
-
-	static bool Supports( const TSharedRef<FPropertyEditor>& InPropertyEditor )
-	{
-		const FProperty* NodeProperty = InPropertyEditor->GetProperty();
+	void Construct(const FArguments& InArgs, const TSharedRef<FPropertyEditor>& InPropertyEditor);
+	void GetDesiredWidth(float& OutMinDesiredWidth, float& OutMaxDesiredWidth) const;
 
-		return PropertyEditorHelpers::IsStaticArray( *InPropertyEditor->GetPropertyNode() ) 
-			|| PropertyEditorHelpers::IsDynamicArray( *InPropertyEditor->GetPropertyNode() );
-	}
+	static bool Supports(const TSharedRef<FPropertyEditor>& InPropertyEditor);
 
-	void GetDesiredWidth( float& OutMinDesiredWidth, float &OutMaxDesiredWidth )
-	{
-		OutMinDesiredWidth = 170.0f;
-		OutMaxDesiredWidth = 170.0f;
-	}
 private:
-	FText GetArrayTextValue() const
-	{
-		FString ArrayString;
-		FPropertyAccess::Result GetValResult = PropertyEditor->GetPropertyHandle()->GetValueAsDisplayString(ArrayString);
-
-		if (GetValResult == FPropertyAccess::MultipleValues)
-		{
-			return LOCTEXT("MultipleValues", "Multiple Values");
-		}
-		return FText::Format( LOCTEXT("NumArrayItemsFmt", "{0} Array elements"), FText::AsNumber(PropertyEditor->GetPropertyNode()->GetNumChildNodes()) );
-	}
+	FText GetArrayTextValue() const;
 
 	/** @return True if the property can be edited */
-	bool CanEdit() const
-	{
-		return PropertyEditor.IsValid() ? !PropertyEditor->IsEditConst() : true;
-	}
-
-	FReply OnDragDropTarget(const FGeometry& InGeometry, const FDragDropEvent& InDragDropEvent)
-	{
-		FObjectProperty* ObjectProperty = nullptr;
-		if (const FArrayProperty* ArrayProperty = CastField<FArrayProperty>(PropertyEditor->GetProperty()))
-		{
-			ObjectProperty = CastField<FObjectProperty>(ArrayProperty->Inner);
-		}
-
-		// Only try to add entries if we are dropping on an asset array
-		TSharedPtr<FDragDropOperation> DragOperation = InDragDropEvent.GetOperation();
-		if (ObjectProperty && DragOperation && DragOperation->IsOfType<FAssetDragDropOp>())
-		{
-			TSharedPtr<FAssetDragDropOp> DragDropOp = StaticCastSharedPtr<FAssetDragDropOp>(DragOperation);
-			if (DragDropOp.IsValid())
-			{
-				for (FAssetData AssetData : DragDropOp->GetAssets())
-				{
-					// if the type matches
-					if (AssetData.IsInstanceOf(ObjectProperty->PropertyClass))
-					{
-						PropertyEditor->AddGivenItem(AssetData.GetObjectPathString());
-					}
-				}
-				// Let this bubble up to the rest of the row
-				return FReply::Unhandled();
-
-			}
-		}
-		return FReply::Unhandled();
-	}
-
-	bool IsValidAssetDropOp(TSharedPtr<FDragDropOperation> InOperation)
-	{
-		FObjectProperty* ObjectProperty = nullptr;
-		if (const FArrayProperty* ArrayProperty = CastField<FArrayProperty>(PropertyEditor->GetProperty()))
-		{
-			ObjectProperty = CastField<FObjectProperty>(ArrayProperty->Inner);
-		}
+	bool CanEdit() const;
 
-		// Only try to add entries if we are dropping on an asset array
-		if (ObjectProperty && InOperation->IsOfType<FAssetDragDropOp>())
-		{
-			return true;
-		}
-		return false;
-	}
-
-	bool WillAddValidElements(TSharedPtr<FDragDropOperation> InOperation)
-	{
-		FObjectProperty* ObjectProperty = nullptr;
-
-		if (const FArrayProperty* ArrayProperty = CastField<FArrayProperty>(PropertyEditor->GetProperty()))
-		{
-			ObjectProperty = CastField<FObjectProperty>(ArrayProperty->Inner);
-		}
-
-		// Only try to add entries if we are dropping on an asset array
-		if (ObjectProperty && InOperation->IsOfType<FAssetDragDropOp>())
-		{
-			bool bHasOnlyValidElements = true;
-			TSharedPtr<FAssetDragDropOp> DragDropOp = StaticCastSharedPtr<FAssetDragDropOp>(InOperation);
-			if (DragDropOp.IsValid())
-			{
-				for (FAssetData AssetData : DragDropOp->GetAssets())
-				{
-					// if the type does not match
-					if (!AssetData.IsInstanceOf(ObjectProperty->PropertyClass))
-					{
-						bHasOnlyValidElements = false;
-						break;
-					}
-				}
-			}
-			return bHasOnlyValidElements;
-		}
-		
-		return false;
-	}
+	FReply OnDragDropTarget(const FGeometry& InGeometry, const FDragDropEvent& InDragDropEvent);
+	bool IsValidAssetDropOp(TSharedPtr<FDragDropOperation> InOperation);
+	bool WillAddValidElements(TSharedPtr<FDragDropOperation> InOperation);
 
 private:
 	TSharedPtr<FPropertyEditor> PropertyEditor;
 };
-
-#undef LOCTEXT_NAMESPACE
