Temporarily disabling the use of property setters and getters in property panels

diff --git a/Engine/Source/Editor/PropertyEditor/Private/PropertyHandleImpl.cpp b/Engine/Source/Editor/PropertyEditor/Private/PropertyHandleImpl.cpp
--- a/Engine/Source/Editor/PropertyEditor/Private/PropertyHandleImpl.cpp
+++ b/Engine/Source/Editor/PropertyEditor/Private/PropertyHandleImpl.cpp
@@ -407,19 +407,12 @@ FPropertyAccess::Result FPropertyValueImpl::ImportText( const TArray<FObjectBase
 			}
 
 			// Set the new value.
-			{
-				EPropertyPortFlags PortFlags = (Flags & EPropertyValueSetFlags::InstanceObjects) != 0 ? PPF_InstanceSubobjects : PPF_None;
-				FPropertyNodeEditStack PropStack(InPropertyNode, CurObject);
-				FPropertyTextUtilities::TextToPropertyHelper(*NewValue, InPropertyNode, NodeProperty, PropStack.GetDirectPropertyAddress(), CurObject, PortFlags);
-				PropStack.CommitChanges();
-			}
+			EPropertyPortFlags PortFlags = (Flags & EPropertyValueSetFlags::InstanceObjects) != 0 ? PPF_InstanceSubobjects : PPF_None;
+			FPropertyTextUtilities::TextToPropertyHelper(*NewValue, InPropertyNode, NodeProperty, Cur, PortFlags);
 
 			// Cache the value of the property after having modified it.
 			FString ValueAfterImport;
-			{
-				FPropertyNodeEditStack PropStack(InPropertyNode, CurObject);
-				FPropertyTextUtilities::PropertyToTextHelper(ValueAfterImport, InPropertyNode, NodeProperty, PropStack.GetDirectPropertyAddress(), CurObject, PPF_None);
-			}
+			FPropertyTextUtilities::PropertyToTextHelper(ValueAfterImport, InPropertyNode, NodeProperty, Cur, PPF_None);
 
 			if (CurObject)
 			{
@@ -3396,20 +3389,20 @@ bool FPropertyHandleInt::Supports( TSharedRef<FPropertyNode> PropertyNode )
 }
 
 template <typename PropertyClass, typename ValueType>
-ValueType GetIntegerValue(FPropertyNodeEditStack& PropStack, FPropertyValueImpl& Implementation)
+ValueType GetIntegerValue(void* PropValue, FPropertyValueImpl& Implementation )
 {
-	check(Implementation.IsPropertyTypeOf(PropertyClass::StaticClass()));
-	return Implementation.GetPropertyValue<PropertyClass>(PropStack.GetDirectPropertyAddress());
+	check( Implementation.IsPropertyTypeOf( PropertyClass::StaticClass() ) );
+	return Implementation.GetPropertyValue<PropertyClass>(PropValue);
 }
 
 FPropertyAccess::Result FPropertyHandleInt::GetValue(int8& OutValue) const
 {
-	FPropertyNodeEditStack PropStack;
-	FPropertyAccess::Result Res = Implementation->GetValueEditStack(PropStack);
+	void* PropValue = nullptr;
+	FPropertyAccess::Result Res = Implementation->GetValueData(PropValue);
 
 	if (Res == FPropertyAccess::Success)
 	{
-		OutValue = GetIntegerValue<FInt8Property, int8>(PropStack, *Implementation);
+		OutValue = GetIntegerValue<FInt8Property, int8>(PropValue, *Implementation);
 	}
 
 	return Res;
@@ -3417,12 +3410,12 @@ FPropertyAccess::Result FPropertyHandleInt::GetValue(int8& OutValue) const
 
 FPropertyAccess::Result FPropertyHandleInt::GetValue(int16& OutValue) const
 {
-	FPropertyNodeEditStack PropStack;
-	FPropertyAccess::Result Res = Implementation->GetValueEditStack(PropStack);
+	void* PropValue = nullptr;
+	FPropertyAccess::Result Res = Implementation->GetValueData(PropValue);
 
 	if (Res == FPropertyAccess::Success)
 	{
-		OutValue = GetIntegerValue<FInt16Property, int16>(PropStack, *Implementation);
+		OutValue = GetIntegerValue<FInt16Property, int16>(PropValue, *Implementation);
 	}
 
 	return Res;
@@ -3430,24 +3423,24 @@ FPropertyAccess::Result FPropertyHandleInt::GetValue(int16& OutValue) const
 
 FPropertyAccess::Result FPropertyHandleInt::GetValue(int32& OutValue) const
 {
-	FPropertyNodeEditStack PropStack;
-	FPropertyAccess::Result Res = Implementation->GetValueEditStack(PropStack);
+	void* PropValue = nullptr;
+	FPropertyAccess::Result Res = Implementation->GetValueData(PropValue);
 
 	if (Res == FPropertyAccess::Success)
 	{
-		OutValue = GetIntegerValue<FIntProperty, int32>(PropStack, *Implementation);
+		OutValue = GetIntegerValue<FIntProperty, int32>(PropValue, *Implementation);
 	}
 
 	return Res;
 }
 FPropertyAccess::Result FPropertyHandleInt::GetValue( int64& OutValue ) const
 {
-	FPropertyNodeEditStack PropStack;
-	FPropertyAccess::Result Res = Implementation->GetValueEditStack(PropStack);
+	void* PropValue = nullptr;
+	FPropertyAccess::Result Res = Implementation->GetValueData( PropValue );
 
 	if( Res == FPropertyAccess::Success )
 	{
-		OutValue = GetIntegerValue<FInt64Property,int64>(PropStack, *Implementation);
+		OutValue = GetIntegerValue<FInt64Property,int64>(PropValue, *Implementation);
 	}
 
 	return Res;
@@ -3455,12 +3448,12 @@ FPropertyAccess::Result FPropertyHandleInt::GetValue( int64& OutValue ) const
 
 FPropertyAccess::Result FPropertyHandleInt::GetValue(uint16& OutValue) const
 {
-	FPropertyNodeEditStack PropStack;
-	FPropertyAccess::Result Res = Implementation->GetValueEditStack(PropStack);
+	void* PropValue = nullptr;
+	FPropertyAccess::Result Res = Implementation->GetValueData(PropValue);
 
 	if (Res == FPropertyAccess::Success)
 	{
-		OutValue = GetIntegerValue<FUInt16Property, uint16>(PropStack, *Implementation);
+		OutValue = GetIntegerValue<FUInt16Property, uint16>(PropValue, *Implementation);
 	}
 
 	return Res;
@@ -3468,12 +3461,12 @@ FPropertyAccess::Result FPropertyHandleInt::GetValue(uint16& OutValue) const
 
 FPropertyAccess::Result FPropertyHandleInt::GetValue(uint32& OutValue) const
 {
-	FPropertyNodeEditStack PropStack;
-	FPropertyAccess::Result Res = Implementation->GetValueEditStack(PropStack);
+	void* PropValue = nullptr;
+	FPropertyAccess::Result Res = Implementation->GetValueData(PropValue);
 
 	if (Res == FPropertyAccess::Success)
 	{
-		OutValue = GetIntegerValue<FUInt32Property, uint32>(PropStack, *Implementation);
+		OutValue = GetIntegerValue<FUInt32Property, uint32>(PropValue, *Implementation);
 	}
 
 	return Res;
@@ -3481,12 +3474,12 @@ FPropertyAccess::Result FPropertyHandleInt::GetValue(uint32& OutValue) const
 
 FPropertyAccess::Result FPropertyHandleInt::GetValue(uint64& OutValue) const
 {
-	FPropertyNodeEditStack PropStack;
-	FPropertyAccess::Result Res = Implementation->GetValueEditStack(PropStack);
+	void* PropValue = nullptr;
+	FPropertyAccess::Result Res = Implementation->GetValueData(PropValue);
 
 	if (Res == FPropertyAccess::Success)
 	{
-		OutValue = GetIntegerValue<FUInt64Property, uint64>(PropStack, *Implementation);
+		OutValue = GetIntegerValue<FUInt64Property, uint64>(PropValue, *Implementation);
 	}
 
 	return Res;
@@ -3593,12 +3586,11 @@ bool FPropertyHandleFloat::Supports( TSharedRef<FPropertyNode> PropertyNode )
 
 FPropertyAccess::Result FPropertyHandleFloat::GetValue( float& OutValue ) const
 {
-	FPropertyNodeEditStack PropStack;
-	FPropertyAccess::Result Res = Implementation->GetValueEditStack(PropStack);
+	void* PropValue = nullptr;
+	FPropertyAccess::Result Res = Implementation->GetValueData( PropValue );
 
 	if( Res == FPropertyAccess::Success )
 	{
-		void* PropValue = PropStack.GetDirectPropertyAddress();
 		OutValue = Implementation->GetPropertyValue<FFloatProperty>(PropValue);
 	}
 
@@ -3632,12 +3624,11 @@ bool FPropertyHandleDouble::Supports( TSharedRef<FPropertyNode> PropertyNode )
 
 FPropertyAccess::Result FPropertyHandleDouble::GetValue( double& OutValue ) const
 {
-	FPropertyNodeEditStack PropStack;
-	FPropertyAccess::Result Res = Implementation->GetValueEditStack(PropStack);
+	void* PropValue = nullptr;
+	FPropertyAccess::Result Res = Implementation->GetValueData( PropValue );
 
 	if (Res == FPropertyAccess::Success)
 	{
-		void* PropValue = PropStack.GetDirectPropertyAddress();
 		OutValue = Implementation->GetPropertyValue<FDoubleProperty>(PropValue);
 	}
 
@@ -3673,12 +3664,11 @@ bool FPropertyHandleBool::Supports( TSharedRef<FPropertyNode> PropertyNode )
 
 FPropertyAccess::Result FPropertyHandleBool::GetValue( bool& OutValue ) const
 {
-	FPropertyNodeEditStack PropStack;
-	FPropertyAccess::Result Res = Implementation->GetValueEditStack(PropStack);
+	void* PropValue = nullptr;
+	FPropertyAccess::Result Res = Implementation->GetValueData( PropValue );
 
 	if( Res == FPropertyAccess::Success )
 	{
-		void* PropValue = PropStack.GetDirectPropertyAddress();
 		OutValue = Implementation->GetPropertyValue<FBoolProperty>(PropValue);
 	}
 
@@ -3719,12 +3709,11 @@ bool FPropertyHandleByte::Supports( TSharedRef<FPropertyNode> PropertyNode )
 
 FPropertyAccess::Result FPropertyHandleByte::GetValue( uint8& OutValue ) const
 {
-	FPropertyNodeEditStack PropStack;
-	FPropertyAccess::Result Res = Implementation->GetValueEditStack(PropStack);
+	void* PropValue = nullptr;
+	FPropertyAccess::Result Res = Implementation->GetValueData( PropValue );
 
 	if( Res == FPropertyAccess::Success )
 	{
-		void* PropValue = PropStack.GetDirectPropertyAddress();
 		TSharedPtr<FPropertyNode> PropertyNodePin = Implementation->GetPropertyNode();
 
 		FProperty* Property = PropertyNodePin->GetProperty();
@@ -3809,12 +3798,11 @@ FPropertyAccess::Result FPropertyHandleString::SetValue( const TCHAR* NewValue,
 
 FPropertyAccess::Result FPropertyHandleString::GetValue( FName& OutValue ) const
 {
-	FPropertyNodeEditStack PropStack;
-	FPropertyAccess::Result Res = Implementation->GetValueEditStack(PropStack);
+	void* PropValue = nullptr;
+	FPropertyAccess::Result Res = Implementation->GetValueData( PropValue );
 
 	if( Res == FPropertyAccess::Success )
 	{
-		void* PropValue = PropStack.GetDirectPropertyAddress();
 		OutValue = Implementation->GetPropertyValue<FNameProperty>(PropValue);
 	}
 
@@ -3848,12 +3836,11 @@ FPropertyAccess::Result FPropertyHandleObject::GetValue( UObject*& OutValue ) co
 
 FPropertyAccess::Result FPropertyHandleObject::GetValue( const UObject*& OutValue ) const
 {
-	FPropertyNodeEditStack PropStack;
-	FPropertyAccess::Result Res = Implementation->GetValueEditStack(PropStack);
+	void* PropValue = nullptr;
+	FPropertyAccess::Result Res = Implementation->GetValueData( PropValue );
 
 	if( Res == FPropertyAccess::Success )
 	{
-		void* PropValue = PropStack.GetDirectPropertyAddress();
 		FProperty* Property = GetProperty();
 
 		if (Property->IsA(FObjectPropertyBase::StaticClass()))
@@ -4246,12 +4233,11 @@ bool FPropertyHandleMixed::Supports(TSharedRef<FPropertyNode> PropertyNode)
 
 FPropertyAccess::Result FPropertyHandleMixed::GetValue(double& OutValue) const
 {
-	FPropertyNodeEditStack PropStack;
-	FPropertyAccess::Result Res = Implementation->GetValueEditStack(PropStack);
+	void* PropValue = nullptr;
+	FPropertyAccess::Result Res = Implementation->GetValueData(PropValue);
 
 	if (Res == FPropertyAccess::Success)
 	{
-		void* PropValue = PropStack.GetDirectPropertyAddress();
 		if(Implementation->IsPropertyTypeOf(FFloatProperty::StaticClass()))
 		{
 			OutValue = Implementation->GetPropertyValue<FFloatProperty>(PropValue);
diff --git a/Engine/Source/Editor/PropertyEditor/Private/PropertyNode.cpp b/Engine/Source/Editor/PropertyEditor/Private/PropertyNode.cpp
--- a/Engine/Source/Editor/PropertyEditor/Private/PropertyNode.cpp
+++ b/Engine/Source/Editor/PropertyEditor/Private/PropertyNode.cpp
@@ -777,7 +777,15 @@ FPropertyAccess::Result FPropertyNodeEditStack::InitializeInternal(const FProper
 		}
 		else if (Property->GetOwner<FProperty>() == ParentProperty) // TArrays, TMaps and TSets
 		{
-			MemoryStack.Add(FMemoryFrame(Property, (uint8*)ParentProperty->GetValueAddressAtIndex_Direct(Property, (void*)MemoryStack.Last().Memory, InNode->GetArrayIndex())));
+			uint8* ItemAddress = (uint8*)ParentProperty->GetValueAddressAtIndex_Direct(Property, (void*)MemoryStack.Last().Memory, InNode->GetArrayIndex());
+			if (ItemAddress)
+			{
+				MemoryStack.Add(FMemoryFrame(Property, ItemAddress));
+			}
+			else
+			{
+				return FPropertyAccess::Fail;
+			}
 		}
 		else
 		{
@@ -879,22 +887,16 @@ FPropertyNodeEditStack::~FPropertyNodeEditStack()
 
 FPropertyAccess::Result FPropertyNode::GetPropertyValueString(FString& OutString, const bool bAllowAlternateDisplayValue, EPropertyPortFlags PortFlags) const
 {
-	FPropertyAccess::Result Result = FPropertyAccess::Fail;
-	const FProperty* PropertyPtr = GetProperty();
-	// Check for bogus data
-	if (PropertyPtr != nullptr && GetParentNode() != nullptr)
+	uint8* ValueAddress = nullptr;
+	FPropertyAccess::Result Result = GetSingleReadAddress(ValueAddress);
+
+	if (ValueAddress != nullptr)
 	{
-		UObject* ValueContainer = nullptr;		
-		FPropertyNodeEditStack PropStack;
-		Result = GetSingleObject(ValueContainer);
-		if (Result == FPropertyAccess::Success)
-		{
-			Result = PropStack.Initialize(this, ValueContainer);
-		}
+		const FProperty* PropertyPtr = GetProperty();
 
-		if (Result == FPropertyAccess::Success)
+		// Check for bogus data
+		if (PropertyPtr != nullptr && GetParentNode() != nullptr)
 		{
-			uint8* ValueAddress = PropStack.GetDirectPropertyAddress();
 			FPropertyTextUtilities::PropertyToTextHelper(OutString, this, PropertyPtr, ValueAddress, nullptr, PortFlags);
 
 			UEnum* Enum = nullptr;
@@ -930,6 +932,10 @@ FPropertyAccess::Result FPropertyNode::GetPropertyValueString(FString& OutString
 				}
 			}
 		}
+		else
+		{
+			Result = FPropertyAccess::Fail;
+		}
 	}
 
 	return Result;
@@ -937,22 +943,14 @@ FPropertyAccess::Result FPropertyNode::GetPropertyValueString(FString& OutString
 
 FPropertyAccess::Result FPropertyNode::GetPropertyValueText(FText& OutText, const bool bAllowAlternateDisplayValue) const
 {
-	FPropertyAccess::Result Result = FPropertyAccess::Fail;
-	const FProperty* PropertyPtr = GetProperty();
-	if (PropertyPtr)
-	{
-		UObject* ValueContainer = nullptr;
-		FPropertyNodeEditStack PropStack;
-		Result = GetSingleObject(ValueContainer);
-		if (Result == FPropertyAccess::Success)
-		{
-			Result = PropStack.Initialize(this, ValueContainer);
-		}
+	uint8* ValueAddress = nullptr;
+	FPropertyAccess::Result Result = GetSingleReadAddress(ValueAddress);
 
-		if (Result == FPropertyAccess::Success)
+	if (ValueAddress != nullptr)
+	{
+		const FProperty* PropertyPtr = GetProperty();
+		if (PropertyPtr)
 		{
-			uint8* ValueAddress = PropStack.GetDirectPropertyAddress();
-
 			if (PropertyPtr->IsA(FTextProperty::StaticClass()))
 			{
 				OutText = CastField<FTextProperty>(PropertyPtr)->GetPropertyValue(ValueAddress);
@@ -993,6 +991,10 @@ FPropertyAccess::Result FPropertyNode::GetPropertyValueText(FText& OutText, cons
 				}
 			}
 		}
+		else
+		{
+			Result = FPropertyAccess::Fail;
+		}
 	}
 
 	return Result;
