Fix combo box has always triggered a value change event when it was opened.

diff --git a/Engine/Source/Editor/PropertyEditor/Private/UserInterface/PropertyEditor/SPropertyEditorCombo.cpp b/Engine/Source/Editor/PropertyEditor/Private/UserInterface/PropertyEditor/SPropertyEditorCombo.cpp
--- a/Engine/Source/Editor/PropertyEditor/Private/UserInterface/PropertyEditor/SPropertyEditorCombo.cpp
+++ b/Engine/Source/Editor/PropertyEditor/Private/UserInterface/PropertyEditor/SPropertyEditorCombo.cpp
@@ -143,6 +143,27 @@ void SPropertyEditorCombo::GenerateComboBoxStrings( TArray< TSharedPtr<FString>
 
 	bUsesAlternateDisplayValues = ComboArgs.PropertyHandle->GeneratePossibleValues(OutComboBoxStrings, BasicTooltips, OutRestrictedItems);
 
+	// If we regenerate the entries, let's make sure that the currently selected item has the same shared pointer as
+	// the newly generated item with the same value, so that the generation of elements won't immediately result in a
+	// value changed event (i.e. at every single `OnComboOpening`).
+	if (ComboBox)
+	{
+		if (const TSharedPtr<FString>& SelectedItem = ComboBox->GetSelectedItem())
+		{
+			for (TSharedPtr<FString>& Item : OutComboBoxStrings)
+			{
+				if (Item)
+				{
+					if (*SelectedItem.Get() == *Item.Get())
+					{
+						Item = SelectedItem;
+						break;
+					}
+				}
+			}
+		}
+	}
+
 	// For enums, look for rich tooltip information
 	if(ComboArgs.PropertyHandle.IsValid())
 	{
