[NO NON-PII COMMIT MESSAGE]

diff --git a/Engine/Source/Editor/PropertyEditor/Private/UserInterface/PropertyEditor/SPropertyEditorArray.cpp b/Engine/Source/Editor/PropertyEditor/Private/UserInterface/PropertyEditor/SPropertyEditorArray.cpp
--- a/Engine/Source/Editor/PropertyEditor/Private/UserInterface/PropertyEditor/SPropertyEditorArray.cpp
+++ b/Engine/Source/Editor/PropertyEditor/Private/UserInterface/PropertyEditor/SPropertyEditorArray.cpp
@@ -65,7 +65,16 @@ FText SPropertyEditorArray::GetArrayTextValue() const
 	{
 		return LOCTEXT("MultipleValues", "Multiple Values");
 	}
-	return FText::Format(LOCTEXT("NumArrayItemsFmt", "{0} Array elements"), FText::AsNumber(PropertyEditor->GetPropertyNode()->GetNumChildNodes()));
+
+	const int32 NumChildNodes = PropertyEditor->GetPropertyNode()->GetNumChildNodes();
+	if (NumChildNodes <= 1)
+	{
+		return FText::Format(LOCTEXT("SingleArrayItemFmt", "{0} Array element"), FText::AsNumber(NumChildNodes));
+	}
+	else
+	{
+		return FText::Format(LOCTEXT("NumArrayItemsFmt", "{0} Array elements"), FText::AsNumber(NumChildNodes));
+	}
 }
 
 bool SPropertyEditorArray::CanEdit() const
