Allow for bool functions in edit conditions

diff --git a/Engine/Source/Editor/PropertyEditor/Private/EditConditionContext.cpp b/Engine/Source/Editor/PropertyEditor/Private/EditConditionContext.cpp
--- a/Engine/Source/Editor/PropertyEditor/Private/EditConditionContext.cpp
+++ b/Engine/Source/Editor/PropertyEditor/Private/EditConditionContext.cpp
@@ -3,8 +3,10 @@
 #include "EditConditionContext.h"
 #include "EditConditionParser.h"
 
+#include "ObjectPropertyNode.h"
 #include "PropertyNode.h"
 #include "PropertyEditorHelpers.h"
+#include "PropertyPathHelpers.h"
 
 DEFINE_LOG_CATEGORY(LogEditCondition);
 
@@ -61,8 +63,45 @@ const FBoolProperty* FEditConditionContext::GetSingleBoolProperty(const TSharedP
 	return BoolProperty;
 }
 
+const TWeakObjectPtr<UFunction> FEditConditionContext::GetFunction(const FString& FieldName) const
+{
+	TSharedPtr<FPropertyNode> PinnedNode = PropertyNode.Pin();
+	if (!PinnedNode.IsValid())
+	{
+		return nullptr;
+	}
+
+	const FProperty* Property = PinnedNode->GetProperty();
+	if (Property == nullptr)
+	{
+		return nullptr;
+	}
+
+	TWeakObjectPtr<UFunction> Function = FindUField<UFunction>(Property->GetOwnerStruct(), *FieldName);
+
+	if (Function == nullptr && FieldName.Contains(TEXT(".")))
+	{
+		// Function not found in struct, try to see if this is a static function
+		Function = FindObject<UFunction>(nullptr, *FieldName, true);
+
+		if (Function.IsValid() && !Function->HasAnyFunctionFlags(EFunctionFlags::FUNC_Static))
+		{
+			Function = nullptr;
+		}
+	}
+
+	return Function;
+}
+
 static TSet<TPair<FName, FString>> AlreadyLogged;
 
+/**
+ * Attempt to property within the node & log if not found.
+ * 
+ * @param PropertyNode		Property owning the field to search for, ex: a UObject
+ * @param FieldNam			Name of field to search for
+ * @return					Field for given fieldname if found, nullptr otherwise 
+ */
 template<typename T>
 const T* FindTypedField(const TSharedPtr<FPropertyNode>& PropertyNode, const FString& FieldName)
 {
@@ -133,7 +172,7 @@ static const uint8* GetPropertyValuePtr(const FProperty* Property, const TShared
 	return ValuePtr;
 }
 
-TOptional<bool> FEditConditionContext::GetBoolValue(const FString& PropertyName) const
+TOptional<bool> FEditConditionContext::GetBoolValue(const FString& PropertyName, TWeakObjectPtr<UFunction> CachedFunction) const
 {
 	TSharedPtr<FPropertyNode> PinnedNode = PropertyNode.Pin();
 	if (!PinnedNode.IsValid())
@@ -141,6 +180,72 @@ TOptional<bool> FEditConditionContext::GetBoolValue(const FString& PropertyName)
 		return TOptional<bool>();
 	}
 
+	if (const UFunction* Function = CachedFunction.Get())
+	{
+		if (CastField<FBoolProperty>(Function->GetReturnProperty()) == nullptr)
+		{
+			// Function return type not bool, return undefined
+			return TOptional<bool>();
+		}
+
+		// Check for external static function references 
+		if (Function->HasAnyFunctionFlags(EFunctionFlags::FUNC_Static))
+		{
+			FString StaticFunctionName = {};
+			UObject* EditConditionExpressionCDO = Function->GetOuterUClass()->GetDefaultObject();
+			Function->GetName(StaticFunctionName);
+
+			{
+				FEditorScriptExecutionGuard ScriptExecutionGuard;
+				FCachedPropertyPath Path(StaticFunctionName);
+				bool bResult = true;
+				if (PropertyPathHelpers::GetPropertyValue(EditConditionExpressionCDO, Path, bResult))
+				{
+					return bResult;
+				}
+				else
+				{
+					// Execution failed, return undefined
+					return TOptional<bool>();
+				}
+			}
+		}
+
+		// We might be selecting multiple objects, account for that
+		TArray<UObject*> OutObjects;
+		FObjectPropertyNode* ObjectNode = PinnedNode->FindObjectItemParent();
+		if (ObjectNode)
+		{
+			int32 NumObjects = ObjectNode->GetNumObjects();
+			for (int32 ObjectIndex = 0; ObjectIndex < NumObjects; ++ObjectIndex)
+			{
+				FEditorScriptExecutionGuard ScriptExecutionGuard;
+
+				UObject* FunctionTarget = ObjectNode->GetUObject(ObjectIndex);
+				FCachedPropertyPath FunctionPath(PropertyName);
+				bool bResult = true;
+				if (PropertyPathHelpers::GetPropertyValue(FunctionTarget, FunctionPath, bResult))
+				{
+					if (!bResult)
+					{
+						return false;
+					}
+				}
+				else
+				{
+					// Execution failed, return undefined
+					return TOptional<bool>();
+				}
+			}
+
+			// Executed our target function over relevant objects with no condtion failures, return true
+			if (NumObjects > 0)
+			{
+				return true;
+			}
+		}
+	}
+
 	const FBoolProperty* BoolProperty = FindTypedField<FBoolProperty>(PinnedNode, PropertyName);
 	if (BoolProperty == nullptr)
 	{
@@ -183,7 +288,7 @@ TOptional<bool> FEditConditionContext::GetBoolValue(const FString& PropertyName)
 	return Result;
 }
 
-TOptional<int64> FEditConditionContext::GetIntegerValue(const FString& PropertyName) const
+TOptional<int64> FEditConditionContext::GetIntegerValue(const FString& PropertyName, TWeakObjectPtr<UFunction> CachedFunction) const
 {
 	TSharedPtr<FPropertyNode> PinnedNode = PropertyNode.Pin();
 	if (!PinnedNode.IsValid())
@@ -191,6 +296,12 @@ TOptional<int64> FEditConditionContext::GetIntegerValue(const FString& PropertyN
 		return TOptional<int64>();
 	}
 
+	if (const UFunction* Function = CachedFunction.Get())
+	{ 
+		// EditConditions Currently only support bool, see: UE-175891
+		return TOptional<int64>();
+	}
+
 	const FProperty* Property = FindTypedField<FProperty>(PinnedNode, PropertyName);
 	const FNumericProperty* NumericProperty = CastField<FNumericProperty>(Property);
 
@@ -244,7 +355,7 @@ TOptional<int64> FEditConditionContext::GetIntegerValue(const FString& PropertyN
 	return Result;
 }
 
-TOptional<double> FEditConditionContext::GetNumericValue(const FString& PropertyName) const
+TOptional<double> FEditConditionContext::GetNumericValue(const FString& PropertyName, TWeakObjectPtr<UFunction> CachedFunction) const
 {
 	TSharedPtr<FPropertyNode> PinnedNode = PropertyNode.Pin();
 	if (!PinnedNode.IsValid())
@@ -252,6 +363,12 @@ TOptional<double> FEditConditionContext::GetNumericValue(const FString& Property
 		return TOptional<double>();
 	}
 
+	if (const UFunction* Function = CachedFunction.Get())
+	{
+		// EditConditions Currently only support bool, see: UE-175891
+		return TOptional<double>();
+	}
+
 	const FNumericProperty* NumericProperty = FindTypedField<FNumericProperty>(PinnedNode, PropertyName);
 	if (NumericProperty == nullptr)
 	{
@@ -304,7 +421,7 @@ TOptional<double> FEditConditionContext::GetNumericValue(const FString& Property
 	return Result;
 }
 
-TOptional<FString> FEditConditionContext::GetEnumValue(const FString& PropertyName) const
+TOptional<FString> FEditConditionContext::GetEnumValue(const FString& PropertyName, TWeakObjectPtr<UFunction> CachedFunction) const
 {
 	TSharedPtr<FPropertyNode> PinnedNode = PropertyNode.Pin();
 	if (!PinnedNode.IsValid())
@@ -312,6 +429,12 @@ TOptional<FString> FEditConditionContext::GetEnumValue(const FString& PropertyNa
 		return TOptional<FString>();
 	}
 
+	if (const UFunction* Function = CachedFunction.Get())
+	{
+		// EditConditions Currently only support bool, see: UE-175891
+		return TOptional<FString>();
+	}
+
 	const FProperty* Property = FindTypedField<FProperty>(PinnedNode, PropertyName);
 	if (Property == nullptr)
 	{
@@ -379,7 +502,7 @@ TOptional<FString> FEditConditionContext::GetEnumValue(const FString& PropertyNa
 	return EnumType->GetNameStringByValue(Result.GetValue());
 }
 
-TOptional<UObject*> FEditConditionContext::GetPointerValue(const FString& PropertyName) const
+TOptional<UObject*> FEditConditionContext::GetPointerValue(const FString& PropertyName, TWeakObjectPtr<UFunction> CachedFunction) const
 {
 	TSharedPtr<FPropertyNode> PinnedNode = PropertyNode.Pin();
 	if (!PinnedNode.IsValid())
@@ -387,6 +510,12 @@ TOptional<UObject*> FEditConditionContext::GetPointerValue(const FString& Proper
 		return TOptional<UObject*>();
 	}
 
+	if (const UFunction* Function = CachedFunction.Get())
+	{
+		// EditConditions Currently only support bool, see: UE-175891
+		return TOptional<UObject*>();
+	}
+
 	const FProperty* Property = FindTypedField<FProperty>(PinnedNode, PropertyName);
 	if (Property == nullptr)
 	{
@@ -435,7 +564,7 @@ TOptional<UObject*> FEditConditionContext::GetPointerValue(const FString& Proper
 	return Result;
 }
 
-TOptional<FString> FEditConditionContext::GetTypeName(const FString& PropertyName) const
+TOptional<FString> FEditConditionContext::GetTypeName(const FString& PropertyName, TWeakObjectPtr<UFunction> CachedFunction) const
 {
 	TSharedPtr<FPropertyNode> PinnedNode = PropertyNode.Pin();
 	if (!PinnedNode.IsValid())
@@ -443,26 +572,37 @@ TOptional<FString> FEditConditionContext::GetTypeName(const FString& PropertyNam
 		return TOptional<FString>();
 	}
 
-	const FProperty* Property = FindTypedField<FProperty>(PinnedNode, PropertyName);
-	if (Property == nullptr)
+	auto GetPropertyName = [](const FProperty* Property) -> TOptional<FString>
 	{
-		return TOptional<FString>();
-	}
+		if (Property == nullptr)
+		{
+			return TOptional<FString>();
+		}
 
-	if (const FEnumProperty* EnumProperty = CastField<FEnumProperty>(Property))
-	{
-		return EnumProperty->GetEnum()->GetName();
-	}
-	else if (const FByteProperty* ByteProperty = CastField<FByteProperty>(Property))
-	{
-		const UEnum* EnumType = ByteProperty->GetIntPropertyEnum();
-		if (EnumType != nullptr)
+		if (const FEnumProperty* EnumProperty = CastField<FEnumProperty>(Property))
+		{
+			return EnumProperty->GetEnum()->GetName();
+		}
+		else if (const FByteProperty* ByteProperty = CastField<FByteProperty>(Property))
 		{
-			return EnumType->GetName();
+			const UEnum* EnumType = ByteProperty->GetIntPropertyEnum();
+			if (EnumType != nullptr)
+			{
+				return EnumType->GetName();
+			}
 		}
+
+		return Property->GetCPPType();
+	};
+
+	if (const UFunction* Function = CachedFunction.Get())
+	{
+		return GetPropertyName(Function->GetReturnProperty());
 	}
 
-	return Property->GetCPPType();
+	const FProperty* Property = FindTypedField<FProperty>(PinnedNode, PropertyName);
+
+	return GetPropertyName(Property);
 }
 
 TOptional<int64> FEditConditionContext::GetIntegerValueOfEnum(const FString& EnumTypeName, const FString& MemberName) const
diff --git a/Engine/Source/Editor/PropertyEditor/Private/EditConditionContext.h b/Engine/Source/Editor/PropertyEditor/Private/EditConditionContext.h
--- a/Engine/Source/Editor/PropertyEditor/Private/EditConditionContext.h
+++ b/Engine/Source/Editor/PropertyEditor/Private/EditConditionContext.h
@@ -3,21 +3,34 @@
 #pragma once
 
 #include "CoreMinimal.h"
+#include "UObject/WeakObjectPtr.h"
 
 DECLARE_LOG_CATEGORY_EXTERN(LogEditCondition, Log, All);
 
+class UFunction;
+
 class IEditConditionContext
 {
 public:
 	virtual FName GetContextName() const = 0;
 
-	virtual TOptional<bool> GetBoolValue(const FString& PropertyName) const = 0;
-	virtual TOptional<int64> GetIntegerValue(const FString& PropertyName) const = 0;
-	virtual TOptional<double> GetNumericValue(const FString& PropertyName) const = 0;
-	virtual TOptional<FString> GetEnumValue(const FString& PropertyName) const = 0;
-	virtual TOptional<UObject*> GetPointerValue(const FString& PropertyName) const = 0;
-	virtual TOptional<FString> GetTypeName(const FString& PropertyName) const = 0;
+	virtual TOptional<bool> GetBoolValue(const FString& PropertyName, TWeakObjectPtr<UFunction> CachedFunction = nullptr) const = 0;
+	virtual TOptional<int64> GetIntegerValue(const FString& PropertyName, TWeakObjectPtr<UFunction> CachedFunction = nullptr) const = 0;
+	virtual TOptional<double> GetNumericValue(const FString& PropertyName, TWeakObjectPtr<UFunction> CachedFunction = nullptr) const = 0;
+	virtual TOptional<FString> GetEnumValue(const FString& PropertyName, TWeakObjectPtr<UFunction> CachedFunction = nullptr) const = 0;
+	virtual TOptional<UObject*> GetPointerValue(const FString& PropertyName, TWeakObjectPtr<UFunction> CachedFunction = nullptr) const = 0;
+	virtual TOptional<FString> GetTypeName(const FString& PropertyName, TWeakObjectPtr<UFunction> CachedFunction = nullptr) const = 0;
 	virtual TOptional<int64> GetIntegerValueOfEnum(const FString& EnumType, const FString& EnumValue) const = 0;
+
+	/**
+	 * Get field as a function, avoiding multiple function searches
+	 *
+	 * @TODO: Get property in general, not just for functions. See: UE-175891.
+	 *
+	 * @param FieldName Name of the field to try to interpret as a function
+	 * @return Function if found or already cached, else
+	 */
+	virtual const TWeakObjectPtr<UFunction> GetFunction(const FString& FieldName) const = 0;
 };
 
 class FProperty;
@@ -33,14 +46,16 @@ public:
 
 	virtual FName GetContextName() const override;
 
-	virtual TOptional<bool> GetBoolValue(const FString& PropertyName) const override;
-	virtual TOptional<int64> GetIntegerValue(const FString& PropertyName) const override;
-	virtual TOptional<double> GetNumericValue(const FString& PropertyName) const override;
-	virtual TOptional<FString> GetEnumValue(const FString& PropertyName) const override;
-	virtual TOptional<UObject*> GetPointerValue(const FString& PropertyName) const override;
-	virtual TOptional<FString> GetTypeName(const FString& PropertyName) const override;
+	virtual TOptional<bool> GetBoolValue(const FString& PropertyName, TWeakObjectPtr<UFunction> CachedFunction = nullptr) const override;
+	virtual TOptional<int64> GetIntegerValue(const FString& PropertyName, TWeakObjectPtr<UFunction> CachedFunction = nullptr) const override;
+	virtual TOptional<double> GetNumericValue(const FString& PropertyName, TWeakObjectPtr<UFunction> CachedFunction = nullptr) const override;
+	virtual TOptional<FString> GetEnumValue(const FString& PropertyName, TWeakObjectPtr<UFunction> CachedFunction = nullptr) const override;
+	virtual TOptional<UObject*> GetPointerValue(const FString& PropertyName, TWeakObjectPtr<UFunction> CachedFunction = nullptr) const override;
+	virtual TOptional<FString> GetTypeName(const FString& PropertyName, TWeakObjectPtr<UFunction> CachedFunction = nullptr) const override;
 	virtual TOptional<int64> GetIntegerValueOfEnum(const FString& EnumType, const FString& EnumValue) const override;
 
+	virtual const TWeakObjectPtr<UFunction> GetFunction(const FString& FieldName) const override;
+
 	/**
 	 * Fetch the single boolean property referenced.
 	 * Returns nullptr if more than one property is referenced.
diff --git a/Engine/Source/Editor/PropertyEditor/Private/EditConditionParser.cpp b/Engine/Source/Editor/PropertyEditor/Private/EditConditionParser.cpp
--- a/Engine/Source/Editor/PropertyEditor/Private/EditConditionParser.cpp
+++ b/Engine/Source/Editor/PropertyEditor/Private/EditConditionParser.cpp
@@ -216,13 +216,13 @@ TOptional<T> GetValueInternal(const IEditConditionContext& Context, const FStrin
 template<>
 TOptional<bool> GetValueInternal<bool>(const IEditConditionContext& Context, const FString& PropertyName)
 {
-	return Context.GetBoolValue(PropertyName);
+	return Context.GetBoolValue(PropertyName, Context.GetFunction(PropertyName));
 }
 
 template<>
 TOptional<double> GetValueInternal<double>(const IEditConditionContext& Context, const FString& PropertyName)
 {
-	return Context.GetNumericValue(PropertyName);
+	return Context.GetNumericValue(PropertyName, Context.GetFunction(PropertyName));
 }
 
 template <typename T>
@@ -287,7 +287,7 @@ static FExpressionResult ApplyBitwiseAnd(const EditConditionParserTokens::FPrope
 		return MakeError(FText::Format(LOCTEXT("InvalidEnumValue", "EditCondition attempted to use an invalid enum value \"{0}::{1}\"."), FText::FromString(Enum.Type), FText::FromString(Enum.Value)));
 	}
 
-	TOptional<int64> PropertyValue = Context.GetIntegerValue(Property.PropertyName);
+	TOptional<int64> PropertyValue = Context.GetIntegerValue(Property.PropertyName, Context.GetFunction(Property.PropertyName));
 	if (!PropertyValue.IsSet())
 	{
 		return MakeError(FText::Format(LOCTEXT("InvalidOperand", "EditCondition attempted to use an invalid operand \"{0}\"."), FText::FromString(Property.PropertyName)));
@@ -298,13 +298,15 @@ static FExpressionResult ApplyBitwiseAnd(const EditConditionParserTokens::FPrope
 
 static FExpressionResult ApplyPropertyIsNull(const EditConditionParserTokens::FPropertyToken& Property, const IEditConditionContext& Context, bool bNegate)
 {
-	TOptional<FString> TypeName = Context.GetTypeName(Property.PropertyName);
+	TWeakObjectPtr<UFunction> CachedFunction = Context.GetFunction(Property.PropertyName);
+
+	TOptional<FString> TypeName = Context.GetTypeName(Property.PropertyName, CachedFunction);
 	if (!TypeName.IsSet())
 	{
 		return MakeError(FText::Format(LOCTEXT("InvalidOperand", "EditCondition attempted to use an invalid operand \"{0}\"."), FText::FromString(Property.PropertyName)));
 	}
 
-	TOptional<UObject*> Ptr = Context.GetPointerValue(Property.PropertyName);
+	TOptional<UObject*> Ptr = Context.GetPointerValue(Property.PropertyName, CachedFunction);
 	if (!Ptr.IsSet())
 	{
 		return MakeError(FText::Format(LOCTEXT("InvalidOperand", "EditCondition attempted to use an invalid operand \"{0}\"."), FText::FromString(Property.PropertyName)));
@@ -321,16 +323,19 @@ static FExpressionResult ApplyPropertyIsNull(const EditConditionParserTokens::FP
 
 static FExpressionResult ApplyPropertiesEqual(const EditConditionParserTokens::FPropertyToken& A, const EditConditionParserTokens::FPropertyToken& B, const IEditConditionContext& Context, bool bNegate)
 {
-	TOptional<UObject*> PtrA = Context.GetPointerValue(A.PropertyName);
-	TOptional<UObject*> PtrB = Context.GetPointerValue(B.PropertyName);
+	TWeakObjectPtr<UFunction> CachedFunctionA = Context.GetFunction(A.PropertyName);
+	TWeakObjectPtr<UFunction> CachedFunctionB = Context.GetFunction(A.PropertyName);
+
+	TOptional<UObject*> PtrA = Context.GetPointerValue(A.PropertyName, CachedFunctionA);
+	TOptional<UObject*> PtrB = Context.GetPointerValue(B.PropertyName, CachedFunctionB);
 	if (PtrA.IsSet() && PtrB.IsSet())
 	{
 		const bool bAreEqual = PtrA.GetValue() == PtrB.GetValue();
 		return MakeValue(bNegate ? !bAreEqual : bAreEqual);
 	}
 
-	TOptional<FString> TypeNameA = Context.GetTypeName(A.PropertyName);
-	TOptional<FString> TypeNameB = Context.GetTypeName(B.PropertyName);
+	TOptional<FString> TypeNameA = Context.GetTypeName(A.PropertyName, CachedFunctionA);
+	TOptional<FString> TypeNameB = Context.GetTypeName(B.PropertyName, CachedFunctionB);
 	if (!TypeNameA.IsSet())
 	{
 		return MakeError(FText::Format(LOCTEXT("InvalidOperand", "EditCondition attempted to use an invalid operand \"{0}\"."), FText::FromString(A.PropertyName)));
@@ -346,32 +351,28 @@ static FExpressionResult ApplyPropertiesEqual(const EditConditionParserTokens::F
 		return MakeError(FText::Format(LOCTEXT("OperandTypeMismatch", "EditCondition attempted to compare operands of different types: \"{0}\" and \"{1}\"."), FText::FromString(A.PropertyName), FText::FromString(B.PropertyName)));
 	}
 
-	TOptional<bool> bEqual;
-
-	TOptional<bool> BoolA = Context.GetBoolValue(A.PropertyName);
-	TOptional<bool> BoolB = Context.GetBoolValue(B.PropertyName);
+	TOptional<bool> BoolA = Context.GetBoolValue(A.PropertyName, CachedFunctionA);
+	TOptional<bool> BoolB = Context.GetBoolValue(B.PropertyName, CachedFunctionB);
 	if (BoolA.IsSet() && BoolB.IsSet())
 	{
-		bEqual = BoolA.GetValue() == BoolB.GetValue();
+		const bool bAreEqual = BoolA.GetValue() == BoolB.GetValue();
+		return MakeValue(bNegate ? !bAreEqual : bAreEqual);
 	}
 
-	TOptional<double> DoubleA = Context.GetNumericValue(A.PropertyName);
-	TOptional<double> DoubleB = Context.GetNumericValue(B.PropertyName);
+	TOptional<double> DoubleA = Context.GetNumericValue(A.PropertyName, CachedFunctionA);
+	TOptional<double> DoubleB = Context.GetNumericValue(B.PropertyName, CachedFunctionB);
 	if (DoubleA.IsSet() && DoubleB.IsSet())
 	{
-		bEqual = DoubleA.GetValue() == DoubleB.GetValue();
+		const bool bAreEqual = DoubleA.GetValue() == DoubleB.GetValue();
+		return MakeValue(bNegate ? !bAreEqual : bAreEqual);
 	}
 
-	TOptional<FString> EnumA = Context.GetEnumValue(A.PropertyName);
-	TOptional<FString> EnumB = Context.GetEnumValue(B.PropertyName);
+	TOptional<FString> EnumA = Context.GetEnumValue(A.PropertyName, CachedFunctionA);
+	TOptional<FString> EnumB = Context.GetEnumValue(B.PropertyName, CachedFunctionB);
 	if (EnumA.IsSet() && EnumB.IsSet())
 	{
-		bEqual = EnumA.GetValue() == EnumB.GetValue();
-	}
-
-	if (bEqual.IsSet())
-	{
-		return MakeValue(bNegate ? !bEqual.GetValue() : bEqual.GetValue());
+		const bool bAreEqual = EnumA.GetValue() == EnumB.GetValue();
+		return MakeValue(bNegate ? !bAreEqual : bAreEqual);
 	}
 
 	return MakeError(FText::Format(LOCTEXT("OperandTypeMismatch", "EditCondition attempted to compare operands of different types: \"{0}\" and \"{1}\"."), FText::FromString(A.PropertyName), FText::FromString(B.PropertyName)));
@@ -646,7 +647,9 @@ void CreateNumberOperators(TOperatorJumpTable<IEditConditionContext>& OperatorJu
 
 static FExpressionResult EnumPropertyEquals(const EditConditionParserTokens::FEnumToken& Enum, const EditConditionParserTokens::FPropertyToken& Property, const IEditConditionContext& Context, bool bNegate)
 {
-	TOptional<FString> TypeName = Context.GetTypeName(Property.PropertyName);
+	TWeakObjectPtr<UFunction> CachedFunction = Context.GetFunction(Property.PropertyName);
+
+	TOptional<FString> TypeName = Context.GetTypeName(Property.PropertyName, CachedFunction);
 	if (!TypeName.IsSet())
 	{
 		return MakeError(FText::Format(LOCTEXT("InvalidOperand_Type", "EditCondition attempted to use an invalid operand \"{0}\" (type error)."), FText::FromString(Property.PropertyName)));
@@ -657,7 +660,7 @@ static FExpressionResult EnumPropertyEquals(const EditConditionParserTokens::FEn
 		return MakeError(FText::Format(LOCTEXT("OperandTypeMismatch", "EditCondition attempted to compare operands of different types: \"{0}\" and \"{1}\"."), FText::FromString(Property.PropertyName), FText::FromString(Enum.Type + TEXT("::") + Enum.Value)));
 	}
 
-	TOptional<FString> ValueProp = Context.GetEnumValue(Property.PropertyName);
+	TOptional<FString> ValueProp = Context.GetEnumValue(Property.PropertyName, CachedFunction);
 	if (!ValueProp.IsSet())
 	{
 		return MakeError(FText::Format(LOCTEXT("InvalidOperand_Value", "EditCondition attempted to use an invalid operand \"{0}\" (value error)."), FText::FromString(Property.PropertyName)));
@@ -795,7 +798,7 @@ TValueOrError<bool, FText> FEditConditionParser::Evaluate(const FEditConditionEx
 		const FPropertyToken* PropertyResult = Result.GetValue().Cast<FPropertyToken>();
 		if (PropertyResult != nullptr)
 		{
-			TOptional<bool> PropertyValue = Context.GetBoolValue(PropertyResult->PropertyName);
+			TOptional<bool> PropertyValue = Context.GetBoolValue(PropertyResult->PropertyName, Context.GetFunction(PropertyResult->PropertyName));
 			if (PropertyValue.IsSet())
 			{
 				return MakeValue(PropertyValue.GetValue());
@@ -821,7 +824,7 @@ TSharedPtr<FEditConditionExpression> FEditConditionParser::Parse(const FString&
 		CompileResultType CompileResult = ExpressionParser::Compile(LexResult.StealValue(), ExpressionGrammar);
 		if (CompileResult.IsValid())
 		{
-			return TSharedPtr<FEditConditionExpression>(new FEditConditionExpression(CompileResult.StealValue()));
+			return MakeShared<FEditConditionExpression>(CompileResult.StealValue());
 		}
 		else
 		{
diff --git a/Engine/Source/Editor/PropertyEditor/Private/EditConditionParserTests.cpp b/Engine/Source/Editor/PropertyEditor/Private/EditConditionParserTests.cpp
--- a/Engine/Source/Editor/PropertyEditor/Private/EditConditionParserTests.cpp
+++ b/Engine/Source/Editor/PropertyEditor/Private/EditConditionParserTests.cpp
@@ -11,6 +11,106 @@ UEditConditionTestObject::UEditConditionTestObject(const FObjectInitializer& Obj
 {
 }
 
+void UEditConditionTestObject::VoidFunction() const
+{
+	ensureAlwaysMsgf(false, TEXT("Error. Should not be executing a void function in edit conditions."));
+}
+
+bool UEditConditionTestObject::GetBoolFunction() const
+{
+	return BoolProperty;
+}
+
+EditConditionTestEnum UEditConditionTestObject::GetEnumFunction() const
+{
+	return EnumProperty;
+}
+
+TEnumAsByte<EditConditionByteEnum> UEditConditionTestObject::GetByteEnumFunction() const
+{
+	return ByteEnumProperty;
+}
+
+double UEditConditionTestObject::GetDoubleFunction() const
+{
+	return DoubleProperty;
+}
+
+int32 UEditConditionTestObject::GetIntegerFunction() const
+{
+	return IntegerProperty;
+}
+
+uint8 UEditConditionTestObject::GetUintBitfieldFunction() const
+{
+	return UintBitfieldProperty;
+}
+
+UObject* UEditConditionTestObject::GetUObjectPtrFunction() const
+{
+	return UObjectPtr;
+}
+
+TSoftClassPtr<UObject> UEditConditionTestObject::GetSoftClassPtrFunction() const
+{
+	return SoftClassPtr;
+}
+
+TWeakObjectPtr<UObject> UEditConditionTestObject::GetWeakObjectPtrFunction() const
+{
+	return WeakObjectPtr;
+}
+
+void UEditConditionTestObject::StaticVoidFunction() 
+{
+	ensureAlwaysMsgf(false, TEXT("Error. Should not be executing a void function in edit conditions."));
+}
+
+bool UEditConditionTestObject::StaticGetBoolFunction()
+{
+	return UEditConditionTestObject::StaticClass()->GetDefaultObject<UEditConditionTestObject>()->BoolProperty;
+}
+
+EditConditionTestEnum UEditConditionTestObject::StaticGetEnumFunction()
+{
+	return UEditConditionTestObject::StaticClass()->GetDefaultObject<UEditConditionTestObject>()->EnumProperty;
+}
+
+TEnumAsByte<EditConditionByteEnum> UEditConditionTestObject::StaticGetByteEnumFunction()
+{
+	return UEditConditionTestObject::StaticClass()->GetDefaultObject<UEditConditionTestObject>()->ByteEnumProperty;
+}
+
+double UEditConditionTestObject::StaticGetDoubleFunction()
+{
+	return UEditConditionTestObject::StaticClass()->GetDefaultObject<UEditConditionTestObject>()->DoubleProperty;
+}
+
+int32 UEditConditionTestObject::StaticGetIntegerFunction()
+{
+	return UEditConditionTestObject::StaticClass()->GetDefaultObject<UEditConditionTestObject>()->IntegerProperty;
+}
+
+uint8 UEditConditionTestObject::StaticGetUintBitfieldFunction()
+{
+	return UEditConditionTestObject::StaticClass()->GetDefaultObject<UEditConditionTestObject>()->UintBitfieldProperty;
+}
+
+UObject* UEditConditionTestObject::StaticGetUObjectPtrFunction()
+{
+	return UEditConditionTestObject::StaticClass()->GetDefaultObject<UEditConditionTestObject>()->UObjectPtr;
+}
+
+TSoftClassPtr<UObject> UEditConditionTestObject::StaticGetSoftClassPtrFunction()
+{
+	return UEditConditionTestObject::StaticClass()->GetDefaultObject<UEditConditionTestObject>()->SoftClassPtr;
+}
+
+TWeakObjectPtr<UObject> UEditConditionTestObject::StaticGetWeakObjectPtrFunction()
+{
+	return UEditConditionTestObject::StaticClass()->GetDefaultObject<UEditConditionTestObject>()->WeakObjectPtr;
+}
+
 #if WITH_DEV_AUTOMATION_TESTS
 struct FTestEditConditionContext : IEditConditionContext
 {
@@ -30,7 +130,7 @@ struct FTestEditConditionContext : IEditConditionContext
 		return "TestEditConditionContext";
 	}
 
-	virtual TOptional<bool> GetBoolValue(const FString& PropertyName) const override
+	virtual TOptional<bool> GetBoolValue(const FString& PropertyName, TWeakObjectPtr<UFunction> CachedFunction = nullptr) const override
 	{
 		if (const bool* Value = BoolValues.Find(PropertyName))
 		{
@@ -39,7 +139,7 @@ struct FTestEditConditionContext : IEditConditionContext
 		return TOptional<bool>();
 	}
 
-	virtual TOptional<int64> GetIntegerValue(const FString& PropertyName) const override
+	virtual TOptional<int64> GetIntegerValue(const FString& PropertyName, TWeakObjectPtr<UFunction> CachedFunction = nullptr) const override
 	{
 		if (const int64* Value = IntegerValues.Find(PropertyName))
 		{
@@ -48,7 +148,7 @@ struct FTestEditConditionContext : IEditConditionContext
 		return TOptional<int64>();
 	}
 
-	virtual TOptional<double> GetNumericValue(const FString& PropertyName) const override
+	virtual TOptional<double> GetNumericValue(const FString& PropertyName, TWeakObjectPtr<UFunction> CachedFunction = nullptr) const override
 	{
 		if (const double* Value = DoubleValues.Find(PropertyName))
 		{
@@ -57,7 +157,7 @@ struct FTestEditConditionContext : IEditConditionContext
 		return TOptional<double>();
 	}
 
-	virtual TOptional<FString> GetEnumValue(const FString& PropertyName) const override
+	virtual TOptional<FString> GetEnumValue(const FString& PropertyName, TWeakObjectPtr<UFunction> CachedFunction = nullptr) const override
 	{
 		if (const FString* Value = EnumValues.Find(PropertyName))
 		{
@@ -77,12 +177,17 @@ struct FTestEditConditionContext : IEditConditionContext
 		return TOptional<int64>();
 	}
 
-	virtual TOptional<UObject*> GetPointerValue(const FString& PropertyName) const override
+	virtual const TWeakObjectPtr<UFunction> GetFunction(const FString& FieldName) const override
+	{
+		return nullptr;
+	}
+
+	virtual TOptional<UObject*> GetPointerValue(const FString& PropertyName, TWeakObjectPtr<UFunction> CachedFunction = nullptr) const override
 	{
 		return TOptional<UObject*>();
 	}
 
-	virtual TOptional<FString> GetTypeName(const FString& PropertyName) const override
+	virtual TOptional<FString> GetTypeName(const FString& PropertyName, TWeakObjectPtr<UFunction> CachedFunction = nullptr) const override
 	{
 		TOptional<FString> Result;
 
@@ -177,6 +282,14 @@ bool FEditConditionParser_Parse::RunTest(const FString& Parameters)
 	bResult &= CanParse(Parser, TEXT("!BoolProperty"), 2, 1);
 	bResult &= CanParse(Parser, TEXT("BoolProperty == true"), 3, 1);
 	bResult &= CanParse(Parser, TEXT("BoolProperty == false"), 3, 1);
+	bResult &= CanParse(Parser, TEXT("GetBoolFunction"), 1, 1);
+	bResult &= CanParse(Parser, TEXT("!GetBoolFunction"), 2, 1);
+	bResult &= CanParse(Parser, TEXT("GetBoolFunction == true"), 3, 1);
+	bResult &= CanParse(Parser, TEXT("GetBoolFunction == false"), 3, 1);
+	bResult &= CanParse(Parser, TEXT("PropertyEditor.EditConditionTestObject.StaticGetBoolFunction"), 1, 1);
+	bResult &= CanParse(Parser, TEXT("!PropertyEditor.EditConditionTestObject.StaticGetBoolFunction"), 2, 1);
+	bResult &= CanParse(Parser, TEXT("PropertyEditor.EditConditionTestObject.StaticGetBoolFunction == true"), 3, 1);
+	bResult &= CanParse(Parser, TEXT("PropertyEditor.EditConditionTestObject.StaticGetBoolFunction == false"), 3, 1);
 	bResult &= CanParse(Parser, TEXT("IntProperty == 0"), 3, 1);
 	bResult &= CanParse(Parser, TEXT("IntProperty != 0"), 3, 1);
 	bResult &= CanParse(Parser, TEXT("IntProperty > 0"), 3, 1);
@@ -195,6 +308,12 @@ bool FEditConditionParser_Parse::RunTest(const FString& Parameters)
 	bResult &= CanParse(Parser, TEXT("Enum != EType::Value && BoolProperty"), 5, 2);
 	bResult &= CanParse(Parser, TEXT("Enum == EType::Value || BoolProperty == false"), 7, 2);
 	bResult &= CanParse(Parser, TEXT("Enum != EType::Value || BoolProperty == bFoo"), 7, 3);
+	bResult &= CanParse(Parser, TEXT("Enum != EType::Value && GetBoolFunction"), 5, 2);
+	bResult &= CanParse(Parser, TEXT("Enum == EType::Value || GetBoolFunction == false"), 7, 2);
+	bResult &= CanParse(Parser, TEXT("Enum != EType::Value || GetBoolFunction == bFoo"), 7, 3);
+	bResult &= CanParse(Parser, TEXT("Enum != EType::Value && PropertyEditor.EditConditionTestObject.StaticGetBoolFunction"), 5, 2);
+	bResult &= CanParse(Parser, TEXT("Enum == EType::Value || PropertyEditor.EditConditionTestObject.StaticGetBoolFunction == false"), 7, 2);
+	bResult &= CanParse(Parser, TEXT("Enum != EType::Value || PropertyEditor.EditConditionTestObject.StaticGetBoolFunction == bFoo"), 7, 3);
 	bResult &= CanParse(Parser, TEXT("Enum == EType::Value && Foo != 5"), 7, 2);
 	bResult &= CanParse(Parser, TEXT("Enum != EType::Value && Foo == Bar"), 7, 3);
 	bResult &= CanParse(Parser, TEXT("PointerProperty == nullptr"), 3, 1);
@@ -231,6 +350,58 @@ static bool CanEvaluate(const FEditConditionParser& Parser, const IEditCondition
 	return true;
 }
 
+/** For expressions we expect / wish to be considered mal-formed */
+static bool CanNotEvaluate(const FEditConditionParser& Parser, const IEditConditionContext& Context, const FString& Expression)
+{
+	TSharedPtr<FEditConditionExpression> Parsed = Parser.Parse(Expression);
+	if (!Parsed.IsValid())
+	{
+		ensureMsgf(false, TEXT("Failed to parse expression: %s"), *Expression);
+		return false;
+	}
+
+	TValueOrError<bool, FText> Result = Parser.Evaluate(*Parsed.Get(), Context);
+	if (Result.IsValid())
+	{
+		ensureMsgf(false, TEXT("Error. Expected expression to fail evaluation: %s"), *Expression);
+		return false;
+	}
+
+	return true;
+}
+
+static bool RunVoidTests(const IEditConditionContext& Context)
+{
+	FEditConditionParser Parser;
+	bool bResult = true;
+
+	bResult &= CanNotEvaluate(Parser, Context, TEXT("VoidFunction"));
+	bResult &= CanNotEvaluate(Parser, Context, TEXT("!VoidFunction"));
+
+	bResult &= CanNotEvaluate(Parser, Context, TEXT("VoidFunction == true"));
+	bResult &= CanNotEvaluate(Parser, Context, TEXT("VoidFunction == false"));
+
+	bResult &= CanNotEvaluate(Parser, Context, TEXT("VoidFunction == VoidFunction"));
+	bResult &= CanNotEvaluate(Parser, Context, TEXT("VoidFunction != VoidFunction"));
+
+	bResult &= CanNotEvaluate(Parser, Context, TEXT("VoidFunction != true"));
+	bResult &= CanNotEvaluate(Parser, Context, TEXT("VoidFunction != false"));
+
+	bResult &= CanNotEvaluate(Parser, Context, TEXT("PropertyEditor.EditConditionTestObject.StaticVoidFunction"));
+	bResult &= CanNotEvaluate(Parser, Context, TEXT("!PropertyEditor.EditConditionTestObject.StaticVoidFunction"));
+
+	bResult &= CanNotEvaluate(Parser, Context, TEXT("PropertyEditor.EditConditionTestObject.StaticVoidFunction == true"));
+	bResult &= CanNotEvaluate(Parser, Context, TEXT("PropertyEditor.EditConditionTestObject.StaticVoidFunction == false"));
+
+	bResult &= CanNotEvaluate(Parser, Context, TEXT("PropertyEditor.EditConditionTestObject.StaticVoidFunction == PropertyEditor.EditConditionTestObject.StaticVoidFunction"));
+	bResult &= CanNotEvaluate(Parser, Context, TEXT("PropertyEditor.EditConditionTestObject.StaticVoidFunction != PropertyEditor.EditConditionTestObject.StaticVoidFunction"));
+
+	bResult &= CanNotEvaluate(Parser, Context, TEXT("PropertyEditor.EditConditionTestObject.StaticVoidFunction != true"));
+	bResult &= CanNotEvaluate(Parser, Context, TEXT("PropertyEditor.EditConditionTestObject.StaticVoidFunction != false"));
+
+	return bResult;
+}
+
 static bool RunBoolTests(const IEditConditionContext& Context)
 {
 	FEditConditionParser Parser;
@@ -243,6 +414,7 @@ static bool RunBoolTests(const IEditConditionContext& Context)
 	bResult &= CanEvaluate(Parser, Context, TEXT("fAlSe"), false);
 	bResult &= CanEvaluate(Parser, Context, TEXT("!true"), false);
 	bResult &= CanEvaluate(Parser, Context, TEXT("!false"), true);
+
 	bResult &= CanEvaluate(Parser, Context, TEXT("BoolProperty"), true);
 	bResult &= CanEvaluate(Parser, Context, TEXT("!BoolProperty"), false);
 	
@@ -254,6 +426,30 @@ static bool RunBoolTests(const IEditConditionContext& Context)
 
 	bResult &= CanEvaluate(Parser, Context, TEXT("BoolProperty != true"), false);
 	bResult &= CanEvaluate(Parser, Context, TEXT("BoolProperty != false"), true);
+
+	bResult &= CanEvaluate(Parser, Context, TEXT("GetBoolFunction"), true);
+	bResult &= CanEvaluate(Parser, Context, TEXT("!GetBoolFunction"), false);
+
+	bResult &= CanEvaluate(Parser, Context, TEXT("GetBoolFunction == true"), true);
+	bResult &= CanEvaluate(Parser, Context, TEXT("GetBoolFunction == false"), false);
+
+	bResult &= CanEvaluate(Parser, Context, TEXT("GetBoolFunction == GetBoolFunction"), true);
+	bResult &= CanEvaluate(Parser, Context, TEXT("GetBoolFunction != GetBoolFunction"), false);
+
+	bResult &= CanEvaluate(Parser, Context, TEXT("GetBoolFunction != true"), false);
+	bResult &= CanEvaluate(Parser, Context, TEXT("GetBoolFunction != false"), true);
+
+	bResult &= CanEvaluate(Parser, Context, TEXT("PropertyEditor.EditConditionTestObject.StaticGetBoolFunction"), false);
+	bResult &= CanEvaluate(Parser, Context, TEXT("!PropertyEditor.EditConditionTestObject.StaticGetBoolFunction"), true);
+
+	bResult &= CanEvaluate(Parser, Context, TEXT("PropertyEditor.EditConditionTestObject.StaticGetBoolFunction == true"), false);
+	bResult &= CanEvaluate(Parser, Context, TEXT("PropertyEditor.EditConditionTestObject.StaticGetBoolFunction == false"), true);
+
+	bResult &= CanEvaluate(Parser, Context, TEXT("PropertyEditor.EditConditionTestObject.StaticGetBoolFunction == PropertyEditor.EditConditionTestObject.StaticGetBoolFunction"), true);
+	bResult &= CanEvaluate(Parser, Context, TEXT("PropertyEditor.EditConditionTestObject.StaticGetBoolFunction != PropertyEditor.EditConditionTestObject.StaticGetBoolFunction"), false);
+
+	bResult &= CanEvaluate(Parser, Context, TEXT("PropertyEditor.EditConditionTestObject.StaticGetBoolFunction != true"), true);
+	bResult &= CanEvaluate(Parser, Context, TEXT("PropertyEditor.EditConditionTestObject.StaticGetBoolFunction != false"), false);
 	
 	bResult &= CanEvaluate(Parser, Context, TEXT("true && true"), true);
 	bResult &= CanEvaluate(Parser, Context, TEXT("true && false"), false);
@@ -261,9 +457,18 @@ static bool RunBoolTests(const IEditConditionContext& Context)
 	bResult &= CanEvaluate(Parser, Context, TEXT("false && false"), false);
 	bResult &= CanEvaluate(Parser, Context, TEXT("true && true && true"), true);
 	bResult &= CanEvaluate(Parser, Context, TEXT("true && true && false"), false);
+
 	bResult &= CanEvaluate(Parser, Context, TEXT("BoolProperty && BoolProperty"), true);
 	bResult &= CanEvaluate(Parser, Context, TEXT("BoolProperty && false"), false);
 	bResult &= CanEvaluate(Parser, Context, TEXT("false && BoolProperty"), false);
+
+	bResult &= CanEvaluate(Parser, Context, TEXT("GetBoolFunction && GetBoolFunction"), true);
+	bResult &= CanEvaluate(Parser, Context, TEXT("GetBoolFunction && false"), false);
+	bResult &= CanEvaluate(Parser, Context, TEXT("false && GetBoolFunction"), false);
+
+	bResult &= CanEvaluate(Parser, Context, TEXT("PropertyEditor.EditConditionTestObject.StaticGetBoolFunction && PropertyEditor.EditConditionTestObject.StaticGetBoolFunction"), false);
+	bResult &= CanEvaluate(Parser, Context, TEXT("PropertyEditor.EditConditionTestObject.StaticGetBoolFunction && false"), false);
+	bResult &= CanEvaluate(Parser, Context, TEXT("false && PropertyEditor.EditConditionTestObject.StaticGetBoolFunction"), false);
 	
 	bResult &= CanEvaluate(Parser, Context, TEXT("true || true"), true);
 	bResult &= CanEvaluate(Parser, Context, TEXT("true || false"), true);
@@ -271,18 +476,46 @@ static bool RunBoolTests(const IEditConditionContext& Context)
 	bResult &= CanEvaluate(Parser, Context, TEXT("false || false"), false);
 	bResult &= CanEvaluate(Parser, Context, TEXT("true || true || true"), true);
 	bResult &= CanEvaluate(Parser, Context, TEXT("true || true || false"), true);
+
 	bResult &= CanEvaluate(Parser, Context, TEXT("BoolProperty || BoolProperty"), true);
 	bResult &= CanEvaluate(Parser, Context, TEXT("BoolProperty || false"), true);
 	bResult &= CanEvaluate(Parser, Context, TEXT("false || BoolProperty"), true);
 
+	bResult &= CanEvaluate(Parser, Context, TEXT("GetBoolFunction || GetBoolFunction"), true);
+	bResult &= CanEvaluate(Parser, Context, TEXT("GetBoolFunction || false"), true);
+	bResult &= CanEvaluate(Parser, Context, TEXT("false || GetBoolFunction"), true);
+
+	bResult &= CanEvaluate(Parser, Context, TEXT("PropertyEditor.EditConditionTestObject.StaticGetBoolFunction || PropertyEditor.EditConditionTestObject.StaticGetBoolFunction"), false);
+	bResult &= CanEvaluate(Parser, Context, TEXT("PropertyEditor.EditConditionTestObject.StaticGetBoolFunction || false"), false);
+	bResult &= CanEvaluate(Parser, Context, TEXT("false || PropertyEditor.EditConditionTestObject.StaticGetBoolFunction"), false);
+
 	return bResult;
 }
 
+IMPLEMENT_SIMPLE_AUTOMATION_TEST(FEditConditionParser_EvaluateVoidFails, "EditConditionParser.EvaluateVoidFails", EAutomationTestFlags::EditorContext | EAutomationTestFlags::EngineFilter)
+bool FEditConditionParser_EvaluateVoidFails::RunTest(const FString& Parameters)
+{
+	// We only expect errors first time, since the API only logs first time
+	static bool bExpectErrors = true;
+
+	if (bExpectErrors)
+	{
+		AddExpectedError(TEXT("EditCondition attempted to use an invalid operand"), EAutomationExpectedErrorFlags::Contains, 2);
+		bExpectErrors = false;
+	}
+
+	FTestEditConditionContext TestContext;
+
+	return RunVoidTests(TestContext);
+}
+
 IMPLEMENT_SIMPLE_AUTOMATION_TEST(FEditConditionParser_EvaluateBool, "EditConditionParser.EvaluateBool", EAutomationTestFlags::EditorContext | EAutomationTestFlags::EngineFilter)
 bool FEditConditionParser_EvaluateBool::RunTest(const FString& Parameters)
 {
 	FTestEditConditionContext TestContext;
 	TestContext.SetupBool(TEXT("BoolProperty"), true);
+	TestContext.SetupBool(TEXT("GetBoolFunction"), true);
+	TestContext.SetupBool(TEXT("PropertyEditor.EditConditionTestObject.StaticGetBoolFunction"), false);
 
 	return RunBoolTests(TestContext);
 }
@@ -477,6 +710,26 @@ bool FEditConditionParser_EvaluateUObject::RunTest(const FString& Parameters)
 		bAllResults &= RunEnumTests(Context, ByteEnumType, ByteEnumPropertyName);
 	}
 
+	// void expression failure
+	{
+		// We only expect errors first time, since the API only logs first time
+		static bool bExpectErrors = true;
+
+		if (bExpectErrors)
+		{
+			AddExpectedError(TEXT("EditCondition parsing failed"), EAutomationExpectedErrorFlags::Contains, 1);
+			AddExpectedError(TEXT("EditCondition attempted to use an invalid operand"), EAutomationExpectedErrorFlags::Contains, 2);
+			bExpectErrors = false;
+		}
+		
+		// Using bool property for context since doesn't matter as long as parent is same (due to evaluation failing)
+		static const FName BoolPropertyName(TEXT("BoolProperty"));
+		TSharedPtr<FPropertyNode> PropertyNode = ObjectNode->FindChildPropertyNode(BoolPropertyName, true);
+		FEditConditionContext Context(*PropertyNode.Get());
+
+		bAllResults &= RunVoidTests(Context);
+	}
+
 	// bool comparison
 	{
 		static const FName BoolPropertyName(TEXT("BoolProperty"));
diff --git a/Engine/Source/Editor/PropertyEditor/Private/EditConditionParserTests.h b/Engine/Source/Editor/PropertyEditor/Private/EditConditionParserTests.h
--- a/Engine/Source/Editor/PropertyEditor/Private/EditConditionParserTests.h
+++ b/Engine/Source/Editor/PropertyEditor/Private/EditConditionParserTests.h
@@ -26,6 +26,13 @@ enum EditConditionByteEnum : int
 	Second = 31
 };
 
+
+/**
+ * Test object for edit condition property checks
+ * 
+ * Note: Currently only bool functions are supported (Including static). 
+ * remaining are not used because currently not supported by edit condition parser
+ */
 UCLASS(transient)
 class UEditConditionTestObject : public UObject
 {
@@ -57,4 +64,68 @@ class UEditConditionTestObject : public UObject
 
 	UPROPERTY(EditAnywhere, Category=Test)
 	TWeakObjectPtr<UObject> WeakObjectPtr;
+
+public:
+	/** Used in test cases that should fail, should not be able to execute a void function in edit condition */
+	UFUNCTION()
+	void VoidFunction() const;
+
+	UFUNCTION()
+	bool GetBoolFunction() const;
+
+	UFUNCTION()
+	EditConditionTestEnum GetEnumFunction() const;
+
+	UFUNCTION()
+	TEnumAsByte<EditConditionByteEnum> GetByteEnumFunction() const;
+
+	UFUNCTION()
+	double GetDoubleFunction() const;
+
+	UFUNCTION()
+	int32 GetIntegerFunction() const;
+
+	UFUNCTION()
+	uint8 GetUintBitfieldFunction() const;
+
+	UFUNCTION()
+	UObject* GetUObjectPtrFunction() const;
+
+	UFUNCTION()
+	TSoftClassPtr<UObject> GetSoftClassPtrFunction() const;
+
+	UFUNCTION()
+	TWeakObjectPtr<UObject> GetWeakObjectPtrFunction() const;
+	
+public:
+	/** Used in test cases that should fail, should not be able to execute a void function in edit condition */
+	UFUNCTION()
+	void StaticVoidFunction();
+
+	UFUNCTION()
+	static bool StaticGetBoolFunction();
+
+	UFUNCTION()
+	static EditConditionTestEnum StaticGetEnumFunction();
+
+	UFUNCTION()
+	static TEnumAsByte<EditConditionByteEnum> StaticGetByteEnumFunction();
+
+	UFUNCTION()
+	static double StaticGetDoubleFunction();
+
+	UFUNCTION()
+	static int32 StaticGetIntegerFunction();
+
+	UFUNCTION()
+	static uint8 StaticGetUintBitfieldFunction();
+
+	UFUNCTION()
+	static UObject* StaticGetUObjectPtrFunction();
+
+	UFUNCTION()
+	static TSoftClassPtr<UObject> StaticGetSoftClassPtrFunction();
+
+	UFUNCTION()
+	static TWeakObjectPtr<UObject> StaticGetWeakObjectPtrFunction();
 };
