PropertyNode: Change OnRebuildChildren to an event.
This fixes nodes not rebuilding when there are multiple listeners.
E.g. FDetailArrayBuilder generating a row that contains instanced object property. Previously changing the object instance would "freeze" the array builder.

diff --git a/Engine/Source/Editor/PropertyEditor/Private/PropertyHandleImpl.cpp b/Engine/Source/Editor/PropertyEditor/Private/PropertyHandleImpl.cpp
--- a/Engine/Source/Editor/PropertyEditor/Private/PropertyHandleImpl.cpp
+++ b/Engine/Source/Editor/PropertyEditor/Private/PropertyHandleImpl.cpp
@@ -595,13 +595,23 @@ void FPropertyValueImpl::SetOnPropertyResetToDefault(const FSimpleDelegate& InOn
 	}
 }
 
-void FPropertyValueImpl::SetOnRebuildChildren( const FSimpleDelegate& InOnRebuildChildren )
+FDelegateHandle FPropertyValueImpl::SetOnRebuildChildren( const FSimpleDelegate& InOnRebuildChildren )
 {
 	if( PropertyNode.IsValid() )
 	{
-		PropertyNode.Pin()->SetOnRebuildChildren( InOnRebuildChildren );
+		return PropertyNode.Pin()->OnRebuildChildren().Add(InOnRebuildChildren);
+	}
+	return {};
+} 
+
+void FPropertyValueImpl::UnregisterOnRebuildChildren(FDelegateHandle Handle)
+{
+	if (PropertyNode.IsValid())
+	{
+		PropertyNode.Pin()->OnRebuildChildren().Remove(Handle);
 	}
 }
+
 /**
  * Gets the max valid index for a array property of an object
  * @param InObjectNode - The parent of the variable being clamped
@@ -4765,9 +4775,14 @@ FPropertyAccess::Result FPropertyHandleArray::GetNumElements( uint32 &OutNumItem
 	return FPropertyAccess::Success;
 }
 
-void FPropertyHandleArray::SetOnNumElementsChanged( FSimpleDelegate& OnChildrenChanged )
+FDelegateHandle FPropertyHandleArray::SetOnNumElementsChanged( const FSimpleDelegate& OnChildrenChanged )
 {
-	Implementation->SetOnRebuildChildren( OnChildrenChanged );
+	return Implementation->SetOnRebuildChildren( OnChildrenChanged );
+}
+
+void FPropertyHandleArray::UnregisterOnNumElementsChanged(FDelegateHandle Handle)
+{
+	Implementation->UnregisterOnRebuildChildren(Handle);
 }
 
 TSharedPtr<IPropertyHandleArray> FPropertyHandleArray::AsArray()
@@ -4923,9 +4938,14 @@ TSharedRef<IPropertyHandle> FPropertyHandleSet::GetElement(int32 Index) const
 	return PropertyEditorHelpers::GetPropertyHandle(PropertyNode.ToSharedRef(), Implementation->GetNotifyHook(), Implementation->GetPropertyUtilities()).ToSharedRef();
 }
 
-void FPropertyHandleSet::SetOnNumElementsChanged(FSimpleDelegate& OnChildrenChanged)
+FDelegateHandle FPropertyHandleSet::SetOnNumElementsChanged( const FSimpleDelegate& OnChildrenChanged )
+{
+	return Implementation->SetOnRebuildChildren(OnChildrenChanged);
+}
+
+void FPropertyHandleSet::UnregisterOnNumElementsChanged(FDelegateHandle Handle)
 {
-	Implementation->SetOnRebuildChildren(OnChildrenChanged);
+	Implementation->UnregisterOnRebuildChildren(Handle);
 }
 
 TSharedPtr<IPropertyHandleSet> FPropertyHandleSet::AsSet()
@@ -5021,9 +5041,14 @@ FPropertyAccess::Result FPropertyHandleMap::GetNumElements(uint32& OutNumChildre
 	return FPropertyAccess::Success;
 }
 
-void FPropertyHandleMap::SetOnNumElementsChanged(FSimpleDelegate& OnChildrenChanged)
+FDelegateHandle FPropertyHandleMap::SetOnNumElementsChanged( const FSimpleDelegate& OnChildrenChanged )
+{
+	return Implementation->SetOnRebuildChildren(OnChildrenChanged);
+}
+
+void FPropertyHandleMap::UnregisterOnNumElementsChanged(FDelegateHandle Handle)
 {
-	Implementation->SetOnRebuildChildren(OnChildrenChanged);
+	Implementation->UnregisterOnRebuildChildren(Handle);
 }
 
 TSharedPtr<IPropertyHandleMap> FPropertyHandleMap::AsMap()
diff --git a/Engine/Source/Editor/PropertyEditor/Private/PropertyHandleImpl.h b/Engine/Source/Editor/PropertyEditor/Private/PropertyHandleImpl.h
--- a/Engine/Source/Editor/PropertyEditor/Private/PropertyHandleImpl.h
+++ b/Engine/Source/Editor/PropertyEditor/Private/PropertyHandleImpl.h
@@ -171,7 +171,12 @@ public:
 	/**
 	 * Sets a delegate to call when children of the property node must be rebuilt
 	 */
-	void SetOnRebuildChildren( const FSimpleDelegate& InOnRebuildChildren );
+	FDelegateHandle SetOnRebuildChildren( const FSimpleDelegate& InOnRebuildChildren );
+
+	/**
+	 * Unregisters a delegate called when children of the property node must be rebuilt
+	 */
+	void UnregisterOnRebuildChildren(FDelegateHandle Handle);
 
 	/**
 	 * Get the value of a property as a formatted string.
@@ -683,7 +688,8 @@ public:
 	virtual FPropertyAccess::Result DeleteItem( int32 Index ) override;
 	virtual FPropertyAccess::Result SwapItems(int32 FirstIndex, int32 SecondIndex) override;
 	virtual FPropertyAccess::Result GetNumElements( uint32& OutNumItems ) const override;
-	virtual void SetOnNumElementsChanged( FSimpleDelegate& InOnNumElementsChanged ) override;
+	virtual FDelegateHandle SetOnNumElementsChanged( const FSimpleDelegate& InOnNumElementsChanged ) override;
+	virtual void UnregisterOnNumElementsChanged(FDelegateHandle Handle) override;
 	virtual TSharedPtr<IPropertyHandleArray> AsArray() override;
 	virtual TSharedRef<IPropertyHandle> GetElement( int32 Index ) const override;
 	virtual FPropertyAccess::Result MoveElementTo(int32 OriginalIndex, int32 NewIndex) override;
@@ -717,7 +723,8 @@ public:
 	virtual FPropertyAccess::Result GetNumElements(uint32& OutNumElements) override;
 	virtual TSharedRef<IPropertyHandle> GetElement(int32 Index) const override;
 
-	virtual void SetOnNumElementsChanged(FSimpleDelegate& InOnNumElementsChanged) override;
+	virtual FDelegateHandle SetOnNumElementsChanged( const FSimpleDelegate& InOnNumElementsChanged ) override;
+	virtual void UnregisterOnNumElementsChanged(FDelegateHandle Handle) override;
 	virtual bool HasDocumentation() override { return true; }
 	virtual FString GetDocumentationLink() override { return FString("Engine/UI/LevelEditor/Details/Properties/Set/"); }
 	virtual FString GetDocumentationExcerptName() override { return FString("Sets"); }
@@ -737,7 +744,8 @@ public:
 	virtual FPropertyAccess::Result Empty() override;
 	virtual FPropertyAccess::Result DeleteItem(int32 Index) override;
 	virtual FPropertyAccess::Result GetNumElements(uint32& OutNumElements) override;
-	virtual void SetOnNumElementsChanged(FSimpleDelegate& InOnNumElementsChanged) override;
+	virtual FDelegateHandle SetOnNumElementsChanged( const FSimpleDelegate& InOnNumElementsChanged ) override;
+	virtual void UnregisterOnNumElementsChanged(FDelegateHandle Handle) override;
 	virtual bool HasDocumentation() override { return true; }
 	virtual FString GetDocumentationLink() override { return FString("Engine/UI/LevelEditor/Details/Properties/Map/"); }
 	virtual FString GetDocumentationExcerptName() override { return FString("Maps"); }
diff --git a/Engine/Source/Editor/PropertyEditor/Private/PropertyNode.cpp b/Engine/Source/Editor/PropertyEditor/Private/PropertyNode.cpp
--- a/Engine/Source/Editor/PropertyEditor/Private/PropertyNode.cpp
+++ b/Engine/Source/Editor/PropertyEditor/Private/PropertyNode.cpp
@@ -302,7 +302,7 @@ void FPropertyNode::RebuildChildren()
 	bChildrenRebuilt = true;
 
 	// Notify any listener that children have been rebuilt
-	OnRebuildChildren.ExecuteIfBound();
+	OnRebuildChildrenEvent.Broadcast();
 }
 
 void FPropertyNode::AddChildNode(TSharedPtr<FPropertyNode> InNode)
@@ -3068,9 +3068,9 @@ void FPropertyNode::SetExpandedChildPropertyNodes(const TSet<FString>& InNodesTo
 	} while (RecursiveStack.Num() > 0);
 }
 
-void FPropertyNode::SetOnRebuildChildren( FSimpleDelegate InOnRebuildChildren )
+FDelegateHandle FPropertyNode::SetOnRebuildChildren(const FSimpleDelegate& InOnRebuildChildren)
 {
-	OnRebuildChildren = InOnRebuildChildren;
+	return OnRebuildChildrenEvent.Add(InOnRebuildChildren);
 }
 
 TSharedPtr< FPropertyItemValueDataTrackerSlate > FPropertyNode::GetValueTracker( UObject* Object, uint32 ObjIndex )
diff --git a/Engine/Source/Editor/PropertyEditor/Private/PropertyNode.h b/Engine/Source/Editor/PropertyEditor/Private/PropertyNode.h
--- a/Engine/Source/Editor/PropertyEditor/Private/PropertyNode.h
+++ b/Engine/Source/Editor/PropertyEditor/Private/PropertyNode.h
@@ -688,7 +688,9 @@ public:
 
 	void NotifyPostChange(FPropertyChangedEvent& InPropertyChangedEvent, FNotifyHook* InNotifyHook);
 
-	void SetOnRebuildChildren(FSimpleDelegate InOnRebuildChildren);
+	DECLARE_EVENT(FPropertyNode, FPropertyChildrenRebuiltEvent);
+	FDelegateHandle SetOnRebuildChildren(const FSimpleDelegate& InOnRebuildChildren);
+	FPropertyChildrenRebuiltEvent& OnRebuildChildren() { return OnRebuildChildrenEvent; }
 
 	/**
 	 * Propagates the property change to all instances of an archetype
@@ -1171,7 +1173,7 @@ protected:
 	TArray< TSharedPtr<FPropertyNode> > ChildNodes;
 
 	/** Called when this node's children are rebuilt */
-	FSimpleDelegate OnRebuildChildren;
+	FPropertyChildrenRebuiltEvent OnRebuildChildrenEvent;
 
 	/** Called when this node's property value is about to change (called during NotifyPreChange) */
 	FPropertyValuePreChangeEvent PropertyValuePreChangeEvent;
diff --git a/Engine/Source/Editor/PropertyEditor/Public/PropertyCustomizationHelpers.h b/Engine/Source/Editor/PropertyEditor/Public/PropertyCustomizationHelpers.h
--- a/Engine/Source/Editor/PropertyEditor/Public/PropertyCustomizationHelpers.h
+++ b/Engine/Source/Editor/PropertyEditor/Public/PropertyCustomizationHelpers.h
@@ -419,18 +419,17 @@ public:
 		, bDisplayElementNum(InDisplayElementNum)
 	{
 		check( ArrayProperty.IsValid() );
-
+		
 		// Delegate for when the number of children in the array changes
 		FSimpleDelegate OnNumChildrenChanged = FSimpleDelegate::CreateRaw( this, &FDetailArrayBuilder::OnNumChildrenChanged );
-		ArrayProperty->SetOnNumElementsChanged( OnNumChildrenChanged );
+		OnNumElementsChangedHandle = ArrayProperty->SetOnNumElementsChanged( OnNumChildrenChanged );
 
 		BaseProperty->MarkHiddenByCustomization();
 	}
 	
 	~FDetailArrayBuilder()
 	{
-		FSimpleDelegate Empty;
-		ArrayProperty->SetOnNumElementsChanged( Empty );
+		ArrayProperty->UnregisterOnNumElementsChanged(OnNumElementsChangedHandle);
 	}
 
 	void SetDisplayName( const FText& InDisplayName )
@@ -541,6 +540,7 @@ private:
 	bool bGenerateHeader;
 	bool bDisplayResetToDefault;
 	bool bDisplayElementNum;
+	FDelegateHandle OnNumElementsChangedHandle;
 };
 
 /**
diff --git a/Engine/Source/Editor/PropertyEditor/Public/PropertyHandle.h b/Engine/Source/Editor/PropertyEditor/Public/PropertyHandle.h
--- a/Engine/Source/Editor/PropertyEditor/Public/PropertyHandle.h
+++ b/Engine/Source/Editor/PropertyEditor/Public/PropertyHandle.h
@@ -766,7 +766,12 @@ public:
 	/**
 	 * Sets a delegate to call when the number of elements changes
 	 */
-	virtual void SetOnNumElementsChanged( FSimpleDelegate& InOnNumElementsChanged ) = 0;
+	virtual FDelegateHandle SetOnNumElementsChanged( const FSimpleDelegate& InOnNumElementsChanged ) = 0;
+
+	/**
+	 * Unregisters a delegate that is called when the number of elements changes
+	 */
+	virtual void UnregisterOnNumElementsChanged(FDelegateHandle Handle) = 0;
 };
 
 /**
@@ -814,7 +819,12 @@ public:
 	/**
 	 * Sets a delegate to call when the number of elements changes
 	 */
-	virtual void SetOnNumElementsChanged(FSimpleDelegate& InOnNumElementsChanged) = 0;
+	virtual FDelegateHandle SetOnNumElementsChanged(const FSimpleDelegate& InOnNumElementsChanged) = 0;
+
+	/**
+	 * Unregisters a delegate that is called when the number of elements changes
+	 */
+	virtual void UnregisterOnNumElementsChanged(FDelegateHandle Handle) = 0;
 };
 
 /**
@@ -856,5 +866,10 @@ public:
 	/**
 	 * Sets a delegate to call when the number of elements changes
 	 */
-	virtual void SetOnNumElementsChanged(FSimpleDelegate& InOnNumElementsChanged) = 0;
+	virtual FDelegateHandle SetOnNumElementsChanged(const FSimpleDelegate& InOnNumElementsChanged) = 0;
+	
+	/**
+	 * Unregisters a delegate that is called when the number of elements changes
+	 */
+	virtual void UnregisterOnNumElementsChanged(FDelegateHandle Handle) = 0;
 };
