Remove default value for enabled attributes on detail rows

These default to being enabled at the call-site, and having them always set in the constructor causes IPropertyTypeCustomizations to always override enabled states that are set by an owning IDetailCustomization

diff --git a/Engine/Source/Editor/PropertyEditor/Private/DetailPropertyRow.cpp b/Engine/Source/Editor/PropertyEditor/Private/DetailPropertyRow.cpp
--- a/Engine/Source/Editor/PropertyEditor/Private/DetailPropertyRow.cpp
+++ b/Engine/Source/Editor/PropertyEditor/Private/DetailPropertyRow.cpp
@@ -18,8 +18,7 @@
 #define LOCTEXT_NAMESPACE	"DetailPropertyRow"
 
 FDetailPropertyRow::FDetailPropertyRow(TSharedPtr<FPropertyNode> InPropertyNode, TSharedRef<FDetailCategoryImpl> InParentCategory, TSharedPtr<FComplexPropertyNode> InExternalRootNode)
-	: CustomIsEnabledAttrib( true )
-	, PropertyNode( InPropertyNode )
+	: PropertyNode( InPropertyNode )
 	, ParentCategory( InParentCategory )
 	, ExternalRootNode( InExternalRootNode )
 	, bShowPropertyButtons( true )
diff --git a/Engine/Source/Editor/PropertyEditor/Public/DetailWidgetRow.h b/Engine/Source/Editor/PropertyEditor/Public/DetailWidgetRow.h
--- a/Engine/Source/Editor/PropertyEditor/Public/DetailWidgetRow.h
+++ b/Engine/Source/Editor/PropertyEditor/Public/DetailWidgetRow.h
@@ -113,7 +113,6 @@ public:
 		, ExtensionWidget( *this, 0.0f, 0.0f, HAlign_Right, VAlign_Center)
 		, WholeRowWidget( *this, 0.0f, 0.0f, HAlign_Fill, VAlign_Fill )
 		, VisibilityAttr( EVisibility::Visible )
-		, IsEnabledAttr( true )
 		, FilterTextString()
 		, CopyMenuAction()
 		, PasteMenuAction()
