Add an optional order to RegisterCustomClassLayout

diff --git a/Engine/Source/Editor/PropertyEditor/Private/PropertyEditorModule.cpp b/Engine/Source/Editor/PropertyEditor/Private/PropertyEditorModule.cpp
--- a/Engine/Source/Editor/PropertyEditor/Private/PropertyEditorModule.cpp
+++ b/Engine/Source/Editor/PropertyEditor/Private/PropertyEditorModule.cpp
@@ -432,14 +432,13 @@ TSharedRef<IPropertyChangeListener> FPropertyEditorModule::CreatePropertyChangeL
 	return MakeShareable( new FPropertyChangeListener );
 }
 
-void FPropertyEditorModule::RegisterCustomClassLayout( FName ClassName, FOnGetDetailCustomizationInstance DetailLayoutDelegate )
+void FPropertyEditorModule::RegisterCustomClassLayout( FName ClassName, FOnGetDetailCustomizationInstance DetailLayoutDelegate, FRegisterCustomClassLayoutParams Params )
 {
 	if (ClassName != NAME_None)
 	{
 		FDetailLayoutCallback Callback;
 		Callback.DetailLayoutDelegate = DetailLayoutDelegate;
-		// @todo: DetailsView: Fix me: this specifies the order in which detail layouts should be queried
-		Callback.Order = ClassNameToDetailLayoutNameMap.Num();
+		Callback.Order = Params.OptionalOrder.Get(ClassNameToDetailLayoutNameMap.Num());
 
 		ClassNameToDetailLayoutNameMap.Add(ClassName, Callback);
 	}
diff --git a/Engine/Source/Editor/PropertyEditor/Public/PropertyEditorModule.h b/Engine/Source/Editor/PropertyEditor/Public/PropertyEditorModule.h
--- a/Engine/Source/Editor/PropertyEditor/Public/PropertyEditorModule.h
+++ b/Engine/Source/Editor/PropertyEditor/Public/PropertyEditorModule.h
@@ -11,6 +11,8 @@
 #include "Widgets/SWidget.h"
 #include "Widgets/SWindow.h"
 
+#include "Misc/Optional.h"
+
 class FAssetEditorToolkit;
 class FNotifyHook;
 class IPropertyHandle;
@@ -210,6 +212,12 @@ private:
 };
 
 
+struct FRegisterCustomClassLayoutParams
+{
+	/* Optional order to register this class layout with. Registration order is used when not specified. Lower values are added first */
+	TOptional<int32> OptionalOrder;
+};
+
 class FPropertyEditorModule : public IModuleInterface
 {
 public:
@@ -256,7 +264,7 @@ public:
 	 * @param ClassName	The name of the class that the custom detail layout is for
 	 * @param DetailLayoutDelegate	The delegate to call when querying for custom detail layouts for the classes properties
 	 */
-	virtual void RegisterCustomClassLayout( FName ClassName, FOnGetDetailCustomizationInstance DetailLayoutDelegate );
+	virtual void RegisterCustomClassLayout( FName ClassName, FOnGetDetailCustomizationInstance DetailLayoutDelegate, FRegisterCustomClassLayoutParams Params = FRegisterCustomClassLayoutParams());
 
 	/**
 	 * Unregisters a custom detail layout delegate for a specific class name
