Fix for SPropertyEditorClass not displaying blueprint display names when the property has ShowDisplayNames set

diff --git a/Engine/Source/Editor/PropertyEditor/Private/UserInterface/PropertyEditor/SPropertyEditorClass.cpp b/Engine/Source/Editor/PropertyEditor/Private/UserInterface/PropertyEditor/SPropertyEditorClass.cpp
--- a/Engine/Source/Editor/PropertyEditor/Private/UserInterface/PropertyEditor/SPropertyEditorClass.cpp
+++ b/Engine/Source/Editor/PropertyEditor/Private/UserInterface/PropertyEditor/SPropertyEditorClass.cpp
@@ -207,22 +207,23 @@ void SPropertyEditorClass::Construct(const FArguments& InArgs, const TSharedPtr<
 }
 
 /** Util to give better names for BP generated classes */
-static FString GetClassDisplayName(const UObject* Object, bool bShowDisplayNames)
+static FText GetClassDisplayName(const UObject* Object, bool bShowDisplayNames)
 {
 	const UClass* Class = Cast<UClass>(Object);
 	if (Class != nullptr)
 	{
-		const UBlueprint* BP = UBlueprint::GetBlueprintFromClass(Class);
-		if (BP != nullptr)
+		if (bShowDisplayNames)
 		{
-			return BP->GetName();
+			return Class->GetDisplayNameText();
 		}
-		if (bShowDisplayNames && Class->HasMetaData(TEXT("DisplayName")))
+		
+		UBlueprint* BP = UBlueprint::GetBlueprintFromClass(Class);
+		if(BP != nullptr)
 		{
-			return Class->GetMetaData(TEXT("DisplayName"));
+			return FText::FromString(BP->GetName());
 		}
 	}
-	return (Object) ? Object->GetName() : "None";
+	return (Object) ? FText::FromString(Object->GetName()) : LOCTEXT("InvalidObject", "None");
 }
 
 FText SPropertyEditorClass::GetDisplayValueAsString() const
@@ -240,13 +241,13 @@ FText SPropertyEditorClass::GetDisplayValueAsString() const
 
 			if(Result == FPropertyAccess::Success && ObjectValue != nullptr)
 			{
-				return FText::FromString(GetClassDisplayName(ObjectValue, bShowDisplayNames));
+				return GetClassDisplayName(ObjectValue, bShowDisplayNames);
 			}
 
 			return FText::FromString(FPaths::GetBaseFilename(PropertyEditor->GetValueAsString()));
 		}
 
-		return FText::FromString(GetClassDisplayName(SelectedClass.Get(), bShowDisplayNames));
+		return GetClassDisplayName(SelectedClass.Get(), bShowDisplayNames);
 	}
 	else
 	{
