Reverted change to SPropertyEditorNumeric that allowed spinning during multi-selection.

diff --git a/Engine/Source/Editor/PropertyEditor/Private/UserInterface/PropertyEditor/SPropertyEditorNumeric.h b/Engine/Source/Editor/PropertyEditor/Private/UserInterface/PropertyEditor/SPropertyEditorNumeric.h
--- a/Engine/Source/Editor/PropertyEditor/Private/UserInterface/PropertyEditor/SPropertyEditorNumeric.h
+++ b/Engine/Source/Editor/PropertyEditor/Private/UserInterface/PropertyEditor/SPropertyEditorNumeric.h
@@ -320,7 +320,7 @@ public:
 				}
 			}
 
-			const bool bAllowSpin = !PropertyHandle->GetBoolMetaData("NoSpinbox");
+			const bool bAllowSpin = (!ObjectPropertyNode || (1 == ObjectPropertyNode->GetNumObjects())) && !PropertyHandle->GetBoolMetaData("NoSpinbox");
 
 			ChildSlot
 			[
