Fixed access in ItemPropertyNode.cpp to a shared ptr constructed from a weak ptr. Code was not testing that object still exist, causing a crash when resetting Linked Material Blend Layer to Default.

diff --git a/Engine/Source/Editor/PropertyEditor/Private/ItemPropertyNode.cpp b/Engine/Source/Editor/PropertyEditor/Private/ItemPropertyNode.cpp
--- a/Engine/Source/Editor/PropertyEditor/Private/ItemPropertyNode.cpp
+++ b/Engine/Source/Editor/PropertyEditor/Private/ItemPropertyNode.cpp
@@ -512,8 +512,11 @@ FText FItemPropertyNode::GetDisplayName() const
 				FinalDisplayName =  FText::FromString( PropertyPtr->GetName() );
 			}
 		}
-		else
-		{
+		else if (const TSharedPtr<FPropertyNode> ParentNode = ParentNodeWeakPtr.Pin())
+		{	
+			// Sets and maps do not have a display index.
+			FProperty* ParentProperty = ParentNode->GetProperty();
+
 			// Get the ArraySizeEnum class from meta data.
 			static const FName NAME_ArraySizeEnum("ArraySizeEnum");
 			UEnum* ArraySizeEnum = NULL; 
@@ -521,11 +524,7 @@ FText FItemPropertyNode::GetDisplayName() const
 			{
 				ArraySizeEnum	= FindObject<UEnum>(NULL, *Property->GetMetaData(NAME_ArraySizeEnum));
 			}
-			
-			// Sets and maps do not have a display index.
-			const TSharedPtr<FPropertyNode> ParentNode = ParentNodeWeakPtr.Pin();
-			FProperty* ParentProperty = ParentNode->GetProperty();
-
+		
 			// Also handle UArray's having the ArraySizeEnum entry...
 			if (ArraySizeEnum == nullptr && CastField<FArrayProperty>(ParentProperty) != nullptr && ParentProperty->HasMetaData(NAME_ArraySizeEnum))
 			{
