Moving PropertyTableConstants.h to Public since other modules need it.

diff --git a/Engine/Source/Editor/PropertyEditor/Private/UserInterface/PropertyTable/PropertyTableConstants.h b/Engine/Source/Editor/PropertyEditor/Private/UserInterface/PropertyTable/PropertyTableConstants.h
deleted file mode 100644
diff --git a/Engine/Source/Editor/PropertyEditor/Public/UserInterface/PropertyTable/PropertyTableConstants.h b/Engine/Source/Editor/PropertyEditor/Public/UserInterface/PropertyTable/PropertyTableConstants.h
new file mode 100644
--- /dev/null
+++ b/Engine/Source/Editor/PropertyEditor/Public/UserInterface/PropertyTable/PropertyTableConstants.h
@@ -0,0 +1,11 @@
+// Copyright Epic Games, Inc. All Rights Reserved.
+#pragma once
+
+#include "CoreMinimal.h"
+
+namespace PropertyTableConstants
+{
+	static const FName NormalFontStyle( TEXT("PropertyTable.NormalFont") );
+	static const FName BoldFontStyle( TEXT("PropertyTable.BoldFont") );
+	static const FName FilterFontStyle( TEXT("PropertyTable.FilterFont") );
+};
