Fixes editor crash in sequencer when altering the data of the signature of a skeletal mesh event

diff --git a/Engine/Source/Editor/PropertyEditor/Private/SDetailsView.cpp b/Engine/Source/Editor/PropertyEditor/Private/SDetailsView.cpp
--- a/Engine/Source/Editor/PropertyEditor/Private/SDetailsView.cpp
+++ b/Engine/Source/Editor/PropertyEditor/Private/SDetailsView.cpp
@@ -873,8 +873,6 @@ void SDetailsView::PreSetObject(int32 InNewNumObjects)
 		DestroyColorPicker();
 	}
 
-	ClearKeyboardFocusIfWithin(AsShared());
-
 	// Save existing expanded items first
 	for(TSharedPtr<FComplexPropertyNode>& RootNode : RootPropertyNodes)
 	{
