Added a new hidden property validation list.

diff --git a/Engine/Plugins/Runtime/ModelViewViewModel/Source/ModelViewViewModelBlueprint/Private/MVVMViewBlueprintCompiler.cpp b/Engine/Plugins/Runtime/ModelViewViewModel/Source/ModelViewViewModelBlueprint/Private/MVVMViewBlueprintCompiler.cpp
--- a/Engine/Plugins/Runtime/ModelViewViewModel/Source/ModelViewViewModelBlueprint/Private/MVVMViewBlueprintCompiler.cpp
+++ b/Engine/Plugins/Runtime/ModelViewViewModel/Source/ModelViewViewModelBlueprint/Private/MVVMViewBlueprintCompiler.cpp
@@ -591,7 +591,7 @@ bool FMVVMViewBlueprintCompiler::PreCompile(UWidgetBlueprintGeneratedClass* Clas
 	CompilerBindings.Reset(NumBindings*2);
 	BindingSourceContexts.Reset(NumBindings*2);
 
-	FPropertyEditorPermissionList::Get().AddPermissionList(Class, FNamePermissionList(), EPropertyEditorPermissionListRules::AllowListAllProperties);
+	FPropertyEditorPermissionList::Get().AddPermissionList(Class, FNamePermissionList(), EPropertyPermissionListRules::AllowListAllProperties);
 
 	PreCompileBindingSources(Class, BlueprintView);
 	PreCompileSourceCreators(Class, BlueprintView);
diff --git a/Engine/Source/Editor/PropertyEditor/Private/PropertyEditorPermissionList.cpp b/Engine/Source/Editor/PropertyEditor/Private/PropertyEditorPermissionList.cpp
--- a/Engine/Source/Editor/PropertyEditor/Private/PropertyEditorPermissionList.cpp
+++ b/Engine/Source/Editor/PropertyEditor/Private/PropertyEditorPermissionList.cpp
@@ -10,10 +10,10 @@ DEFINE_LOG_CATEGORY(LogPropertyEditorPermissionList);
 
 namespace
 {
-	const FName PropertyEditorPermissionListOwner = "PropertyEditorPermissionList";
+	const FName PropertyPermissionListOwner = "PropertyPermissionList";
 }
 
-FPropertyEditorPermissionList::FPropertyEditorPermissionList()
+FPropertyPermissionList::FPropertyPermissionList()
 {
 	if (GEditor)
 	{
@@ -21,19 +21,19 @@ FPropertyEditorPermissionList::FPropertyEditorPermissionList()
 	}
 	else
 	{
-		FCoreDelegates::OnPostEngineInit.AddRaw(this, &FPropertyEditorPermissionList::RegisterOnBlueprintCompiled);
+		FCoreDelegates::OnPostEngineInit.AddRaw(this, &FPropertyPermissionList::RegisterOnBlueprintCompiled);
 	}
 }
 
-void FPropertyEditorPermissionList::RegisterOnBlueprintCompiled()
+void FPropertyPermissionList::RegisterOnBlueprintCompiled()
 {
 	if (ensure(GEditor))
 	{
-		GEditor->OnBlueprintCompiled().AddRaw(this, &FPropertyEditorPermissionList::ClearCache);
+		GEditor->OnBlueprintCompiled().AddRaw(this, &FPropertyPermissionList::ClearCache);
 	}
 }
 
-FPropertyEditorPermissionList::~FPropertyEditorPermissionList()
+FPropertyPermissionList::~FPropertyPermissionList()
 {
 	if (GEditor)
 	{
@@ -41,7 +41,7 @@ FPropertyEditorPermissionList::~FPropertyEditorPermissionList()
 	}
 }
 
-void FPropertyEditorPermissionList::ClearCacheAndBroadcast(TSoftObjectPtr<UStruct> ObjectStruct, FName OwnerName)
+void FPropertyPermissionList::ClearCacheAndBroadcast(TSoftObjectPtr<UStruct> ObjectStruct, FName OwnerName)
 {
 	// The cache isn't too expensive to recompute, so it is cleared
 	// and lazily repopulated any time the raw PermissionList changes.
@@ -53,9 +53,9 @@ void FPropertyEditorPermissionList::ClearCacheAndBroadcast(TSoftObjectPtr<UStruc
 	}
 }
 
-void FPropertyEditorPermissionList::AddPermissionList(TSoftObjectPtr<UStruct> Struct, const FNamePermissionList& PermissionList, EPropertyEditorPermissionListRules Rules)
+void FPropertyPermissionList::AddPermissionList(TSoftObjectPtr<UStruct> Struct, const FNamePermissionList& PermissionList, EPropertyPermissionListRules Rules)
 {
-	FPropertyEditorPermissionListEntry& Entry = RawPropertyEditorPermissionList.FindOrAdd(Struct);
+	FPropertyPermissionListEntry& Entry = RawPropertyPermissionList.FindOrAdd(Struct);
 	Entry.PermissionList = PermissionList;
 	// Always use the most permissive rule previously set
 	if (Entry.Rules > Rules)
@@ -66,26 +66,26 @@ void FPropertyEditorPermissionList::AddPermissionList(TSoftObjectPtr<UStruct> St
 	ClearCacheAndBroadcast(Struct);
 }
 
-void FPropertyEditorPermissionList::RemovePermissionList(TSoftObjectPtr<UStruct> Struct)
+void FPropertyPermissionList::RemovePermissionList(TSoftObjectPtr<UStruct> Struct)
 {
-	if (RawPropertyEditorPermissionList.Remove(Struct) > 0)
+	if (RawPropertyPermissionList.Remove(Struct) > 0)
 	{
 		ClearCacheAndBroadcast(Struct);
 	}
 }
 
-void FPropertyEditorPermissionList::ClearPermissionList()
+void FPropertyPermissionList::ClearPermissionList()
 {
 	TArray<TSoftObjectPtr<UStruct>> Keys;
-	RawPropertyEditorPermissionList.Reset();
+	RawPropertyPermissionList.Reset();
 	ClearCacheAndBroadcast();
 }
 
-void FPropertyEditorPermissionList::UnregisterOwner(const FName Owner)
+void FPropertyPermissionList::UnregisterOwner(const FName Owner)
 {
 	TArray<TSoftObjectPtr<UStruct>> StructsToRemove;
 
-	for (TPair<TSoftObjectPtr<UStruct>, FPropertyEditorPermissionListEntry>& Pair : RawPropertyEditorPermissionList)
+	for (TPair<TSoftObjectPtr<UStruct>, FPropertyPermissionListEntry>& Pair : RawPropertyPermissionList)
 	{
 		Pair.Value.PermissionList.UnregisterOwner(Owner);
 		if (Pair.Value.PermissionList.GetOwnerNames().Num() == 0)
@@ -105,18 +105,18 @@ void FPropertyEditorPermissionList::UnregisterOwner(const FName Owner)
 	ClearCacheAndBroadcast(nullptr, Owner);
 }
 
-void FPropertyEditorPermissionList::AddToAllowList(TSoftObjectPtr<UStruct> Struct, const FName PropertyName, const FName Owner)
+void FPropertyPermissionList::AddToAllowList(TSoftObjectPtr<UStruct> Struct, const FName PropertyName, const FName Owner)
 {
-	FPropertyEditorPermissionListEntry& Entry = RawPropertyEditorPermissionList.FindOrAdd(Struct);
+	FPropertyPermissionListEntry& Entry = RawPropertyPermissionList.FindOrAdd(Struct);
 	if (Entry.PermissionList.AddAllowListItem(Owner, PropertyName))
 	{
 		ClearCacheAndBroadcast(Struct, Owner);
 	}
 }
 
-void FPropertyEditorPermissionList::AddToAllowList(TSoftObjectPtr<UStruct> Struct, const TArray<FName>& PropertyNames, const FName Owner)
+void FPropertyPermissionList::AddToAllowList(TSoftObjectPtr<UStruct> Struct, const TArray<FName>& PropertyNames, const FName Owner)
 {
-	FPropertyEditorPermissionListEntry& Entry = RawPropertyEditorPermissionList.FindOrAdd(Struct);
+	FPropertyPermissionListEntry& Entry = RawPropertyPermissionList.FindOrAdd(Struct);
 	bool bAddedItem = false;
 	for (const FName& PropertyName : PropertyNames)
 	{
@@ -132,56 +132,67 @@ void FPropertyEditorPermissionList::AddToAllowList(TSoftObjectPtr<UStruct> Struc
 	}
 }
 
-void FPropertyEditorPermissionList::RemoveFromAllowList(TSoftObjectPtr<UStruct> Struct, const FName PropertyName, const FName Owner)
+void FPropertyPermissionList::RemoveFromAllowList(TSoftObjectPtr<UStruct> Struct, const FName PropertyName, const FName Owner)
 {
-	FPropertyEditorPermissionListEntry& Entry = RawPropertyEditorPermissionList.FindOrAdd(Struct);
+	FPropertyPermissionListEntry& Entry = RawPropertyPermissionList.FindOrAdd(Struct);
 	if (Entry.PermissionList.RemoveAllowListItem(Owner, PropertyName))
 	{
 		ClearCacheAndBroadcast(Struct, Owner);
 	}
 }
 
-void FPropertyEditorPermissionList::AddToDenyList(TSoftObjectPtr<UStruct> Struct, const FName PropertyName, const FName Owner)
+void FPropertyPermissionList::AddToDenyList(TSoftObjectPtr<UStruct> Struct, const FName PropertyName, const FName Owner)
 {
-	FPropertyEditorPermissionListEntry& Entry = RawPropertyEditorPermissionList.FindOrAdd(Struct);
+	FPropertyPermissionListEntry& Entry = RawPropertyPermissionList.FindOrAdd(Struct);
 	if (Entry.PermissionList.AddDenyListItem(Owner, PropertyName))
 	{
 		ClearCacheAndBroadcast(Struct, Owner);
 	}
 }
 
-void FPropertyEditorPermissionList::RemoveFromDenyList(TSoftObjectPtr<UStruct> Struct, const FName PropertyName, const FName Owner)
+void FPropertyPermissionList::RemoveFromDenyList(TSoftObjectPtr<UStruct> Struct, const FName PropertyName, const FName Owner)
 {
-	FPropertyEditorPermissionListEntry& Entry = RawPropertyEditorPermissionList.FindOrAdd(Struct);
+	FPropertyPermissionListEntry& Entry = RawPropertyPermissionList.FindOrAdd(Struct);
 	if (Entry.PermissionList.RemoveDenyListItem(Owner, PropertyName))
 	{
 		ClearCacheAndBroadcast(Struct, Owner);
 	}
 }
 
-void FPropertyEditorPermissionList::SetEnabled(bool bEnable)
+void FPropertyPermissionList::SetEnabled(bool bEnable, bool bAnnounce)
 {
-	bEnablePropertyEditorPermissionList = bEnable;
-	PermissionListEnabledDelegate.Broadcast();
+	if (bEnablePermissionList != bEnable)
+	{
+		bEnablePermissionList = bEnable;
+		if (bAnnounce)
+		{
+			PermissionListEnabledDelegate.Broadcast();
+		}
+	}
+}
+
+void FPropertyPermissionList::ClearCache()
+{
+	CachedPropertyPermissionList.Reset();
 }
 
-void FPropertyEditorPermissionList::ClearCache()
+bool FPropertyPermissionList::HasFiltering(const UStruct* ObjectStruct) const
 {
-	CachedPropertyEditorPermissionList.Reset();
+	return ObjectStruct ? GetCachedPermissionListForStruct(ObjectStruct).HasFiltering() : false;
 }
 
-bool FPropertyEditorPermissionList::DoesPropertyPassFilter(const UStruct* ObjectStruct, FName PropertyName) const
+bool FPropertyPermissionList::DoesPropertyPassFilter(const UStruct* ObjectStruct, FName PropertyName) const
 {
-	if (bEnablePropertyEditorPermissionList && ObjectStruct)
+	if (bEnablePermissionList && ObjectStruct)
 	{
 		return GetCachedPermissionListForStruct(ObjectStruct).PassesFilter(PropertyName);
 	}
 	return true;
 }
 
-const FNamePermissionList& FPropertyEditorPermissionList::GetCachedPermissionListForStruct(const UStruct* Struct) const
+const FNamePermissionList& FPropertyPermissionList::GetCachedPermissionListForStruct(const UStruct* Struct) const
 {
-	const FNamePermissionList* CachedPermissionList = CachedPropertyEditorPermissionList.Find(Struct);
+	const FNamePermissionList* CachedPermissionList = CachedPropertyPermissionList.Find(Struct);
 	if (CachedPermissionList)
 	{
 		return *CachedPermissionList;
@@ -193,21 +204,21 @@ const FNamePermissionList& FPropertyEditorPermissionList::GetCachedPermissionLis
 	return GetCachedPermissionListForStructHelper(Struct, bShouldPermissionListAllProperties);
 }
 
-const FNamePermissionList& FPropertyEditorPermissionList::GetCachedPermissionListForStructHelper(const UStruct* Struct, bool& bInOutShouldAllowListAllProperties) const
+const FNamePermissionList& FPropertyPermissionList::GetCachedPermissionListForStructHelper(const UStruct* Struct, bool& bInOutShouldAllowListAllProperties) const
 {
 	check(Struct);
 
-	const FPropertyEditorPermissionListEntry* Entry = RawPropertyEditorPermissionList.Find(Struct);
+	const FPropertyPermissionListEntry* Entry = RawPropertyPermissionList.Find(Struct);
 	// If an entry is set to AllowListAll, then treat it as if it has no manually-defined properties.
-	const bool bIsThisAllowListEmpty = Entry ? (Entry->PermissionList.GetAllowList().Num() == 0 || Entry->Rules == EPropertyEditorPermissionListRules::AllowListAllProperties) : true;
+	const bool bIsThisAllowListEmpty = Entry ? (Entry->PermissionList.GetAllowList().Num() == 0 || Entry->Rules == EPropertyPermissionListRules::AllowListAllProperties) : true;
 
 	// Normally this case would be caught in GetCachedPermissionListForStruct, but when being called recursively from a subclass
 	// we still need to update bInOutShouldAllowListAllProperties so that new PermissionLists cache properly.
-	FNamePermissionList* PermissionList = CachedPropertyEditorPermissionList.Find(Struct);
+	FNamePermissionList* PermissionList = CachedPropertyPermissionList.Find(Struct);
 	if (PermissionList)
 	{
 		// Same check as below, but it specifically has to come after the recursive call so we have to duplicate the code here
-		if (Entry && Entry->Rules == EPropertyEditorPermissionListRules::AllowListAllProperties)
+		if (Entry && Entry->Rules == EPropertyPermissionListRules::AllowListAllProperties)
 		{
 			bInOutShouldAllowListAllProperties = true;
 		}
@@ -240,7 +251,7 @@ const FNamePermissionList& FPropertyEditorPermissionList::GetCachedPermissionLis
 				NewPermissionList.Append(Entry->PermissionList);
 			}
 
-			bInOutShouldAllowListAllProperties = Entry->Rules == EPropertyEditorPermissionListRules::AllowListAllProperties;
+			bInOutShouldAllowListAllProperties = Entry->Rules == EPropertyPermissionListRules::AllowListAllProperties;
 		}
 
 		// PermissionList all properties if the flag is set, the parent Struct has a PermissionList, and this Struct has no PermissionList
@@ -249,11 +260,11 @@ const FNamePermissionList& FPropertyEditorPermissionList::GetCachedPermissionLis
 		{
 			for (TFieldIterator<FProperty> Property(Struct, EFieldIteratorFlags::ExcludeSuper, EFieldIteratorFlags::ExcludeDeprecated); Property; ++Property)
 			{
-				NewPermissionList.AddAllowListItem(PropertyEditorPermissionListOwner, Property->GetFName());
+				NewPermissionList.AddAllowListItem(PropertyPermissionListOwner, Property->GetFName());
 			}
 		}
 
-		PermissionList = &CachedPropertyEditorPermissionList.Add(Struct, NewPermissionList);
+		PermissionList = &CachedPropertyPermissionList.Add(Struct, MoveTemp(NewPermissionList));
 	}
 	
 	// If this Struct has no PermissionList, then the ShouldPermissionListAllProperties rule just forwards its current value on to the next subclass.
@@ -261,15 +272,15 @@ const FNamePermissionList& FPropertyEditorPermissionList::GetCachedPermissionLis
 	// In this case, simply add a dummy entry to the Struct's PermissionList that (likely) won't ever collide with a real property name
 	if (!bIsThisAllowListEmpty)
 	{
-		bInOutShouldAllowListAllProperties = Entry->Rules == EPropertyEditorPermissionListRules::AllowListAllSubclassProperties;
+		bInOutShouldAllowListAllProperties = Entry->Rules == EPropertyPermissionListRules::AllowListAllSubclassProperties;
 	}
 
 	return *PermissionList;
 }
 
-bool FPropertyEditorPermissionList::IsSpecificPropertyAllowListed(const UStruct* ObjectStruct, FName PropertyName) const
+bool FPropertyPermissionList::IsSpecificPropertyAllowListed(const UStruct* ObjectStruct, FName PropertyName) const
 {
-	const FPropertyEditorPermissionListEntry* Entry = RawPropertyEditorPermissionList.Find(ObjectStruct);
+	const FPropertyPermissionListEntry* Entry = RawPropertyPermissionList.Find(ObjectStruct);
 	if (Entry)
 	{
 		return Entry->PermissionList.GetAllowList().Contains(PropertyName);
@@ -277,9 +288,9 @@ bool FPropertyEditorPermissionList::IsSpecificPropertyAllowListed(const UStruct*
 	return false;
 }
 
-bool FPropertyEditorPermissionList::IsSpecificPropertyDenyListed(const UStruct* ObjectStruct, FName PropertyName) const
+bool FPropertyPermissionList::IsSpecificPropertyDenyListed(const UStruct* ObjectStruct, FName PropertyName) const
 {
-	const FPropertyEditorPermissionListEntry* Entry = RawPropertyEditorPermissionList.Find(ObjectStruct);
+	const FPropertyPermissionListEntry* Entry = RawPropertyPermissionList.Find(ObjectStruct);
 	if (Entry)
 	{
 		return Entry->PermissionList.GetDenyList().Contains(PropertyName);
diff --git a/Engine/Source/Editor/PropertyEditor/Public/PropertyEditorPermissionList.h b/Engine/Source/Editor/PropertyEditor/Public/PropertyEditorPermissionList.h
--- a/Engine/Source/Editor/PropertyEditor/Public/PropertyEditorPermissionList.h
+++ b/Engine/Source/Editor/PropertyEditor/Public/PropertyEditorPermissionList.h
@@ -26,7 +26,7 @@ DECLARE_MULTICAST_DELEGATE_TwoParams(FPermissionListUpdated, TSoftObjectPtr<UStr
  * Note that because class C manually defines a PermissionList, it does not inherit the AllowListAllProperties rule from class B, while
  * class D does not define a PermissionList, so it does inherit the rule, causing all of class D's properties to also get added to the AllowList.
  */
-enum class EPropertyEditorPermissionListRules : uint8
+enum class EPropertyPermissionListRules : uint8
 {
 	// If no PermissionList is manually defined for this Struct, AllowList all properties from this Struct and its subclasses
 	AllowListAllProperties,
@@ -37,23 +37,17 @@ enum class EPropertyEditorPermissionListRules : uint8
 	UseExistingPermissionList
 };
 
-struct FPropertyEditorPermissionListEntry
+struct FPropertyPermissionListEntry
 {
     FNamePermissionList PermissionList;
-    EPropertyEditorPermissionListRules Rules = EPropertyEditorPermissionListRules::UseExistingPermissionList;
+	EPropertyPermissionListRules Rules = EPropertyPermissionListRules::UseExistingPermissionList;
 };
 
-class PROPERTYEDITOR_API FPropertyEditorPermissionList
+class PROPERTYEDITOR_API FPropertyPermissionList
 {
 public:
-	static FPropertyEditorPermissionList& Get()
-	{
-		static FPropertyEditorPermissionList PermissionList;
-		return PermissionList;
-	}
-
 	/** Add a set of rules for a specific base UStruct to determine which properties are visible in all details panels */
-	void AddPermissionList(TSoftObjectPtr<UStruct> Struct, const FNamePermissionList& PermissionList, EPropertyEditorPermissionListRules Rules = EPropertyEditorPermissionListRules::UseExistingPermissionList);
+	void AddPermissionList(TSoftObjectPtr<UStruct> Struct, const FNamePermissionList& PermissionList, EPropertyPermissionListRules Rules = EPropertyPermissionListRules::UseExistingPermissionList);
 	/** Remove a set of rules for a specific base UStruct to determine which properties are visible in all details panels */
 	void RemovePermissionList(TSoftObjectPtr<UStruct> Struct);
 	/** Remove all rules */
@@ -80,17 +74,19 @@ public:
 	FSimpleMulticastDelegate PermissionListEnabledDelegate;
 
 	/** Controls whether DoesPropertyPassFilter always returns true or performs property-based filtering. */
-	bool IsEnabled() const { return bEnablePropertyEditorPermissionList; }
-	/** Turn on or off the property editor PermissionList. DoesPropertyPassFilter will always return true if disabled. */
-	void SetEnabled(bool bEnable);
+	bool IsEnabled() const { return bEnablePermissionList; }
+	/** 
+	 * Turn on or off the property PermissionList. DoesPropertyPassFilter will always return true if disabled.
+	 * Optionally the change can be announced through a broadcast, which is typically recommended, but may not be needed
+	 * if the permission list is only temporary enabled, e.g. during saving.
+	 */
+	void SetEnabled(bool bEnable, bool bAnnounce = true);
 
-	/** Whether the Details View should show special menu entries to add/remove items in the PermissionList */
-	bool ShouldShowMenuEntries() const { return bShouldShowMenuEntries;}
-	/** Turn on or off menu entries to modify the PermissionList from a Details View */
-	void SetShouldShowMenuEntries(bool bShow) { bShouldShowMenuEntries = bShow; }
+	/** Checks if the provided struct has any entries in the allow, deny and deny all lists and returns true in that case. */
+	bool HasFiltering(const UStruct* ObjectStruct) const;
 
 	/**
-	 * Checks if a property passes the PermissionList/DenyList filtering specified by PropertyEditorPermissionLists
+	 * Checks if a property passes the PermissionList/DenyList filtering specified by PropertyPermissionLists
 	 * This should be relatively fast as it maintains a flattened cache of all inherited PermissionLists for every UStruct (which is generated lazily).
 	 */
 	bool DoesPropertyPassFilter(const UStruct* ObjectStruct, FName PropertyName) const;
@@ -101,34 +97,68 @@ public:
 	bool IsSpecificPropertyDenyListed(const UStruct* ObjectStruct, FName PropertyName) const;
 
 	/** Gets a read-only copy of the original, un-flattened PermissionList. */
-	const TMap<TSoftObjectPtr<UStruct>, FPropertyEditorPermissionListEntry>& GetRawPermissionList() const { return RawPropertyEditorPermissionList; }
+	const TMap<TSoftObjectPtr<UStruct>, FPropertyPermissionListEntry>& GetRawPermissionList() const { return RawPropertyPermissionList; }
 
-	/** Clear CachedPropertyEditorPermissionList to cause the PermissionListed property list to be regenerated next time it's queried */
+	/** Clear CachedPropertyPermissionList to cause the PermissionListed property list to be regenerated next time it's queried */
 	void ClearCache();
 
 	/** If true, PermissionListUpdatedDelegate will not broadcast when modifying the permission lists */
 	bool bSuppressUpdateDelegate = false;
 
-private:
-	FPropertyEditorPermissionList();
-	~FPropertyEditorPermissionList();
+protected:
+	FPropertyPermissionList();
+	~FPropertyPermissionList();
 
+private:
 	void RegisterOnBlueprintCompiled();
 
 	void ClearCacheAndBroadcast(TSoftObjectPtr<UStruct> ObjectStruct = nullptr, FName OwnerName = NAME_None);
 	
 	/** Whether DoesPropertyPassFilter should perform its PermissionList check or always return true */
-	bool bEnablePropertyEditorPermissionList = false;
-	/** Whether SDetailSingleItemRow should add menu items to add/remove properties to/from the PermissionList */
-	bool bShouldShowMenuEntries = false;
+	bool bEnablePermissionList = false;
 	
-	/** Stores assigned PermissionLists from AddPermissionList(), which are later flattened and stored in CachedPropertyEditorPermissionList. */
-	TMap<TSoftObjectPtr<UStruct>, FPropertyEditorPermissionListEntry> RawPropertyEditorPermissionList;
+	/** Stores assigned PermissionLists from AddPermissionList(), which are later flattened and stored in CachedPropertyPermissionList. */
+	TMap<TSoftObjectPtr<UStruct>, FPropertyPermissionListEntry> RawPropertyPermissionList;
 
 	/** Lazily-constructed combined cache of both the flattened class PermissionList and struct PermissionList */
-	mutable TMap<TWeakObjectPtr<const UStruct>, FNamePermissionList> CachedPropertyEditorPermissionList;
+	mutable TMap<TWeakObjectPtr<const UStruct>, FNamePermissionList> CachedPropertyPermissionList;
 
 	/** Get or create the cached PermissionList for a specific UStruct */
 	const FNamePermissionList& GetCachedPermissionListForStruct(const UStruct* Struct) const;
 	const FNamePermissionList& GetCachedPermissionListForStructHelper(const UStruct* Struct, bool& bInOutShouldAllowListAllProperties) const;
 };
+
+class PROPERTYEDITOR_API FPropertyEditorPermissionList : public FPropertyPermissionList
+{
+public:
+	static FPropertyEditorPermissionList& Get()
+	{
+		static FPropertyEditorPermissionList PermissionList;
+		return PermissionList;
+	}
+
+	/** Whether the Details View should show special menu entries to add/remove items in the PermissionList */
+	bool ShouldShowMenuEntries() const { return bShouldShowMenuEntries; }
+	/** Turn on or off menu entries to modify the PermissionList from a Details View */
+	void SetShouldShowMenuEntries(bool bShow) { bShouldShowMenuEntries = bShow; }
+
+private:
+	FPropertyEditorPermissionList() = default;
+	~FPropertyEditorPermissionList() = default;
+
+	/** Whether SDetailSingleItemRow should add menu items to add/remove properties to/from the PermissionList */
+	bool bShouldShowMenuEntries = false;
+};
+
+class PROPERTYEDITOR_API FHiddenPropertyPermissionList : public FPropertyPermissionList
+{
+public:
+	static FHiddenPropertyPermissionList& Get()
+	{
+		static FHiddenPropertyPermissionList PermissionList;
+		return PermissionList;
+	}
+private:
+	FHiddenPropertyPermissionList() = default;
+	~FHiddenPropertyPermissionList() = default;
+};
