Renamed PropertyEditorPermissionList to PropertyPermissionList.

diff --git a/Engine/Source/Editor/PropertyEditor/Private/PropertyEditorPermissionList.cpp b/Engine/Source/Editor/PropertyEditor/Private/PropertyEditorPermissionList.cpp
deleted file mode 100644
diff --git a/Engine/Source/Editor/PropertyEditor/Private/PropertyPermissionList.cpp b/Engine/Source/Editor/PropertyEditor/Private/PropertyPermissionList.cpp
new file mode 100644
--- /dev/null
+++ b/Engine/Source/Editor/PropertyEditor/Private/PropertyPermissionList.cpp
@@ -0,0 +1,299 @@
+// Copyright Epic Games, Inc. All Rights Reserved.
+
+#include "PropertyEditorPermissionList.h"
+
+#include "Editor.h"
+#include "Misc/CoreDelegates.h"
+#include "UObject/UnrealType.h"
+
+DEFINE_LOG_CATEGORY(LogPropertyEditorPermissionList);
+
+namespace
+{
+	const FName PropertyPermissionListOwner = "PropertyPermissionList";
+}
+
+FPropertyPermissionList::FPropertyPermissionList()
+{
+	if (GEditor)
+	{
+		RegisterOnBlueprintCompiled();
+	}
+	else
+	{
+		FCoreDelegates::OnPostEngineInit.AddRaw(this, &FPropertyPermissionList::RegisterOnBlueprintCompiled);
+	}
+}
+
+void FPropertyPermissionList::RegisterOnBlueprintCompiled()
+{
+	if (ensure(GEditor))
+	{
+		GEditor->OnBlueprintCompiled().AddRaw(this, &FPropertyPermissionList::ClearCache);
+	}
+}
+
+FPropertyPermissionList::~FPropertyPermissionList()
+{
+	if (GEditor)
+	{
+		GEditor->OnBlueprintCompiled().RemoveAll(this);
+	}
+}
+
+void FPropertyPermissionList::ClearCacheAndBroadcast(TSoftObjectPtr<UStruct> ObjectStruct, FName OwnerName)
+{
+	// The cache isn't too expensive to recompute, so it is cleared
+	// and lazily repopulated any time the raw PermissionList changes.
+	ClearCache();
+
+	if (!bSuppressUpdateDelegate)
+	{
+		PermissionListUpdatedDelegate.Broadcast(ObjectStruct, OwnerName);
+	}
+}
+
+void FPropertyPermissionList::AddPermissionList(TSoftObjectPtr<UStruct> Struct, const FNamePermissionList& PermissionList, EPropertyPermissionListRules Rules)
+{
+	FPropertyPermissionListEntry& Entry = RawPropertyPermissionList.FindOrAdd(Struct);
+	Entry.PermissionList = PermissionList;
+	// Always use the most permissive rule previously set
+	if (Entry.Rules > Rules)
+	{
+		Entry.Rules = Rules;
+	}
+
+	ClearCacheAndBroadcast(Struct);
+}
+
+void FPropertyPermissionList::RemovePermissionList(TSoftObjectPtr<UStruct> Struct)
+{
+	if (RawPropertyPermissionList.Remove(Struct) > 0)
+	{
+		ClearCacheAndBroadcast(Struct);
+	}
+}
+
+void FPropertyPermissionList::ClearPermissionList()
+{
+	TArray<TSoftObjectPtr<UStruct>> Keys;
+	RawPropertyPermissionList.Reset();
+	ClearCacheAndBroadcast();
+}
+
+void FPropertyPermissionList::UnregisterOwner(const FName Owner)
+{
+	TArray<TSoftObjectPtr<UStruct>> StructsToRemove;
+
+	for (TPair<TSoftObjectPtr<UStruct>, FPropertyPermissionListEntry>& Pair : RawPropertyPermissionList)
+	{
+		Pair.Value.PermissionList.UnregisterOwner(Owner);
+		if (Pair.Value.PermissionList.GetOwnerNames().Num() == 0)
+		{
+			StructsToRemove.Add(Pair.Key);
+		}
+	}
+
+	{
+		TGuardValue<bool> SuppressGuard(bSuppressUpdateDelegate, true);
+		for (TSoftObjectPtr<UStruct>& StructToRemove : StructsToRemove)
+		{
+			RemovePermissionList(StructToRemove);
+		}
+	}
+
+	ClearCacheAndBroadcast(nullptr, Owner);
+}
+
+void FPropertyPermissionList::AddToAllowList(TSoftObjectPtr<UStruct> Struct, const FName PropertyName, const FName Owner)
+{
+	FPropertyPermissionListEntry& Entry = RawPropertyPermissionList.FindOrAdd(Struct);
+	if (Entry.PermissionList.AddAllowListItem(Owner, PropertyName))
+	{
+		ClearCacheAndBroadcast(Struct, Owner);
+	}
+}
+
+void FPropertyPermissionList::AddToAllowList(TSoftObjectPtr<UStruct> Struct, const TArray<FName>& PropertyNames, const FName Owner)
+{
+	FPropertyPermissionListEntry& Entry = RawPropertyPermissionList.FindOrAdd(Struct);
+	bool bAddedItem = false;
+	for (const FName& PropertyName : PropertyNames)
+	{
+		if (Entry.PermissionList.AddAllowListItem(Owner, PropertyName))
+		{
+			bAddedItem = true;
+		}
+	}
+
+	if (bAddedItem)
+	{
+		ClearCacheAndBroadcast(Struct, Owner);
+	}
+}
+
+void FPropertyPermissionList::RemoveFromAllowList(TSoftObjectPtr<UStruct> Struct, const FName PropertyName, const FName Owner)
+{
+	FPropertyPermissionListEntry& Entry = RawPropertyPermissionList.FindOrAdd(Struct);
+	if (Entry.PermissionList.RemoveAllowListItem(Owner, PropertyName))
+	{
+		ClearCacheAndBroadcast(Struct, Owner);
+	}
+}
+
+void FPropertyPermissionList::AddToDenyList(TSoftObjectPtr<UStruct> Struct, const FName PropertyName, const FName Owner)
+{
+	FPropertyPermissionListEntry& Entry = RawPropertyPermissionList.FindOrAdd(Struct);
+	if (Entry.PermissionList.AddDenyListItem(Owner, PropertyName))
+	{
+		ClearCacheAndBroadcast(Struct, Owner);
+	}
+}
+
+void FPropertyPermissionList::RemoveFromDenyList(TSoftObjectPtr<UStruct> Struct, const FName PropertyName, const FName Owner)
+{
+	FPropertyPermissionListEntry& Entry = RawPropertyPermissionList.FindOrAdd(Struct);
+	if (Entry.PermissionList.RemoveDenyListItem(Owner, PropertyName))
+	{
+		ClearCacheAndBroadcast(Struct, Owner);
+	}
+}
+
+void FPropertyPermissionList::SetEnabled(bool bEnable, bool bAnnounce)
+{
+	if (bEnablePermissionList != bEnable)
+	{
+		bEnablePermissionList = bEnable;
+		if (bAnnounce)
+		{
+			PermissionListEnabledDelegate.Broadcast();
+		}
+	}
+}
+
+void FPropertyPermissionList::ClearCache()
+{
+	CachedPropertyPermissionList.Reset();
+}
+
+bool FPropertyPermissionList::HasFiltering(const UStruct* ObjectStruct) const
+{
+	return ObjectStruct ? GetCachedPermissionListForStruct(ObjectStruct).HasFiltering() : false;
+}
+
+bool FPropertyPermissionList::DoesPropertyPassFilter(const UStruct* ObjectStruct, FName PropertyName) const
+{
+	if (bEnablePermissionList && ObjectStruct)
+	{
+		return GetCachedPermissionListForStruct(ObjectStruct).PassesFilter(PropertyName);
+	}
+	return true;
+}
+
+const FNamePermissionList& FPropertyPermissionList::GetCachedPermissionListForStruct(const UStruct* Struct) const
+{
+	const FNamePermissionList* CachedPermissionList = CachedPropertyPermissionList.Find(Struct);
+	if (CachedPermissionList)
+	{
+		return *CachedPermissionList;
+	}
+
+	// Default value doesn't matter since it's a no-op until the first PermissionList is encountered, at which
+	// point the rules will re-assign the value.
+	bool bShouldPermissionListAllProperties = false;
+	return GetCachedPermissionListForStructHelper(Struct, bShouldPermissionListAllProperties);
+}
+
+const FNamePermissionList& FPropertyPermissionList::GetCachedPermissionListForStructHelper(const UStruct* Struct, bool& bInOutShouldAllowListAllProperties) const
+{
+	check(Struct);
+
+	const FPropertyPermissionListEntry* Entry = RawPropertyPermissionList.Find(Struct);
+	// If an entry is set to AllowListAll, then treat it as if it has no manually-defined properties.
+	const bool bIsThisAllowListEmpty = Entry ? (Entry->PermissionList.GetAllowList().Num() == 0 || Entry->Rules == EPropertyPermissionListRules::AllowListAllProperties) : true;
+
+	// Normally this case would be caught in GetCachedPermissionListForStruct, but when being called recursively from a subclass
+	// we still need to update bInOutShouldAllowListAllProperties so that new PermissionLists cache properly.
+	FNamePermissionList* PermissionList = CachedPropertyPermissionList.Find(Struct);
+	if (PermissionList)
+	{
+		// Same check as below, but it specifically has to come after the recursive call so we have to duplicate the code here
+		if (Entry && Entry->Rules == EPropertyPermissionListRules::AllowListAllProperties)
+		{
+			bInOutShouldAllowListAllProperties = true;
+		}
+	}
+	else
+	{
+		FNamePermissionList NewPermissionList;
+
+		UStruct* SuperStruct = Struct->GetSuperStruct();
+		// Recursively fill the cache for all parent structs
+		if (SuperStruct)
+		{
+			NewPermissionList.Append(GetCachedPermissionListForStructHelper(SuperStruct, bInOutShouldAllowListAllProperties));
+		}
+
+		// Append this struct's PermissionList on top of the parent's PermissionList
+		if (Entry)
+		{
+			if (bIsThisAllowListEmpty)
+			{
+				// If the AllowList is empty, we only want to append the DenyLists
+				FNamePermissionList DuplicatePermissionList = Entry->PermissionList;
+				// Hack to get around the fact that there's no easy way to only clear an AllowList
+				TMap<FName, FPermissionListOwners>& AllowList = const_cast<TMap<FName, FPermissionListOwners>&>(DuplicatePermissionList.GetAllowList());
+				AllowList.Empty();
+				NewPermissionList.Append(DuplicatePermissionList);
+			}
+			else
+			{
+				NewPermissionList.Append(Entry->PermissionList);
+			}
+
+			bInOutShouldAllowListAllProperties = Entry->Rules == EPropertyPermissionListRules::AllowListAllProperties;
+		}
+
+		// PermissionList all properties if the flag is set, the parent Struct has a PermissionList, and this Struct has no PermissionList
+		// If the parent Struct's PermissionList is empty then that already implies all properties are visible
+		if (bInOutShouldAllowListAllProperties && NewPermissionList.GetAllowList().Num() > 0)
+		{
+			for (TFieldIterator<FProperty> Property(Struct, EFieldIteratorFlags::ExcludeSuper, EFieldIteratorFlags::ExcludeDeprecated); Property; ++Property)
+			{
+				NewPermissionList.AddAllowListItem(PropertyPermissionListOwner, Property->GetFName());
+			}
+		}
+
+		PermissionList = &CachedPropertyPermissionList.Add(Struct, MoveTemp(NewPermissionList));
+	}
+	
+	// If this Struct has no PermissionList, then the ShouldPermissionListAllProperties rule just forwards its current value on to the next subclass.
+	// This causes an issue in the case where a Struct should have no PermissionListed properties but wants to PermissionList all subclass properties.
+	// In this case, simply add a dummy entry to the Struct's PermissionList that (likely) won't ever collide with a real property name
+	if (!bIsThisAllowListEmpty)
+	{
+		bInOutShouldAllowListAllProperties = Entry->Rules == EPropertyPermissionListRules::AllowListAllSubclassProperties;
+	}
+
+	return *PermissionList;
+}
+
+bool FPropertyPermissionList::IsSpecificPropertyAllowListed(const UStruct* ObjectStruct, FName PropertyName) const
+{
+	const FPropertyPermissionListEntry* Entry = RawPropertyPermissionList.Find(ObjectStruct);
+	if (Entry)
+	{
+		return Entry->PermissionList.GetAllowList().Contains(PropertyName);
+	}
+	return false;
+}
+
+bool FPropertyPermissionList::IsSpecificPropertyDenyListed(const UStruct* ObjectStruct, FName PropertyName) const
+{
+	const FPropertyPermissionListEntry* Entry = RawPropertyPermissionList.Find(ObjectStruct);
+	if (Entry)
+	{
+		return Entry->PermissionList.GetDenyList().Contains(PropertyName);
+	}
+	return false;
+}
