Renamed PropertyEditorPermissionList to PropertyPermissionList.

diff --git a/Engine/Plugins/Runtime/ModelViewViewModel/Source/ModelViewViewModelBlueprint/Private/MVVMDeveloperProjectSettings.cpp b/Engine/Plugins/Runtime/ModelViewViewModel/Source/ModelViewViewModelBlueprint/Private/MVVMDeveloperProjectSettings.cpp
--- a/Engine/Plugins/Runtime/ModelViewViewModel/Source/ModelViewViewModelBlueprint/Private/MVVMDeveloperProjectSettings.cpp
+++ b/Engine/Plugins/Runtime/ModelViewViewModel/Source/ModelViewViewModelBlueprint/Private/MVVMDeveloperProjectSettings.cpp
@@ -5,7 +5,7 @@
 
 #include "BlueprintEditorSettings.h"
 #include "MVVMBlueprintViewModelContext.h"
-#include "PropertyEditorPermissionList.h"
+#include "PropertyPermissionList.h"
 #include "Types/MVVMExecutionMode.h"
 
 #define LOCTEXT_NAMESPACE "MVVMDeveloperProjectSettings"
diff --git a/Engine/Plugins/Runtime/ModelViewViewModel/Source/ModelViewViewModelBlueprint/Private/MVVMViewBlueprintCompiler.cpp b/Engine/Plugins/Runtime/ModelViewViewModel/Source/ModelViewViewModelBlueprint/Private/MVVMViewBlueprintCompiler.cpp
--- a/Engine/Plugins/Runtime/ModelViewViewModel/Source/ModelViewViewModelBlueprint/Private/MVVMViewBlueprintCompiler.cpp
+++ b/Engine/Plugins/Runtime/ModelViewViewModel/Source/ModelViewViewModelBlueprint/Private/MVVMViewBlueprintCompiler.cpp
@@ -8,7 +8,7 @@
 #include "HAL/IConsoleManager.h"
 #include "MVVMBlueprintView.h"
 #include "MVVMDeveloperProjectSettings.h"
-#include "PropertyEditorPermissionList.h"
+#include "PropertyPermissionList.h"
 #include "MVVMFunctionGraphHelper.h"
 #include "Templates/ValueOrError.h"
 #include "Types/MVVMBindingName.h"
diff --git a/Engine/Plugins/Runtime/ModelViewViewModel/Source/ModelViewViewModelEditor/Private/Widgets/SMVVMViewModelBindingListWidget.cpp b/Engine/Plugins/Runtime/ModelViewViewModel/Source/ModelViewViewModelEditor/Private/Widgets/SMVVMViewModelBindingListWidget.cpp
--- a/Engine/Plugins/Runtime/ModelViewViewModel/Source/ModelViewViewModelEditor/Private/Widgets/SMVVMViewModelBindingListWidget.cpp
+++ b/Engine/Plugins/Runtime/ModelViewViewModel/Source/ModelViewViewModelEditor/Private/Widgets/SMVVMViewModelBindingListWidget.cpp
@@ -8,7 +8,7 @@
 #include "MVVMBlueprintViewModelContext.h"
 #include "MVVMDeveloperProjectSettings.h"
 #include "MVVMSubsystem.h"
-#include "PropertyEditorPermissionList.h"
+#include "PropertyPermissionList.h"
 #include "Styling/MVVMEditorStyle.h"
 #include "Types/MVVMAvailableBinding.h"
 #include "Types/MVVMBindingSource.h"
diff --git a/Engine/Source/Editor/BlueprintGraph/Private/BlueprintActionDatabase.cpp b/Engine/Source/Editor/BlueprintGraph/Private/BlueprintActionDatabase.cpp
--- a/Engine/Source/Editor/BlueprintGraph/Private/BlueprintActionDatabase.cpp
+++ b/Engine/Source/Editor/BlueprintGraph/Private/BlueprintActionDatabase.cpp
@@ -85,7 +85,7 @@
 #include "Misc/NamePermissionList.h"
 #include "Modules/ModuleManager.h"
 #include "ProfilingDebugging/CpuProfilerTrace.h"
-#include "PropertyEditorPermissionList.h"
+#include "PropertyPermissionList.h"
 #include "Templates/Casts.h"
 #include "Templates/ChooseClass.h"
 #include "Templates/SubclassOf.h"
diff --git a/Engine/Source/Editor/PropertyEditor/Private/DetailCategoryBuilderImpl.cpp b/Engine/Source/Editor/PropertyEditor/Private/DetailCategoryBuilderImpl.cpp
--- a/Engine/Source/Editor/PropertyEditor/Private/DetailCategoryBuilderImpl.cpp
+++ b/Engine/Source/Editor/PropertyEditor/Private/DetailCategoryBuilderImpl.cpp
@@ -16,7 +16,7 @@
 #include "ObjectPropertyNode.h"
 #include "PropertyCustomizationHelpers.h"
 #include "PropertyEditorModule.h"
-#include "PropertyEditorPermissionList.h"
+#include "PropertyPermissionList.h"
 #include "SDetailCategoryTableRow.h"
 #include "SDetailSingleItemRow.h"
 #include "StructurePropertyNode.h"
diff --git a/Engine/Source/Editor/PropertyEditor/Private/DetailItemNode.cpp b/Engine/Source/Editor/PropertyEditor/Private/DetailItemNode.cpp
--- a/Engine/Source/Editor/PropertyEditor/Private/DetailItemNode.cpp
+++ b/Engine/Source/Editor/PropertyEditor/Private/DetailItemNode.cpp
@@ -9,7 +9,7 @@
 #include "ObjectPropertyNode.h"
 #include "SConstrainedBox.h"
 #include "PropertyCustomizationHelpers.h"
-#include "PropertyEditorPermissionList.h"
+#include "PropertyPermissionList.h"
 #include "SDetailCategoryTableRow.h"
 #include "SDetailSingleItemRow.h"
 
diff --git a/Engine/Source/Editor/PropertyEditor/Private/PropertyPermissionList.cpp b/Engine/Source/Editor/PropertyEditor/Private/PropertyPermissionList.cpp
--- a/Engine/Source/Editor/PropertyEditor/Private/PropertyPermissionList.cpp
+++ b/Engine/Source/Editor/PropertyEditor/Private/PropertyPermissionList.cpp
@@ -1,6 +1,6 @@
 // Copyright Epic Games, Inc. All Rights Reserved.
 
-#include "PropertyEditorPermissionList.h"
+#include "PropertyPermissionList.h"
 
 #include "Editor.h"
 #include "Misc/CoreDelegates.h"
diff --git a/Engine/Source/Editor/PropertyEditor/Private/SDetailSingleItemRow.cpp b/Engine/Source/Editor/PropertyEditor/Private/SDetailSingleItemRow.cpp
--- a/Engine/Source/Editor/PropertyEditor/Private/SDetailSingleItemRow.cpp
+++ b/Engine/Source/Editor/PropertyEditor/Private/SDetailSingleItemRow.cpp
@@ -25,7 +25,7 @@
 #include "Widgets/Images/SImage.h"
 #include "Widgets/Input/SButton.h"
 #include "Widgets/Input/SComboButton.h"
-#include "PropertyEditorPermissionList.h"
+#include "PropertyPermissionList.h"
 
 namespace DetailWidgetConstants
 {
diff --git a/Engine/Source/Editor/PropertyEditor/Private/SDetailsViewBase.cpp b/Engine/Source/Editor/PropertyEditor/Private/SDetailsViewBase.cpp
--- a/Engine/Source/Editor/PropertyEditor/Private/SDetailsViewBase.cpp
+++ b/Engine/Source/Editor/PropertyEditor/Private/SDetailsViewBase.cpp
@@ -26,7 +26,7 @@
 #include "ScopedTransaction.h"
 #include "SDetailNameArea.h"
 #include "SDetailsView.h"
-#include "PropertyEditorPermissionList.h"
+#include "PropertyPermissionList.h"
 #include "ThumbnailRendering/ThumbnailManager.h"
 
 SDetailsViewBase::SDetailsViewBase() :
diff --git a/Engine/Source/Editor/PropertyEditor/Public/PropertyEditorPermissionList.h b/Engine/Source/Editor/PropertyEditor/Public/PropertyEditorPermissionList.h
--- a/Engine/Source/Editor/PropertyEditor/Public/PropertyEditorPermissionList.h
+++ b/Engine/Source/Editor/PropertyEditor/Public/PropertyEditorPermissionList.h
@@ -2,163 +2,7 @@
 
 #pragma once
 
-#include "CoreMinimal.h"
-#include "UObject/SoftObjectPtr.h"
-#include "UObject/WeakObjectPtr.h"
-#include "Misc/NamePermissionList.h"
+// Deprecated in UE 5.3
+#pragma message(__FILE__"(6): warning: Header has been moved. #include PropertyPermissionList.h instead of PropertyEditorPermissionList.h")
 
-DECLARE_LOG_CATEGORY_EXTERN(LogPropertyEditorPermissionList, Log, All);
-
-/** Struct, OwnerName */
-DECLARE_MULTICAST_DELEGATE_TwoParams(FPermissionListUpdated, TSoftObjectPtr<UStruct>, FName);
-
-/**
- * A hierarchical set of rules that can be used to PermissionList all properties of specific Structs without
- * having to manually add every single property in those Structs. These rules are applied in order from
- * the base Struct to the leaf Struct. UseExistingPermissionList has dual-functionality to alternatively
- * inherit the parent Struct's rule if no PermissionList is manually defined.
- * 
- * For example, if you have:
- * class A - (UseExistingPermissionList "MyProp")						PermissionList = "MyProp"
- * class B : public class A - (AllowListListAllProperties)				PermissionList = "MyProp","PropA1","PropA2"
- * class C : public class B - (UseExistingPermissionList "AnotherProp")	PermissionList = "MyProp","PropA1","PropA2","AnotherProp"
- * class D : public class B - (UseExistingPermissionList)				PermissionList = "MyProp","PropA1","PropA2","PropD1","PropD2"
- * Note that because class C manually defines a PermissionList, it does not inherit the AllowListAllProperties rule from class B, while
- * class D does not define a PermissionList, so it does inherit the rule, causing all of class D's properties to also get added to the AllowList.
- */
-enum class EPropertyPermissionListRules : uint8
-{
-	// If no PermissionList is manually defined for this Struct, AllowList all properties from this Struct and its subclasses
-	AllowListAllProperties,
-	// If a PermissionList is manually defined for this Struct, AllowList all properties from this Struct's subclasses.
-	// If this functionality is needed without any properties to AllowList, a fake property must be added to AllowList instead.
-	AllowListAllSubclassProperties,
-	// If a PermissionList is manually defined for this struct, PermissionList those properties. Otherwise, use the parent Struct's rule.
-	UseExistingPermissionList
-};
-
-struct FPropertyPermissionListEntry
-{
-    FNamePermissionList PermissionList;
-	EPropertyPermissionListRules Rules = EPropertyPermissionListRules::UseExistingPermissionList;
-};
-
-class PROPERTYEDITOR_API FPropertyPermissionList
-{
-public:
-	/** Add a set of rules for a specific base UStruct to determine which properties are visible in all details panels */
-	void AddPermissionList(TSoftObjectPtr<UStruct> Struct, const FNamePermissionList& PermissionList, EPropertyPermissionListRules Rules = EPropertyPermissionListRules::UseExistingPermissionList);
-	/** Remove a set of rules for a specific base UStruct to determine which properties are visible in all details panels */
-	void RemovePermissionList(TSoftObjectPtr<UStruct> Struct);
-	/** Remove all rules */
-	void ClearPermissionList();
-
-	/** Unregister an owner from all permission lists currently stored */
-	void UnregisterOwner(const FName Owner);
-
-	/** Add a specific property to a UStruct's AllowList */
-	void AddToAllowList(TSoftObjectPtr<UStruct> Struct, const FName PropertyName, const FName Owner = NAME_None);
-	/** Add a list of properties to a UStruct's AllowList */
-	void AddToAllowList(TSoftObjectPtr<UStruct> Struct, const TArray<FName>& PropertyNames, const FName Owner = NAME_None);
-	/** Remove a specific property from a UStruct's AllowList */
-	void RemoveFromAllowList(TSoftObjectPtr<UStruct> Struct, const FName PropertyName, const FName Owner = NAME_None);
-	/** Add a specific property to a UStruct's DenyList */
-	void AddToDenyList(TSoftObjectPtr<UStruct> Struct, const FName PropertyName, const FName Owner = NAME_None);
-	/** Remove a specific property from a UStruct's DenyList */
-    void RemoveFromDenyList(TSoftObjectPtr<UStruct> Struct, const FName PropertyName, const FName Owner = NAME_None);
-
-	/** When the PermissionList or DenyList for any struct was added to or removed from. */
-    FPermissionListUpdated PermissionListUpdatedDelegate;
-    
-	/** When the entire PermissionList is enabled or disabled */
-	FSimpleMulticastDelegate PermissionListEnabledDelegate;
-
-	/** Controls whether DoesPropertyPassFilter always returns true or performs property-based filtering. */
-	bool IsEnabled() const { return bEnablePermissionList; }
-	/** 
-	 * Turn on or off the property PermissionList. DoesPropertyPassFilter will always return true if disabled.
-	 * Optionally the change can be announced through a broadcast, which is typically recommended, but may not be needed
-	 * if the permission list is only temporary enabled, e.g. during saving.
-	 */
-	void SetEnabled(bool bEnable, bool bAnnounce = true);
-
-	/** Checks if the provided struct has any entries in the allow, deny and deny all lists and returns true in that case. */
-	bool HasFiltering(const UStruct* ObjectStruct) const;
-
-	/**
-	 * Checks if a property passes the PermissionList/DenyList filtering specified by PropertyPermissionLists
-	 * This should be relatively fast as it maintains a flattened cache of all inherited PermissionLists for every UStruct (which is generated lazily).
-	 */
-	bool DoesPropertyPassFilter(const UStruct* ObjectStruct, FName PropertyName) const;
-
-	/** Check whether a property exists on the PermissionList for a specific Struct - this will return false if the property is AllowListed on a parent Struct */
-	bool IsSpecificPropertyAllowListed(const UStruct* ObjectStruct, FName PropertyName) const;
-	/** Check whether a property exists on the DenyList for a specific Struct - this will return false if the property is DenyListed on a parent Struct */
-	bool IsSpecificPropertyDenyListed(const UStruct* ObjectStruct, FName PropertyName) const;
-
-	/** Gets a read-only copy of the original, un-flattened PermissionList. */
-	const TMap<TSoftObjectPtr<UStruct>, FPropertyPermissionListEntry>& GetRawPermissionList() const { return RawPropertyPermissionList; }
-
-	/** Clear CachedPropertyPermissionList to cause the PermissionListed property list to be regenerated next time it's queried */
-	void ClearCache();
-
-	/** If true, PermissionListUpdatedDelegate will not broadcast when modifying the permission lists */
-	bool bSuppressUpdateDelegate = false;
-
-protected:
-	FPropertyPermissionList();
-	~FPropertyPermissionList();
-
-private:
-	void RegisterOnBlueprintCompiled();
-
-	void ClearCacheAndBroadcast(TSoftObjectPtr<UStruct> ObjectStruct = nullptr, FName OwnerName = NAME_None);
-	
-	/** Whether DoesPropertyPassFilter should perform its PermissionList check or always return true */
-	bool bEnablePermissionList = false;
-	
-	/** Stores assigned PermissionLists from AddPermissionList(), which are later flattened and stored in CachedPropertyPermissionList. */
-	TMap<TSoftObjectPtr<UStruct>, FPropertyPermissionListEntry> RawPropertyPermissionList;
-
-	/** Lazily-constructed combined cache of both the flattened class PermissionList and struct PermissionList */
-	mutable TMap<TWeakObjectPtr<const UStruct>, FNamePermissionList> CachedPropertyPermissionList;
-
-	/** Get or create the cached PermissionList for a specific UStruct */
-	const FNamePermissionList& GetCachedPermissionListForStruct(const UStruct* Struct) const;
-	const FNamePermissionList& GetCachedPermissionListForStructHelper(const UStruct* Struct, bool& bInOutShouldAllowListAllProperties) const;
-};
-
-class PROPERTYEDITOR_API FPropertyEditorPermissionList : public FPropertyPermissionList
-{
-public:
-	static FPropertyEditorPermissionList& Get()
-	{
-		static FPropertyEditorPermissionList PermissionList;
-		return PermissionList;
-	}
-
-	/** Whether the Details View should show special menu entries to add/remove items in the PermissionList */
-	bool ShouldShowMenuEntries() const { return bShouldShowMenuEntries; }
-	/** Turn on or off menu entries to modify the PermissionList from a Details View */
-	void SetShouldShowMenuEntries(bool bShow) { bShouldShowMenuEntries = bShow; }
-
-private:
-	FPropertyEditorPermissionList() = default;
-	~FPropertyEditorPermissionList() = default;
-
-	/** Whether SDetailSingleItemRow should add menu items to add/remove properties to/from the PermissionList */
-	bool bShouldShowMenuEntries = false;
-};
-
-class PROPERTYEDITOR_API FHiddenPropertyPermissionList : public FPropertyPermissionList
-{
-public:
-	static FHiddenPropertyPermissionList& Get()
-	{
-		static FHiddenPropertyPermissionList PermissionList;
-		return PermissionList;
-	}
-private:
-	FHiddenPropertyPermissionList() = default;
-	~FHiddenPropertyPermissionList() = default;
-};
+#include "PropertyPermissionList.h"
\ No newline at end of file
diff --git a/Engine/Source/Editor/PropertyEditor/Public/PropertyPermissionList.h b/Engine/Source/Editor/PropertyEditor/Public/PropertyPermissionList.h
new file mode 100644
--- /dev/null
+++ b/Engine/Source/Editor/PropertyEditor/Public/PropertyPermissionList.h
@@ -0,0 +1,164 @@
+// Copyright Epic Games, Inc. All Rights Reserved.
+
+#pragma once
+
+#include "CoreMinimal.h"
+#include "UObject/SoftObjectPtr.h"
+#include "UObject/WeakObjectPtr.h"
+#include "Misc/NamePermissionList.h"
+
+DECLARE_LOG_CATEGORY_EXTERN(LogPropertyEditorPermissionList, Log, All);
+
+/** Struct, OwnerName */
+DECLARE_MULTICAST_DELEGATE_TwoParams(FPermissionListUpdated, TSoftObjectPtr<UStruct>, FName);
+
+/**
+ * A hierarchical set of rules that can be used to PermissionList all properties of specific Structs without
+ * having to manually add every single property in those Structs. These rules are applied in order from
+ * the base Struct to the leaf Struct. UseExistingPermissionList has dual-functionality to alternatively
+ * inherit the parent Struct's rule if no PermissionList is manually defined.
+ * 
+ * For example, if you have:
+ * class A - (UseExistingPermissionList "MyProp")						PermissionList = "MyProp"
+ * class B : public class A - (AllowListListAllProperties)				PermissionList = "MyProp","PropA1","PropA2"
+ * class C : public class B - (UseExistingPermissionList "AnotherProp")	PermissionList = "MyProp","PropA1","PropA2","AnotherProp"
+ * class D : public class B - (UseExistingPermissionList)				PermissionList = "MyProp","PropA1","PropA2","PropD1","PropD2"
+ * Note that because class C manually defines a PermissionList, it does not inherit the AllowListAllProperties rule from class B, while
+ * class D does not define a PermissionList, so it does inherit the rule, causing all of class D's properties to also get added to the AllowList.
+ */
+enum class EPropertyPermissionListRules : uint8
+{
+	// If no PermissionList is manually defined for this Struct, AllowList all properties from this Struct and its subclasses
+	AllowListAllProperties,
+	// If a PermissionList is manually defined for this Struct, AllowList all properties from this Struct's subclasses.
+	// If this functionality is needed without any properties to AllowList, a fake property must be added to AllowList instead.
+	AllowListAllSubclassProperties,
+	// If a PermissionList is manually defined for this struct, PermissionList those properties. Otherwise, use the parent Struct's rule.
+	UseExistingPermissionList
+};
+
+struct FPropertyPermissionListEntry
+{
+    FNamePermissionList PermissionList;
+	EPropertyPermissionListRules Rules = EPropertyPermissionListRules::UseExistingPermissionList;
+};
+
+class PROPERTYEDITOR_API FPropertyPermissionList
+{
+public:
+	/** Add a set of rules for a specific base UStruct to determine which properties are visible in all details panels */
+	void AddPermissionList(TSoftObjectPtr<UStruct> Struct, const FNamePermissionList& PermissionList, EPropertyPermissionListRules Rules = EPropertyPermissionListRules::UseExistingPermissionList);
+	/** Remove a set of rules for a specific base UStruct to determine which properties are visible in all details panels */
+	void RemovePermissionList(TSoftObjectPtr<UStruct> Struct);
+	/** Remove all rules */
+	void ClearPermissionList();
+
+	/** Unregister an owner from all permission lists currently stored */
+	void UnregisterOwner(const FName Owner);
+
+	/** Add a specific property to a UStruct's AllowList */
+	void AddToAllowList(TSoftObjectPtr<UStruct> Struct, const FName PropertyName, const FName Owner = NAME_None);
+	/** Add a list of properties to a UStruct's AllowList */
+	void AddToAllowList(TSoftObjectPtr<UStruct> Struct, const TArray<FName>& PropertyNames, const FName Owner = NAME_None);
+	/** Remove a specific property from a UStruct's AllowList */
+	void RemoveFromAllowList(TSoftObjectPtr<UStruct> Struct, const FName PropertyName, const FName Owner = NAME_None);
+	/** Add a specific property to a UStruct's DenyList */
+	void AddToDenyList(TSoftObjectPtr<UStruct> Struct, const FName PropertyName, const FName Owner = NAME_None);
+	/** Remove a specific property from a UStruct's DenyList */
+    void RemoveFromDenyList(TSoftObjectPtr<UStruct> Struct, const FName PropertyName, const FName Owner = NAME_None);
+
+	/** When the PermissionList or DenyList for any struct was added to or removed from. */
+    FPermissionListUpdated PermissionListUpdatedDelegate;
+    
+	/** When the entire PermissionList is enabled or disabled */
+	FSimpleMulticastDelegate PermissionListEnabledDelegate;
+
+	/** Controls whether DoesPropertyPassFilter always returns true or performs property-based filtering. */
+	bool IsEnabled() const { return bEnablePermissionList; }
+	/** 
+	 * Turn on or off the property PermissionList. DoesPropertyPassFilter will always return true if disabled.
+	 * Optionally the change can be announced through a broadcast, which is typically recommended, but may not be needed
+	 * if the permission list is only temporary enabled, e.g. during saving.
+	 */
+	void SetEnabled(bool bEnable, bool bAnnounce = true);
+
+	/** Checks if the provided struct has any entries in the allow, deny and deny all lists and returns true in that case. */
+	bool HasFiltering(const UStruct* ObjectStruct) const;
+
+	/**
+	 * Checks if a property passes the PermissionList/DenyList filtering specified by PropertyPermissionLists
+	 * This should be relatively fast as it maintains a flattened cache of all inherited PermissionLists for every UStruct (which is generated lazily).
+	 */
+	bool DoesPropertyPassFilter(const UStruct* ObjectStruct, FName PropertyName) const;
+
+	/** Check whether a property exists on the PermissionList for a specific Struct - this will return false if the property is AllowListed on a parent Struct */
+	bool IsSpecificPropertyAllowListed(const UStruct* ObjectStruct, FName PropertyName) const;
+	/** Check whether a property exists on the DenyList for a specific Struct - this will return false if the property is DenyListed on a parent Struct */
+	bool IsSpecificPropertyDenyListed(const UStruct* ObjectStruct, FName PropertyName) const;
+
+	/** Gets a read-only copy of the original, un-flattened PermissionList. */
+	const TMap<TSoftObjectPtr<UStruct>, FPropertyPermissionListEntry>& GetRawPermissionList() const { return RawPropertyPermissionList; }
+
+	/** Clear CachedPropertyPermissionList to cause the PermissionListed property list to be regenerated next time it's queried */
+	void ClearCache();
+
+	/** If true, PermissionListUpdatedDelegate will not broadcast when modifying the permission lists */
+	bool bSuppressUpdateDelegate = false;
+
+protected:
+	FPropertyPermissionList();
+	~FPropertyPermissionList();
+
+private:
+	void RegisterOnBlueprintCompiled();
+
+	void ClearCacheAndBroadcast(TSoftObjectPtr<UStruct> ObjectStruct = nullptr, FName OwnerName = NAME_None);
+	
+	/** Whether DoesPropertyPassFilter should perform its PermissionList check or always return true */
+	bool bEnablePermissionList = false;
+	
+	/** Stores assigned PermissionLists from AddPermissionList(), which are later flattened and stored in CachedPropertyPermissionList. */
+	TMap<TSoftObjectPtr<UStruct>, FPropertyPermissionListEntry> RawPropertyPermissionList;
+
+	/** Lazily-constructed combined cache of both the flattened class PermissionList and struct PermissionList */
+	mutable TMap<TWeakObjectPtr<const UStruct>, FNamePermissionList> CachedPropertyPermissionList;
+
+	/** Get or create the cached PermissionList for a specific UStruct */
+	const FNamePermissionList& GetCachedPermissionListForStruct(const UStruct* Struct) const;
+	const FNamePermissionList& GetCachedPermissionListForStructHelper(const UStruct* Struct, bool& bInOutShouldAllowListAllProperties) const;
+};
+
+class PROPERTYEDITOR_API FPropertyEditorPermissionList : public FPropertyPermissionList
+{
+public:
+	static FPropertyEditorPermissionList& Get()
+	{
+		static FPropertyEditorPermissionList PermissionList;
+		return PermissionList;
+	}
+
+	/** Whether the Details View should show special menu entries to add/remove items in the PermissionList */
+	bool ShouldShowMenuEntries() const { return bShouldShowMenuEntries; }
+	/** Turn on or off menu entries to modify the PermissionList from a Details View */
+	void SetShouldShowMenuEntries(bool bShow) { bShouldShowMenuEntries = bShow; }
+
+private:
+	FPropertyEditorPermissionList() = default;
+	~FPropertyEditorPermissionList() = default;
+
+	/** Whether SDetailSingleItemRow should add menu items to add/remove properties to/from the PermissionList */
+	bool bShouldShowMenuEntries = false;
+};
+
+class PROPERTYEDITOR_API FHiddenPropertyPermissionList : public FPropertyPermissionList
+{
+public:
+	static FHiddenPropertyPermissionList& Get()
+	{
+		static FHiddenPropertyPermissionList PermissionList;
+		return PermissionList;
+	}
+private:
+	FHiddenPropertyPermissionList() = default;
+	~FHiddenPropertyPermissionList() = default;
+};
diff --git a/Engine/Source/Editor/Sequencer/Private/SequencerObjectChangeListener.cpp b/Engine/Source/Editor/Sequencer/Private/SequencerObjectChangeListener.cpp
--- a/Engine/Source/Editor/Sequencer/Private/SequencerObjectChangeListener.cpp
+++ b/Engine/Source/Editor/Sequencer/Private/SequencerObjectChangeListener.cpp
@@ -6,7 +6,7 @@
 #include "Editor.h"
 #include "Modules/ModuleManager.h"
 #include "PropertyEditorModule.h"
-#include "PropertyEditorPermissionList.h"
+#include "PropertyPermissionList.h"
 #include "IPropertyChangeListener.h"
 #include "MovieSceneSequence.h"
 #include "ScopedTransaction.h"
