Submitting for Dan O'Conner
Add the ability to add customizations via property handle that disable object instancing logic associated with CPF_InstancedReference

#localization none
#tests all relevant operations on array and non-array editable versedevices in the details panel

diff --git a/Engine/Source/Editor/PropertyEditor/Private/DetailCategoryBuilderImpl.cpp b/Engine/Source/Editor/PropertyEditor/Private/DetailCategoryBuilderImpl.cpp
--- a/Engine/Source/Editor/PropertyEditor/Private/DetailCategoryBuilderImpl.cpp
+++ b/Engine/Source/Editor/PropertyEditor/Private/DetailCategoryBuilderImpl.cpp
@@ -379,9 +379,9 @@ IDetailPropertyRow& FDetailCategoryImpl::AddProperty(FName PropertyPath, UClass*
 	if (PropertyNode.IsValid())
 	{
 		ParentLayout->SetCustomProperty(PropertyNode);
+		NewCustomization.PropertyRow = MakeShareable(new FDetailPropertyRow(PropertyNode, AsShared()));
 	}
-
-	NewCustomization.PropertyRow = MakeShareable(new FDetailPropertyRow(PropertyNode, AsShared()));
+	
 	NewCustomization.bAdvanced = (Location == EPropertyLocation::Default) ? IsAdvancedLayout(NewCustomization) : (Location == EPropertyLocation::Advanced);
 
 	AddCustomLayout(NewCustomization);
@@ -602,6 +602,25 @@ void FDetailCategoryImpl::SetSortOrder(int32 InSortOrder)
 	SortOrder = InSortOrder;
 }
 
+void FDetailCategoryImpl::AddPropertyDisableInstancedReference(TSharedPtr<IPropertyHandle> PropertyHandle)
+{
+	FDetailLayoutCustomization NewCustomization;
+	NewCustomization.bCustom = true;
+
+	TSharedPtr<FDetailLayoutBuilderImpl> ParentLayout = GetParentLayoutImpl();
+	TSharedPtr<FPropertyNode> PropertyNode = ParentLayout->GetPropertyNode(PropertyHandle);
+	if (PropertyNode.IsValid())
+	{
+		ParentLayout->SetCustomProperty(PropertyNode);
+		PropertyNode->SetIgnoreInstancedReference();
+	}
+
+	NewCustomization.PropertyRow = MakeShareable(new FDetailPropertyRow(PropertyNode, AsShared()));
+	NewCustomization.bAdvanced = IsAdvancedLayout(NewCustomization);
+
+	AddCustomLayout(NewCustomization);
+}
+
 bool FDetailCategoryImpl::IsAdvancedDropdownEnabled() const
 {
 	return !bForceAdvanced;
diff --git a/Engine/Source/Editor/PropertyEditor/Private/DetailCategoryBuilderImpl.h b/Engine/Source/Editor/PropertyEditor/Private/DetailCategoryBuilderImpl.h
--- a/Engine/Source/Editor/PropertyEditor/Private/DetailCategoryBuilderImpl.h
+++ b/Engine/Source/Editor/PropertyEditor/Private/DetailCategoryBuilderImpl.h
@@ -185,6 +185,7 @@ public:
 	virtual void SetShowAdvanced(bool bShowAdvanced) override;
 	virtual int32 GetSortOrder() const override;
 	virtual void SetSortOrder(int32 InSortOrder) override;
+	virtual void AddPropertyDisableInstancedReference(TSharedPtr<IPropertyHandle> PropertyHandle) override;
 
 	/** FDetailTreeNode interface */
 	virtual IDetailsView* GetNodeDetailsView() const override { return GetDetailsView(); }
diff --git a/Engine/Source/Editor/PropertyEditor/Private/PropertyHandleImpl.cpp b/Engine/Source/Editor/PropertyEditor/Private/PropertyHandleImpl.cpp
--- a/Engine/Source/Editor/PropertyEditor/Private/PropertyHandleImpl.cpp
+++ b/Engine/Source/Editor/PropertyEditor/Private/PropertyHandleImpl.cpp
@@ -33,6 +33,11 @@
 
 #define LOCTEXT_NAMESPACE "PropertyHandleImplementation"
 
+static bool ShouldOwnInstance(const FObjectProperty* ObjectProperty, const FPropertyNode* ParentNode )
+{
+	return ObjectProperty && ParentNode && ObjectProperty->HasAnyPropertyFlags(CPF_InstancedReference) && !ParentNode->IsIgnoringInstancedReference();
+}
+
 bool IsTemplate(UObject* Obj)
 {
 	return Obj &&
@@ -1224,7 +1229,7 @@ void FPropertyValueImpl::ClearChildren()
 							// If the inner property is an instanced property we must move the old objects to the 
 							// transient package so code looking for objects of this type on the parent doesn't find them
 							FObjectProperty* InnerObjectProperty = CastField<FObjectProperty>(ArrayProperty->Inner);
-							if (InnerObjectProperty && InnerObjectProperty->HasAnyPropertyFlags(CPF_InstancedReference))
+							if (ShouldOwnInstance(InnerObjectProperty, PropertyNodePin.Get()))
 							{
 								const int32 ArraySize = ArrayHelper.Num();
 								for (int32 Index = 0; Index < ArraySize; ++Index)
@@ -1246,7 +1251,7 @@ void FPropertyValueImpl::ClearChildren()
 							// If the element property is an instanced property we must move the old objects to the 
 							// transient package so code looking for objects of this type on the parent doesn't find them
 							FObjectProperty* ElementObjectProperty = CastField<FObjectProperty>(SetProperty->ElementProp);
-							if (ElementObjectProperty && ElementObjectProperty->HasAnyPropertyFlags(CPF_InstancedReference))
+							if (ShouldOwnInstance(ElementObjectProperty, PropertyNodePin.Get()))
 							{
 								int32 ElementsToRemove = SetHelper.Num();
 								int32 Index = 0;
@@ -1274,7 +1279,7 @@ void FPropertyValueImpl::ClearChildren()
 							// If the map's value property is an instanced property we must move the old objects to the 
 							// transient package so code looking for objects of this type on the parent doesn't find them
 							FObjectProperty* ValueObjectProperty = CastField<FObjectProperty>(MapProperty->ValueProp);
-							if (ValueObjectProperty && ValueObjectProperty->HasAnyPropertyFlags(CPF_InstancedReference))
+							if (ShouldOwnInstance(ValueObjectProperty, PropertyNodePin.Get()))
 							{
 								int32 ElementsToRemove = MapHelper.Num();
 								int32 Index = 0;
@@ -1512,7 +1517,7 @@ void FPropertyValueImpl::DeleteChild( TSharedPtr<FPropertyNode> ChildNodeToDelet
 						// If the inner property is an instanced property we must move the old object to the 
 						// transient package so code looking for objects of this type on the parent doesn't find it
 						FObjectProperty* InnerObjectProperty = CastField<FObjectProperty>(ArrayProperty->Inner);
-						if (InnerObjectProperty && InnerObjectProperty->HasAnyPropertyFlags(CPF_InstancedReference))
+						if (ShouldOwnInstance(InnerObjectProperty, ParentNode))
 						{
 							if (UObject* InstancedObject = *reinterpret_cast<UObject**>(ArrayHelper.GetRawPtr(ChildNodePtr->GetArrayIndex())))
 							{
@@ -1531,7 +1536,7 @@ void FPropertyValueImpl::DeleteChild( TSharedPtr<FPropertyNode> ChildNodeToDelet
 						// If the element property is an instanced property we must move the old object to the 
 						// transient package so code looking for objects of this type on the parent doesn't find it
 						FObjectProperty* ElementObjectProperty = CastField<FObjectProperty>(SetProperty->ElementProp);
-						if (ElementObjectProperty && ElementObjectProperty->HasAnyPropertyFlags(CPF_InstancedReference))
+						if (ShouldOwnInstance(ElementObjectProperty, ParentNode))
 						{
 							if (UObject* InstancedObject = *reinterpret_cast<UObject**>(SetHelper.GetElementPtr(InternalIndex)))
 							{
@@ -1551,7 +1556,7 @@ void FPropertyValueImpl::DeleteChild( TSharedPtr<FPropertyNode> ChildNodeToDelet
 						// If the map's value property is an instanced property we must move the old object to the 
 						// transient package so code looking for objects of this type on the parent doesn't find it
 						FObjectProperty* ValueObjectProperty = CastField<FObjectProperty>(MapProperty->ValueProp);
-						if (ValueObjectProperty && ValueObjectProperty->HasAnyPropertyFlags(CPF_InstancedReference))
+						if (ShouldOwnInstance(ValueObjectProperty, ParentNode))
 						{
 							if (UObject* InstancedObject = *reinterpret_cast<UObject**>(MapHelper.GetValuePtr(InternalIndex)))
 							{
@@ -1650,7 +1655,7 @@ void FPropertyValueImpl::SwapChildren( TSharedPtr<FPropertyNode> FirstChildNode,
 					// If the inner property is an instanced property we must move the old object to the 
 					// transient package so code looking for objects of this type on the parent doesn't find it
 					FObjectProperty* InnerObjectProperty = CastField<FObjectProperty>(ArrayProperty->Inner);
-					if (InnerObjectProperty && InnerObjectProperty->HasAnyPropertyFlags(CPF_InstancedReference))
+					if(ShouldOwnInstance(InnerObjectProperty, ParentNode))
 					{
 						if (UObject* InstancedObject = *reinterpret_cast<UObject**>(ArrayHelper.GetRawPtr(FirstIndex)))
 						{
diff --git a/Engine/Source/Editor/PropertyEditor/Private/PropertyNode.cpp b/Engine/Source/Editor/PropertyEditor/Private/PropertyNode.cpp
--- a/Engine/Source/Editor/PropertyEditor/Private/PropertyNode.cpp
+++ b/Engine/Source/Editor/PropertyEditor/Private/PropertyNode.cpp
@@ -72,6 +72,7 @@ FPropertyNode::FPropertyNode()
 	, PropertyNodeFlags (EPropertyNodeFlags::NoFlags)
 	, bRebuildChildrenRequested( false )
 	, bChildrenRebuilt(false)
+	, bIgnoreInstancedReference(false)
 	, PropertyPath(TEXT(""))
 	, bIsEditConst(false)
 	, bUpdateEditConstState(true)
@@ -3068,6 +3069,16 @@ void FPropertyNode::SetExpandedChildPropertyNodes(const TSet<FString>& InNodesTo
 	} while (RecursiveStack.Num() > 0);
 }
 
+void FPropertyNode::SetIgnoreInstancedReference()
+{
+	bIgnoreInstancedReference = true;
+}
+
+bool FPropertyNode::IsIgnoringInstancedReference() const
+{
+	return bIgnoreInstancedReference;
+}
+
 FDelegateHandle FPropertyNode::SetOnRebuildChildren(const FSimpleDelegate& InOnRebuildChildren)
 {
 	return OnRebuildChildrenEvent.Add(InOnRebuildChildren);
diff --git a/Engine/Source/Editor/PropertyEditor/Private/PropertyNode.h b/Engine/Source/Editor/PropertyEditor/Private/PropertyNode.h
--- a/Engine/Source/Editor/PropertyEditor/Private/PropertyNode.h
+++ b/Engine/Source/Editor/PropertyEditor/Private/PropertyNode.h
@@ -1058,6 +1058,12 @@ public:
 	 */
 	const FString& GetPropertyPath() const { return PropertyPath; }
 
+	/** Marks this property node as ignoring CPF_InstancedReference */
+	void SetIgnoreInstancedReference();
+
+	/** Queries whether the node would like to ignore CPF_InstancedReference semantics */
+	bool IsIgnoringInstancedReference() const;
+
 protected:
 	TSharedRef<FEditPropertyChain> BuildPropertyChain( FProperty* PropertyAboutToChange ) const;
 	TSharedRef<FEditPropertyChain> BuildPropertyChain( FProperty* PropertyAboutToChange, const TSet<UObject*>& InAffectedArchetypeInstances ) const;
@@ -1221,6 +1227,9 @@ protected:
 	/** Set to true when RebuildChildren is called on the node */
 	bool bChildrenRebuilt;
 
+	/** Set to true when we want to ignore CPF_InstancedReference */
+	bool bIgnoreInstancedReference;
+
 	/** An array of restrictions limiting this property's potential values in property editors.*/
 	TArray<TSharedRef<const FPropertyRestriction>> Restrictions;
 
diff --git a/Engine/Source/Editor/PropertyEditor/Public/DetailCategoryBuilder.h b/Engine/Source/Editor/PropertyEditor/Public/DetailCategoryBuilder.h
--- a/Engine/Source/Editor/PropertyEditor/Public/DetailCategoryBuilder.h
+++ b/Engine/Source/Editor/PropertyEditor/Public/DetailCategoryBuilder.h
@@ -208,4 +208,7 @@ public:
 	* This is designed to be used for dynamic display of advanced properties.
 	*/
 	virtual void SetShowAdvanced(bool bShowAdvanced) = 0;
+
+	/** Add a property, but force it to behave as a normal, peer reference regardless of CPF_InstancedReference */
+	virtual void AddPropertyDisableInstancedReference(TSharedPtr<IPropertyHandle> PropertyHandle) = 0;
 };
