Fix bad merge
#rnx

diff --git a/Engine/Source/Editor/PropertyEditor/Private/PropertyPermissionList.cpp b/Engine/Source/Editor/PropertyEditor/Private/PropertyPermissionList.cpp
--- a/Engine/Source/Editor/PropertyEditor/Private/PropertyPermissionList.cpp
+++ b/Engine/Source/Editor/PropertyEditor/Private/PropertyPermissionList.cpp
@@ -25,6 +25,7 @@ uint32 GetTypeHash(const FPropertyPermissionList::FPermissionListUpdate& Permisi
 		GetTypeHash(PermisisonList.OwnerName));
 }
 
+FPropertyPermissionList::FPropertyPermissionList()
 {
 	if (GEditor)
 	{
