DetailsPanel: Fixed issue that prevented the details view from restoring property expansion states in cases where a details customization forced a refresh while the details view was processing deferred actions, such as when an array property is modified.

diff --git a/Engine/Source/Editor/PropertyEditor/Private/SDetailsViewBase.cpp b/Engine/Source/Editor/PropertyEditor/Private/SDetailsViewBase.cpp
--- a/Engine/Source/Editor/PropertyEditor/Private/SDetailsViewBase.cpp
+++ b/Engine/Source/Editor/PropertyEditor/Private/SDetailsViewBase.cpp
@@ -1052,7 +1052,12 @@ void SDetailsViewBase::Tick( const FGeometry& AllottedGeometry, const double InC
 				DeferredAction.ExecuteIfBound();
 			}
 		} while (DeferredActions.Num() > 0);
+	}
 
+	if (bHadDeferredActions)
+	{
+		// Restore expansion state after processing all deferred actions. Must be done outside of the scope where
+		// bRunningDeferredActions is true since RestoreAllExpandedItems returns immediately when it is true
 		RestoreAllExpandedItems();
 	}
 
