Moved FScriptMapHelper::HasKey() to editor.

diff --git a/Engine/Source/Editor/PropertyEditor/Private/PropertyHandleImpl.cpp b/Engine/Source/Editor/PropertyEditor/Private/PropertyHandleImpl.cpp
--- a/Engine/Source/Editor/PropertyEditor/Private/PropertyHandleImpl.cpp
+++ b/Engine/Source/Editor/PropertyEditor/Private/PropertyHandleImpl.cpp
@@ -361,9 +361,43 @@ FPropertyAccess::Result FPropertyValueImpl::ImportText( const TArray<FObjectBase
 				if (bIsInContainer)
 				{
 					uint8* ValueBaseAddress = ParentNode->GetValueBaseAddress(Cur.StructAddress, bIsSparseClassData);
-					
+
+					/**
+					 * Checks if a key in the map matches the specified key
+					 *
+					 * @param	InBaseAddress	The base address of the map
+					 * @param	InKeyValue		The key to find within the map
+					 *
+					 * @return	True if the key is found, false otherwise
+					 */
+					static auto HasKey = [](const FScriptMapHelper& Helper, void* InBaseAddress, const FString& InKeyValue)
+					{
+						FProperty* KeyProp = Helper.GetKeyProperty();
+						for (int32 Index = 0, ItemsLeft = Helper.Num(); ItemsLeft > 0; ++Index)
+						{
+							if (Helper.IsValidIndex(Index))
+							{
+								--ItemsLeft;
+
+								const uint8* PairPtr = Helper.GetPairPtr(Index);
+								const uint8* KeyPtr = KeyProp->ContainerPtrToValuePtr<const uint8>(PairPtr);
+
+								FString KeyValue;
+								if (KeyPtr != InBaseAddress && KeyProp->ExportText_Direct(KeyValue, KeyPtr, KeyPtr, nullptr, 0))
+								{
+									if ((CastField<FObjectProperty>(KeyProp) != nullptr && KeyValue.Contains(InKeyValue)) || InKeyValue == KeyValue)
+									{
+										return true;
+									}
+								}
+							}
+						}
+
+						return false;
+					};
+
 					FScriptMapHelper MapHelper(MapProperty, ValueBaseAddress);
-					if (MapHelper.HasKey(Cur.BaseAddress, NewValue) && 
+					if (HasKey(MapHelper, Cur.BaseAddress, NewValue) && 
 						(Flags & EPropertyValueSetFlags::InteractiveChange) == 0)
 					{
 						// Duplicate key in the map
diff --git a/Engine/Source/Runtime/CoreUObject/Public/UObject/UnrealType.h b/Engine/Source/Runtime/CoreUObject/Public/UObject/UnrealType.h
--- a/Engine/Source/Runtime/CoreUObject/Public/UObject/UnrealType.h
+++ b/Engine/Source/Runtime/CoreUObject/Public/UObject/UnrealType.h
@@ -5039,39 +5039,6 @@ public:
 		});
 	}
 
-	/**
-	 * Checks if a key in the map matches the specified key
-	 *
-	 * @param	InBaseAddress	The base address of the map
-	 * @param	InKeyValue		The key to find within the map
-	 *
-	 * @return	True if the key is found, false otherwise
-	 */
-	bool HasKey(const void* InBaseAddress, const FString& InKeyValue) const
-	{
-		for (int32 Index = 0, ItemsLeft = Num(); ItemsLeft > 0; ++Index)
-		{
-			if (IsValidIndex(Index))
-			{
-				--ItemsLeft;
-
-				const uint8* PairPtr = GetPairPtr(Index);
-				const uint8* KeyPtr = KeyProp->ContainerPtrToValuePtr<const uint8>(PairPtr);
-
-				FString KeyValue;
-				if (KeyPtr != InBaseAddress && KeyProp->ExportText_Direct(KeyValue, KeyPtr, KeyPtr, nullptr, 0))
-				{
-					if ((CastField<FObjectProperty>(KeyProp) != nullptr && KeyValue.Contains(InKeyValue)) || InKeyValue == KeyValue)
-					{
-						return true;
-					}
-				}
-			}
-		}
-
-		return false;
-	}
-
 	static FScriptMapHelper CreateHelperFormInnerProperties(FProperty* InKeyProperty, FProperty* InValProperty, const void *InMap, EMapPropertyFlags InMapFlags = EMapPropertyFlags::None)
 	{
 		return FScriptMapHelper(
