Moved FScriptSetHelper::HasElement() to editor.

diff --git a/Engine/Source/Editor/PropertyEditor/Private/PropertyHandleImpl.cpp b/Engine/Source/Editor/PropertyEditor/Private/PropertyHandleImpl.cpp
--- a/Engine/Source/Editor/PropertyEditor/Private/PropertyHandleImpl.cpp
+++ b/Engine/Source/Editor/PropertyEditor/Private/PropertyHandleImpl.cpp
@@ -343,8 +343,42 @@ FPropertyAccess::Result FPropertyValueImpl::ImportText( const TArray<FObjectBase
 				{
 					uint8* ValueBaseAddress = ParentNode->GetValueBaseAddress(Cur.StructAddress, bIsSparseClassData);
 					
+					/**
+					 * Checks if an element has already been added to the set
+					 *
+					 * @param	Helper			The set helper used to query the property.
+					 * @param	InBaseAddress	The base address of the set
+					 * @param	InElementValue	The element value to check for
+					 *
+					 * @return	True if the element is found in the set, false otherwise
+					 */
+					static auto HasElement = [](const FScriptSetHelper& Helper, void* InBaseAddress, const FString& InElementValue)
+					{
+						FProperty* ElementProp = Helper.GetElementProperty();
+						for (int32 Index = 0, ItemsLeft = Helper.Num(); ItemsLeft > 0; ++Index)
+						{
+							if (Helper.IsValidIndex(Index))
+							{
+								--ItemsLeft;
+
+								const uint8* Element = Helper.GetElementPtr(Index);
+
+								FString ElementValue;
+								if (Element != InBaseAddress && ElementProp->ExportText_Direct(ElementValue, Element, Element, nullptr, 0))
+								{
+									if ((CastField<FObjectProperty>(ElementProp) != nullptr && ElementValue.Contains(InElementValue)) || ElementValue == InElementValue)
+									{
+										return true;
+									}
+								}
+							}
+						}
+
+						return false;
+					};
+
 					FScriptSetHelper SetHelper(SetProperty, ValueBaseAddress);
-					if (SetHelper.HasElement(Cur.BaseAddress, NewValue) &&
+					if (HasElement(SetHelper, Cur.BaseAddress, NewValue) &&
 						(Flags & EPropertyValueSetFlags::InteractiveChange) == 0)
 					{
 						// Duplicate element in the set
@@ -365,6 +399,7 @@ FPropertyAccess::Result FPropertyValueImpl::ImportText( const TArray<FObjectBase
 					/**
 					 * Checks if a key in the map matches the specified key
 					 *
+					 * @param	Helper			The map helper used to query the property.
 					 * @param	InBaseAddress	The base address of the map
 					 * @param	InKeyValue		The key to find within the map
 					 *
diff --git a/Engine/Source/Runtime/CoreUObject/Public/UObject/UnrealType.h b/Engine/Source/Runtime/CoreUObject/Public/UObject/UnrealType.h
--- a/Engine/Source/Runtime/CoreUObject/Public/UObject/UnrealType.h
+++ b/Engine/Source/Runtime/CoreUObject/Public/UObject/UnrealType.h
@@ -5588,38 +5588,6 @@ public:
 		}
 	}
 
-	/**
-	 * Checks if an element has already been added to the set
-	 *
-	 * @param	InBaseAddress	The base address of the set
-	 * @param	InElementValue	The element value to check for
-	 *
-	 * @return	True if the element is found in the set, false otherwise
-	 */
-	bool HasElement(void* InBaseAddress, const FString& InElementValue) const
-	{
-		for (int32 Index = 0, ItemsLeft = Num(); ItemsLeft > 0; ++Index)
-		{
-			if (IsValidIndex(Index))
-			{
-				--ItemsLeft;
-
-				const uint8* Element = GetElementPtr(Index);
-
-				FString ElementValue;
-				if (Element != InBaseAddress && ElementProp->ExportText_Direct(ElementValue, Element, Element, nullptr, 0))
-				{
-					if ((CastField<FObjectProperty>(ElementProp) != nullptr && ElementValue.Contains(InElementValue)) || ElementValue == InElementValue)
-					{
-						return true;
-					}
-				}
-			}
-		}
-
-		return false;
-	}
-
 	static FScriptSetHelper CreateHelperFormElementProperty(FProperty* InElementProperty, const void *InSet)
 	{
 		check(InElementProperty);
