Fixed duplicates being wrongly reported when adding elements to map or set properties in the editor, where the key type is a UClass* key and the class name is a subset of a name which already exists as a key.

diff --git a/Engine/Source/Editor/PropertyEditor/Private/PropertyHandleImpl.cpp b/Engine/Source/Editor/PropertyEditor/Private/PropertyHandleImpl.cpp
--- a/Engine/Source/Editor/PropertyEditor/Private/PropertyHandleImpl.cpp
+++ b/Engine/Source/Editor/PropertyEditor/Private/PropertyHandleImpl.cpp
@@ -355,6 +355,13 @@ FPropertyAccess::Result FPropertyValueImpl::ImportText( const TArray<FObjectBase
 					static auto HasElement = [](const FScriptSetHelper& Helper, void* InBaseAddress, const FString& InElementValue)
 					{
 						FProperty* ElementProp = Helper.GetElementProperty();
+
+						void* TempElementStorage = ElementProp->AllocateAndInitializeValue();
+						ON_SCOPE_EXIT
+						{
+							ElementProp->DestroyAndFreeValue(TempElementStorage);
+						};
+
 						for (int32 Index = 0, ItemsLeft = Helper.Num(); ItemsLeft > 0; ++Index)
 						{
 							if (Helper.IsValidIndex(Index))
@@ -363,13 +370,9 @@ FPropertyAccess::Result FPropertyValueImpl::ImportText( const TArray<FObjectBase
 
 								const uint8* Element = Helper.GetElementPtr(Index);
 
-								FString ElementValue;
-								if (Element != InBaseAddress && ElementProp->ExportText_Direct(ElementValue, Element, Element, nullptr, 0))
+								if (Element != InBaseAddress && ElementProp->ImportText_Direct(*InElementValue, TempElementStorage, nullptr, 0) && ElementProp->Identical(Element, TempElementStorage))
 								{
-									if ((CastField<FObjectProperty>(ElementProp) != nullptr && ElementValue.Contains(InElementValue)) || ElementValue == InElementValue)
-									{
-										return true;
-									}
+									return true;
 								}
 							}
 						}
@@ -408,6 +411,13 @@ FPropertyAccess::Result FPropertyValueImpl::ImportText( const TArray<FObjectBase
 					static auto HasKey = [](const FScriptMapHelper& Helper, void* InBaseAddress, const FString& InKeyValue)
 					{
 						FProperty* KeyProp = Helper.GetKeyProperty();
+
+						void* TempKeyStorage = KeyProp->AllocateAndInitializeValue();
+						ON_SCOPE_EXIT
+						{
+							KeyProp->DestroyAndFreeValue(TempKeyStorage);
+						};
+
 						for (int32 Index = 0, ItemsLeft = Helper.Num(); ItemsLeft > 0; ++Index)
 						{
 							if (Helper.IsValidIndex(Index))
@@ -417,13 +427,9 @@ FPropertyAccess::Result FPropertyValueImpl::ImportText( const TArray<FObjectBase
 								const uint8* PairPtr = Helper.GetPairPtr(Index);
 								const uint8* KeyPtr = KeyProp->ContainerPtrToValuePtr<const uint8>(PairPtr);
 
-								FString KeyValue;
-								if (KeyPtr != InBaseAddress && KeyProp->ExportText_Direct(KeyValue, KeyPtr, KeyPtr, nullptr, 0))
+								if (KeyPtr != InBaseAddress && KeyProp->ImportText_Direct(*InKeyValue, TempKeyStorage, nullptr, 0) && KeyProp->Identical(KeyPtr, TempKeyStorage))
 								{
-									if ((CastField<FObjectProperty>(KeyProp) != nullptr && KeyValue.Contains(InKeyValue)) || InKeyValue == KeyValue)
-									{
-										return true;
-									}
+									return true;
 								}
 							}
 						}
