Copy and paste on a details row containing an asset picker will now go through the SPropertyEditorAsset paste first.

diff --git a/Engine/Source/Editor/PropertyEditor/Private/DetailPropertyRow.cpp b/Engine/Source/Editor/PropertyEditor/Private/DetailPropertyRow.cpp
--- a/Engine/Source/Editor/PropertyEditor/Private/DetailPropertyRow.cpp
+++ b/Engine/Source/Editor/PropertyEditor/Private/DetailPropertyRow.cpp
@@ -878,6 +878,7 @@ void FDetailPropertyRow::MakeValueWidget( FDetailWidgetRow& Row, const TSharedPt
 		[
 			SAssignNew( PropertyValue, SPropertyValueWidget, PropertyEditor, GetPropertyUtilities() )
 			.ShowPropertyButtons( false ) // We handle this ourselves
+			.InWidgetRow(&Row)
 		];
 		
 		MinWidth = PropertyValue->GetMinDesiredWidth();
diff --git a/Engine/Source/Editor/PropertyEditor/Private/PropertyEditorHelpers.cpp b/Engine/Source/Editor/PropertyEditor/Private/PropertyEditorHelpers.cpp
--- a/Engine/Source/Editor/PropertyEditor/Private/PropertyEditorHelpers.cpp
+++ b/Engine/Source/Editor/PropertyEditor/Private/PropertyEditorHelpers.cpp
@@ -60,6 +60,7 @@ void SPropertyValueWidget::Construct( const FArguments& InArgs, TSharedPtr<FProp
 {
 	MinDesiredWidth = 0.0f;
 	MaxDesiredWidth = 0.0f;
+	WidgetRow = InArgs._InWidgetRow;
 
 	SetEnabled( TAttribute<bool>( PropertyEditor.ToSharedRef(), &FPropertyEditor::IsPropertyEditingEnabled ) );
 
@@ -173,9 +174,11 @@ TSharedRef<SWidget> SPropertyValueWidget::ConstructPropertyEditorWidget( TShared
 		}
 		else if ( SPropertyEditorAsset::Supports( PropertyEditorRef ) )
 		{
+			// SPropertyEditorAsset has its own copy & paste that need to be bound to the widget row's
 			TSharedRef<SPropertyEditorAsset> AssetWidget = 
 				SAssignNew( PropertyWidget, SPropertyEditorAsset, PropertyEditorRef )
-				.ThumbnailPool( InPropertyUtilities.IsValid() ? InPropertyUtilities->GetThumbnailPool() : nullptr );
+				.ThumbnailPool( InPropertyUtilities.IsValid() ? InPropertyUtilities->GetThumbnailPool() : nullptr )
+				.InWidgetRow(WidgetRow);
 			
 			AssetWidget->GetDesiredWidth( MinDesiredWidth, MaxDesiredWidth );
 		}
diff --git a/Engine/Source/Editor/PropertyEditor/Private/PropertyEditorHelpers.h b/Engine/Source/Editor/PropertyEditor/Private/PropertyEditorHelpers.h
--- a/Engine/Source/Editor/PropertyEditor/Private/PropertyEditorHelpers.h
+++ b/Engine/Source/Editor/PropertyEditor/Private/PropertyEditorHelpers.h
@@ -14,6 +14,7 @@
 
 class FNotifyHook;
 class FObjectPropertyNode;
+class FDetailWidgetRow;
 
 /** Property button enums. */
 namespace EPropertyButton
@@ -60,6 +61,7 @@ public:
 		: _ShowPropertyButtons( true )
 	{}
 		SLATE_ARGUMENT( bool, ShowPropertyButtons )
+		SLATE_ARGUMENT(FDetailWidgetRow*, InWidgetRow)
 	SLATE_END_ARGS()
 
 	void Construct( const FArguments& InArgs, TSharedPtr<FPropertyEditor> InPropertyEditor, TSharedPtr<IPropertyUtilities> InPropertyUtilities );
@@ -78,6 +80,9 @@ private:
 	float MinDesiredWidth;
 	/** The maximum desired with if this property value */
 	float MaxDesiredWidth;
+	/** The widget row this value widget is part of */
+	FDetailWidgetRow* WidgetRow;
+
 };
 
 class SEditConditionWidget : public SCompoundWidget
diff --git a/Engine/Source/Editor/PropertyEditor/Private/UserInterface/PropertyEditor/SPropertyEditorAsset.cpp b/Engine/Source/Editor/PropertyEditor/Private/UserInterface/PropertyEditor/SPropertyEditorAsset.cpp
--- a/Engine/Source/Editor/PropertyEditor/Private/UserInterface/PropertyEditor/SPropertyEditorAsset.cpp
+++ b/Engine/Source/Editor/PropertyEditor/Private/UserInterface/PropertyEditor/SPropertyEditorAsset.cpp
@@ -34,11 +34,10 @@
 #include "FileHelpers.h"
 #include "Presentation/PropertyEditor/PropertyEditor.h"
 #include "AssetThumbnail.h"
+#include "DetailWidgetRow.h"
 
 #define LOCTEXT_NAMESPACE "PropertyEditor"
 
-DECLARE_DELEGATE( FOnCopy );
-DECLARE_DELEGATE( FOnPaste );
 
 // Helper to retrieve the correct property that has the applicable metadata.
 static const FProperty* GetActualMetadataProperty(const FProperty* Property)
@@ -247,6 +246,22 @@ void SPropertyEditorAsset::Construct(const FArguments& InArgs, const TSharedPtr<
 	OnShouldFilterActor = InArgs._OnShouldFilterActor;
 	ObjectPath = InArgs._ObjectPath;
 
+	if(InArgs._InWidgetRow.IsSet() && InArgs._InWidgetRow.Get() != nullptr)
+	{
+		if (!InArgs._InWidgetRow.Get()->CopyMenuAction.IsBound())
+		{
+			InArgs._InWidgetRow.Get()->CopyMenuAction = FUIAction(
+				FExecuteAction::CreateSP(this, &SPropertyEditorAsset::OnCopy),
+				FCanExecuteAction());
+		}
+		if (!InArgs._InWidgetRow.Get()->PasteMenuAction.IsBound())
+		{
+			InArgs._InWidgetRow.Get()->PasteMenuAction = FUIAction(
+				FExecuteAction::CreateSP(this, &SPropertyEditorAsset::OnPaste),
+				FCanExecuteAction::CreateSP(this, &SPropertyEditorAsset::CanPaste));
+		}
+	}
+
 	FProperty* Property = nullptr;
 	if (PropertyEditor.IsValid())
 	{
diff --git a/Engine/Source/Editor/PropertyEditor/Private/UserInterface/PropertyEditor/SPropertyEditorAsset.h b/Engine/Source/Editor/PropertyEditor/Private/UserInterface/PropertyEditor/SPropertyEditorAsset.h
--- a/Engine/Source/Editor/PropertyEditor/Private/UserInterface/PropertyEditor/SPropertyEditorAsset.h
+++ b/Engine/Source/Editor/PropertyEditor/Private/UserInterface/PropertyEditor/SPropertyEditorAsset.h
@@ -25,6 +25,7 @@ class IPropertyHandle;
 class SBorder;
 class SComboButton;
 class UFactory;
+class FDetailWidgetRow;
 
 /**
  * A widget used to edit Asset-type properties (UObject-derived properties).
@@ -80,6 +81,7 @@ public:
 		SLATE_ARGUMENT(TSharedPtr<IPropertyHandle>, PropertyHandle)
 		SLATE_ARGUMENT(TArray<FAssetData>, OwnerAssetDataArray)
 		SLATE_EVENT(FOnShouldFilterActor, OnShouldFilterActor)
+		SLATE_ATTRIBUTE(FDetailWidgetRow*, InWidgetRow)
 
 	SLATE_END_ARGS()
 
