Fixed the details panel failing to restore expansion state when running deferred actions

diff --git a/Engine/Source/Editor/PropertyEditor/Private/SDetailsViewBase.cpp b/Engine/Source/Editor/PropertyEditor/Private/SDetailsViewBase.cpp
--- a/Engine/Source/Editor/PropertyEditor/Private/SDetailsViewBase.cpp
+++ b/Engine/Source/Editor/PropertyEditor/Private/SDetailsViewBase.cpp
@@ -1035,29 +1035,27 @@ void SDetailsViewBase::Tick( const FGeometry& AllottedGeometry, const double InC
 
 	if (bHadDeferredActions)
 	{
-		bRunningDeferredActions = true;
-		ON_SCOPE_EXIT { bRunningDeferredActions = false; };
-
-		TArray<FSimpleDelegate> DeferredActionsCopy;
-		
-		do
+		// Note: Extra scope so that RestoreAllExpandedItems actually restores state, as it does nothing when bRunningDeferredActions is true
 		{
-			// Execute any deferred actions
-			DeferredActionsCopy = MoveTemp(DeferredActions);
-			DeferredActions.Reset();
+			bRunningDeferredActions = true;
+			ON_SCOPE_EXIT{ bRunningDeferredActions = false; };
+
+			TArray<FSimpleDelegate> DeferredActionsCopy;
 
-			// Execute any deferred actions
-			for (const FSimpleDelegate& DeferredAction : DeferredActionsCopy)
+			do
 			{
-				DeferredAction.ExecuteIfBound();
-			}
-		} while (DeferredActions.Num() > 0);
-	}
+				// Execute any deferred actions
+				DeferredActionsCopy = MoveTemp(DeferredActions);
+				DeferredActions.Reset();
+
+				// Execute any deferred actions
+				for (const FSimpleDelegate& DeferredAction : DeferredActionsCopy)
+				{
+					DeferredAction.ExecuteIfBound();
+				}
+			} while (DeferredActions.Num() > 0);
+		}
 
-	if (bHadDeferredActions)
-	{
-		// Restore expansion state after processing all deferred actions. Must be done outside of the scope where
-		// bRunningDeferredActions is true since RestoreAllExpandedItems returns immediately when it is true
 		RestoreAllExpandedItems();
 	}
 
