Changing FDetailWidgetRow argument and member variable to TOptional versions because we may not have the row for customized details panels

diff --git a/Engine/Source/Editor/PropertyEditor/Private/PropertyEditorHelpers.cpp b/Engine/Source/Editor/PropertyEditor/Private/PropertyEditorHelpers.cpp
--- a/Engine/Source/Editor/PropertyEditor/Private/PropertyEditorHelpers.cpp
+++ b/Engine/Source/Editor/PropertyEditor/Private/PropertyEditorHelpers.cpp
@@ -60,7 +60,10 @@ void SPropertyValueWidget::Construct( const FArguments& InArgs, TSharedPtr<FProp
 {
 	MinDesiredWidth = 0.0f;
 	MaxDesiredWidth = 0.0f;
-	WidgetRow = InArgs._InWidgetRow;
+	if(InArgs._InWidgetRow.IsSet())
+	{
+		WidgetRow = InArgs._InWidgetRow;
+	}
 
 	SetEnabled( TAttribute<bool>( PropertyEditor.ToSharedRef(), &FPropertyEditor::IsPropertyEditingEnabled ) );
 
diff --git a/Engine/Source/Editor/PropertyEditor/Private/PropertyEditorHelpers.h b/Engine/Source/Editor/PropertyEditor/Private/PropertyEditorHelpers.h
--- a/Engine/Source/Editor/PropertyEditor/Private/PropertyEditorHelpers.h
+++ b/Engine/Source/Editor/PropertyEditor/Private/PropertyEditorHelpers.h
@@ -61,7 +61,7 @@ public:
 		: _ShowPropertyButtons( true )
 	{}
 		SLATE_ARGUMENT( bool, ShowPropertyButtons )
-		SLATE_ARGUMENT(FDetailWidgetRow*, InWidgetRow)
+		SLATE_ARGUMENT(TOptional<FDetailWidgetRow*>, InWidgetRow)
 	SLATE_END_ARGS()
 
 	void Construct( const FArguments& InArgs, TSharedPtr<FPropertyEditor> InPropertyEditor, TSharedPtr<IPropertyUtilities> InPropertyUtilities );
@@ -81,7 +81,7 @@ private:
 	/** The maximum desired with if this property value */
 	float MaxDesiredWidth;
 	/** The widget row this value widget is part of */
-	FDetailWidgetRow* WidgetRow;
+	TOptional<FDetailWidgetRow*> WidgetRow;
 
 };
 
diff --git a/Engine/Source/Editor/PropertyEditor/Private/UserInterface/PropertyEditor/SPropertyEditorAsset.cpp b/Engine/Source/Editor/PropertyEditor/Private/UserInterface/PropertyEditor/SPropertyEditorAsset.cpp
--- a/Engine/Source/Editor/PropertyEditor/Private/UserInterface/PropertyEditor/SPropertyEditorAsset.cpp
+++ b/Engine/Source/Editor/PropertyEditor/Private/UserInterface/PropertyEditor/SPropertyEditorAsset.cpp
@@ -246,17 +246,17 @@ void SPropertyEditorAsset::Construct(const FArguments& InArgs, const TSharedPtr<
 	OnShouldFilterActor = InArgs._OnShouldFilterActor;
 	ObjectPath = InArgs._ObjectPath;
 
-	if(InArgs._InWidgetRow.IsSet() && InArgs._InWidgetRow.Get() != nullptr)
+	if(InArgs._InWidgetRow.IsSet() && InArgs._InWidgetRow.GetValue() != nullptr)
 	{
-		if (!InArgs._InWidgetRow.Get()->CopyMenuAction.IsBound())
+		if (!InArgs._InWidgetRow.GetValue()->CopyMenuAction.IsBound())
 		{
-			InArgs._InWidgetRow.Get()->CopyMenuAction = FUIAction(
+			InArgs._InWidgetRow.GetValue()->CopyMenuAction = FUIAction(
 				FExecuteAction::CreateSP(this, &SPropertyEditorAsset::OnCopy),
 				FCanExecuteAction());
 		}
-		if (!InArgs._InWidgetRow.Get()->PasteMenuAction.IsBound())
+		if (!InArgs._InWidgetRow.GetValue()->PasteMenuAction.IsBound())
 		{
-			InArgs._InWidgetRow.Get()->PasteMenuAction = FUIAction(
+			InArgs._InWidgetRow.GetValue()->PasteMenuAction = FUIAction(
 				FExecuteAction::CreateSP(this, &SPropertyEditorAsset::OnPaste),
 				FCanExecuteAction::CreateSP(this, &SPropertyEditorAsset::CanPaste));
 		}
diff --git a/Engine/Source/Editor/PropertyEditor/Private/UserInterface/PropertyEditor/SPropertyEditorAsset.h b/Engine/Source/Editor/PropertyEditor/Private/UserInterface/PropertyEditor/SPropertyEditorAsset.h
--- a/Engine/Source/Editor/PropertyEditor/Private/UserInterface/PropertyEditor/SPropertyEditorAsset.h
+++ b/Engine/Source/Editor/PropertyEditor/Private/UserInterface/PropertyEditor/SPropertyEditorAsset.h
@@ -81,7 +81,7 @@ public:
 		SLATE_ARGUMENT(TSharedPtr<IPropertyHandle>, PropertyHandle)
 		SLATE_ARGUMENT(TArray<FAssetData>, OwnerAssetDataArray)
 		SLATE_EVENT(FOnShouldFilterActor, OnShouldFilterActor)
-		SLATE_ATTRIBUTE(FDetailWidgetRow*, InWidgetRow)
+		SLATE_ARGUMENT(TOptional<FDetailWidgetRow*>, InWidgetRow)
 
 	SLATE_END_ARGS()
 
