Fixes: while Binding trigger track to skeletal mesh actor endpoint and changing the default value crash

diff --git a/Engine/Source/Editor/PropertyEditor/Private/UserInterface/PropertyEditor/SPropertyEditorNumeric.h b/Engine/Source/Editor/PropertyEditor/Private/UserInterface/PropertyEditor/SPropertyEditorNumeric.h
--- a/Engine/Source/Editor/PropertyEditor/Private/UserInterface/PropertyEditor/SPropertyEditorNumeric.h
+++ b/Engine/Source/Editor/PropertyEditor/Private/UserInterface/PropertyEditor/SPropertyEditorNumeric.h
@@ -510,15 +510,21 @@ private:
 	{
 		const TSharedRef< IPropertyHandle > PropertyHandle = PropertyEditor->GetPropertyHandle();
 		NumericType OrgValue(0);
-		if (bIsUsingSlider || (PropertyHandle->GetValue(OrgValue) == FPropertyAccess::Fail || OrgValue != NewValue))
+		/* sometimes an FProperty may have been destroyed due to 2 different events invoking this handler (with the same
+		 * NewValue) and the first one nullifying the current FProperty ~ in this case it's too late to not invoke the
+		 * method, but we can not run the code instead */
+		if (PropertyHandle->GetProperty())
 		{
-			PropertyHandle->SetValue(NewValue);
-			LastSliderCommittedValue = NewValue;
-		}
+			if (bIsUsingSlider || (PropertyHandle->GetValue(OrgValue) == FPropertyAccess::Fail || OrgValue != NewValue))
+			{
+				PropertyHandle->SetValue(NewValue);
+				LastSliderCommittedValue = NewValue;
+			}
 
-		if (TypeInterface.IsValid() && !TypeInterface->FixedDisplayUnits.IsSet())
-		{
-			TypeInterface->SetupFixedDisplay(NewValue);
+			if (TypeInterface.IsValid() && !TypeInterface->FixedDisplayUnits.IsSet())
+			{
+				TypeInterface->SetupFixedDisplay(NewValue);
+			}
 		}
 	}
 
