Fix Editor crash when changing the root expression type in GampelayTagQuery

diff --git a/Engine/Source/Editor/PropertyEditor/Private/UserInterface/PropertyEditor/SPropertyEditorEditInline.cpp b/Engine/Source/Editor/PropertyEditor/Private/UserInterface/PropertyEditor/SPropertyEditorEditInline.cpp
--- a/Engine/Source/Editor/PropertyEditor/Private/UserInterface/PropertyEditor/SPropertyEditorEditInline.cpp
+++ b/Engine/Source/Editor/PropertyEditor/Private/UserInterface/PropertyEditor/SPropertyEditorEditInline.cpp
@@ -348,7 +348,5 @@ void SPropertyEditorEditInline::OnClassPicked(UClass* InClass)
 
 		// Force a rebuild of the children when this node changes
 		PropertyNode->RequestRebuildChildren();
-
-		ComboButton->SetIsOpen(false);
 	}
 }
