Add support for base class to explicitly allow all their properties when derived class only allow a subset of their own through using the `bAllowListAllProperties` setting in permission lists

diff --git a/Engine/Source/Editor/PropertyEditor/Private/PropertyPermissionList.cpp b/Engine/Source/Editor/PropertyEditor/Private/PropertyPermissionList.cpp
--- a/Engine/Source/Editor/PropertyEditor/Private/PropertyPermissionList.cpp
+++ b/Engine/Source/Editor/PropertyEditor/Private/PropertyPermissionList.cpp
@@ -278,6 +278,14 @@ const FNamePermissionList& FPropertyPermissionList::GetCachedPermissionListForSt
 			}
 			else
 			{
+				// if the parent struct is explicitly allowing all properties and has an empty allow list, add all its prop to new list, before adding the current entry
+				if (bInOutShouldAllowListAllProperties && NewPermissionList.GetAllowList().Num() == 0 && SuperStruct)
+				{
+					for (TFieldIterator<FProperty> Property(SuperStruct, EFieldIteratorFlags::IncludeSuper, EFieldIteratorFlags::ExcludeDeprecated); Property; ++Property)
+					{
+						NewPermissionList.AddAllowListItem(PropertyPermissionListOwner, Property->GetFName());
+					}
+				}
 				NewPermissionList.Append(Entry->PermissionList);
 			}
 
