Fix for the browse to actor button not working for [REDACTED_HANDLE] user device references

diff --git a/Engine/Source/Editor/PropertyEditor/Private/PropertyCustomizationHelpers.cpp b/Engine/Source/Editor/PropertyEditor/Private/PropertyCustomizationHelpers.cpp
--- a/Engine/Source/Editor/PropertyEditor/Private/PropertyCustomizationHelpers.cpp
+++ b/Engine/Source/Editor/PropertyEditor/Private/PropertyCustomizationHelpers.cpp
@@ -526,6 +526,7 @@ void SObjectPropertyEntryBox::Construct( const FArguments& InArgs )
 				.AllowCreate(InArgs._AllowCreate)
 				.DisplayUseSelected(InArgs._DisplayUseSelected)
 				.DisplayBrowse(InArgs._DisplayBrowse)
+				.OnBrowseOverride(InArgs._OnBrowseOverride)
 				.EnableContentPicker(InArgs._EnableContentPicker)
 				.PropertyHandle(PropertyHandle)
 				.OwnerAssetDataArray(OwnerAssetDataArray)
diff --git a/Engine/Source/Editor/PropertyEditor/Private/UserInterface/PropertyEditor/SPropertyEditorAsset.cpp b/Engine/Source/Editor/PropertyEditor/Private/UserInterface/PropertyEditor/SPropertyEditorAsset.cpp
--- a/Engine/Source/Editor/PropertyEditor/Private/UserInterface/PropertyEditor/SPropertyEditorAsset.cpp
+++ b/Engine/Source/Editor/PropertyEditor/Private/UserInterface/PropertyEditor/SPropertyEditorAsset.cpp
@@ -553,13 +553,15 @@ void SPropertyEditorAsset::Construct(const FArguments& InArgs, const TSharedPtr<
 
 	if( InArgs._DisplayBrowse )
 	{
+		FSimpleDelegate OnBrowseDelegate = InArgs._OnBrowseOverride.IsBound() ? InArgs._OnBrowseOverride : FSimpleDelegate::CreateSP(this, &SPropertyEditorAsset::OnBrowse);
+
 		ButtonBox->AddSlot()
 		.Padding( 2.0f, 0.0f )
 		.AutoWidth()
 		.VAlign(VAlign_Center)
 		[
 			PropertyCustomizationHelpers::MakeBrowseButton(
-				FSimpleDelegate::CreateSP( this, &SPropertyEditorAsset::OnBrowse ),
+				OnBrowseDelegate,
 				TAttribute<FText>( this, &SPropertyEditorAsset::GetOnBrowseToolTip ),
 				true,
 				bIsActor
diff --git a/Engine/Source/Editor/PropertyEditor/Private/UserInterface/PropertyEditor/SPropertyEditorAsset.h b/Engine/Source/Editor/PropertyEditor/Private/UserInterface/PropertyEditor/SPropertyEditorAsset.h
--- a/Engine/Source/Editor/PropertyEditor/Private/UserInterface/PropertyEditor/SPropertyEditorAsset.h
+++ b/Engine/Source/Editor/PropertyEditor/Private/UserInterface/PropertyEditor/SPropertyEditorAsset.h
@@ -68,6 +68,7 @@ public:
 		SLATE_ARGUMENT(bool, DisplayThumbnail)
 		SLATE_ARGUMENT(bool, DisplayUseSelected)
 		SLATE_ARGUMENT(bool, DisplayBrowse)
+		SLATE_EVENT(FSimpleDelegate, OnBrowseOverride)
 		SLATE_ARGUMENT(bool, EnableContentPicker)
 		SLATE_ARGUMENT(bool, DisplayCompactSize)
 		SLATE_ARGUMENT(TSharedPtr<FAssetThumbnailPool>, ThumbnailPool)
diff --git a/Engine/Source/Editor/PropertyEditor/Public/PropertyCustomizationHelpers.h b/Engine/Source/Editor/PropertyEditor/Public/PropertyCustomizationHelpers.h
--- a/Engine/Source/Editor/PropertyEditor/Public/PropertyCustomizationHelpers.h
+++ b/Engine/Source/Editor/PropertyEditor/Public/PropertyCustomizationHelpers.h
@@ -210,6 +210,8 @@ public:
 		SLATE_ARGUMENT(bool, DisplayUseSelected)
 		/** Whether to show the 'Browse' button */
 		SLATE_ARGUMENT(bool, DisplayBrowse)
+		/** Optional delegate called when the 'Browse' button is clicked. Used to override the default editor behavior */
+		SLATE_EVENT(FSimpleDelegate, OnBrowseOverride)
 		/** Whether to enable the content Picker */
 		SLATE_ARGUMENT(bool, EnableContentPicker)
 		/** Whether or not to display a smaller, compact size for the asset thumbnail */ 
