Fix flickering "reset to default" arrows in details panel when moving widgets around.

diff --git a/Engine/Source/Editor/PropertyEditor/Private/SDetailSingleItemRow.cpp b/Engine/Source/Editor/PropertyEditor/Private/SDetailSingleItemRow.cpp
--- a/Engine/Source/Editor/PropertyEditor/Private/SDetailSingleItemRow.cpp
+++ b/Engine/Source/Editor/PropertyEditor/Private/SDetailSingleItemRow.cpp
@@ -571,6 +571,7 @@ void SDetailSingleItemRow::Construct( const FArguments& InArgs, FDetailLayoutCus
 
 			TArray<FPropertyRowExtensionButton> ExtensionButtons;
 
+			UpdateResetToDefault();
 			FPropertyRowExtensionButton& ResetToDefault = ExtensionButtons.AddDefaulted_GetRef();
 			ResetToDefault.Label = NSLOCTEXT("PropertyEditor", "ResetToDefault", "Reset to Default");
 			ResetToDefault.UIAction = FUIAction(
