[Feature] DetailsSplitter
This is a splitter that can dynamically connect matching properties as well as provide buttons for copying values between the two details panels.

diff --git a/Engine/Source/Editor/EditorStyle/Private/StarshipStyle.cpp b/Engine/Source/Editor/EditorStyle/Private/StarshipStyle.cpp
--- a/Engine/Source/Editor/EditorStyle/Private/StarshipStyle.cpp
+++ b/Engine/Source/Editor/EditorStyle/Private/StarshipStyle.cpp
@@ -5212,6 +5212,8 @@ void FStarshipEditorStyle::FStyle::SetupPersonaStyle()
 		Set("BlueprintDif.PrevDiff", new IMAGE_BRUSH_SVG("/Starship/Blueprints/diff_prev", Icon16x16));
 		Set("BlueprintDif.HorizontalDiff", new IMAGE_BRUSH_SVG("Starship/Blueprints/icon_horizontal_diff_view", Icon16x16));
 		Set("BlueprintDif.VerticalDiff", new IMAGE_BRUSH_SVG("Starship/Blueprints/icon_vertical_diff_view", Icon16x16));
+		Set("BlueprintDif.CopyPropertyLeft", new CORE_IMAGE_BRUSH_SVG( "Starship/Common/arrow-left", Icon16x16));
+		Set("BlueprintDif.CopyPropertyRight", new CORE_IMAGE_BRUSH_SVG( "Starship/Common/arrow-right", Icon16x16));
 
 		Set("BlueprintDif.ItalicText", 
 			FTextBlockStyle(NormalText)
diff --git a/Engine/Source/Editor/Kismet/Private/DetailsDiff.cpp b/Engine/Source/Editor/Kismet/Private/DetailsDiff.cpp
--- a/Engine/Source/Editor/Kismet/Private/DetailsDiff.cpp
+++ b/Engine/Source/Editor/Kismet/Private/DetailsDiff.cpp
@@ -20,19 +20,26 @@
 
 class UObject;
 
-FDetailsDiff::FDetailsDiff(const UObject* InObject, FOnDisplayedPropertiesChanged InOnDisplayedPropertiesChanged )
+FDetailsDiff::FDetailsDiff(const UObject* InObject, FOnDisplayedPropertiesChanged InOnDisplayedPropertiesChanged, bool bScrollbarOnLeft )
 	: OnDisplayedPropertiesChanged( InOnDisplayedPropertiesChanged )
 	, DisplayedObject( InObject )
-	, DetailsView()
+{
+	ScrollBar = SNew(SLinkableScrollBar);
+	DetailsView = CreateDetailsView(InObject, ScrollBar, bScrollbarOnLeft);
+	DetailsView->SetOnDisplayedPropertiesChanged( ::FOnDisplayedPropertiesChanged::CreateRaw(this, &FDetailsDiff::HandlePropertiesChanged) );
+	DisplayedProperties = DetailsView->GetPropertiesInOrderDisplayed();
+}
+
+TSharedRef<IDetailsView> FDetailsDiff::CreateDetailsView(const UObject* InObject, TSharedPtr<SScrollBar> ExternalScrollbar, bool bScrollbarOnLeft)
 {
 	FDetailsViewArgs DetailsViewArgs;
 	DetailsViewArgs.bShowDifferingPropertiesOption = true;
 	DetailsViewArgs.NameAreaSettings = FDetailsViewArgs::HideNameArea;
-	ScrollBar = SNew(SLinkableScrollBar);
-	DetailsViewArgs.ExternalScrollbar = ScrollBar;
+	DetailsViewArgs.ExternalScrollbar = ExternalScrollbar;
+	DetailsViewArgs.ScrollbarAlignment = bScrollbarOnLeft ? HAlign_Left : HAlign_Right;
 
 	FPropertyEditorModule& EditModule = FModuleManager::Get().GetModuleChecked<FPropertyEditorModule>("PropertyEditor");
-	DetailsView = EditModule.CreateDetailView(DetailsViewArgs);
+	TSharedRef<IDetailsView> DetailsView = EditModule.CreateDetailView(DetailsViewArgs);
 	DetailsView->SetIsPropertyEditingEnabledDelegate(FIsPropertyEditingEnabled::CreateStatic([]{return false; }));
 	if (InObject && InObject->IsA<UBlueprint>())
 	{
@@ -49,9 +56,8 @@ FDetailsDiff::FDetailsDiff(const UObject* InObject, FOnDisplayedPropertiesChange
 	// This is a read only details view (see the property editing delegate), but const correctness at the details view
 	// level would stress the type system a little:
 	DetailsView->SetObject(const_cast<UObject*>(InObject));
-	DetailsView->SetOnDisplayedPropertiesChanged( ::FOnDisplayedPropertiesChanged::CreateRaw(this, &FDetailsDiff::HandlePropertiesChanged) );
 
-	DifferingProperties = DetailsView->GetPropertiesInOrderDisplayed();
+	return DetailsView;
 }
 
 FDetailsDiff::~FDetailsDiff()
@@ -75,7 +81,7 @@ void FDetailsDiff::HandlePropertiesChanged()
 {
 	if( OnDisplayedPropertiesChanged.IsBound() )
 	{
-		DifferingProperties = DetailsView->GetPropertiesInOrderDisplayed();
+		DisplayedProperties = DetailsView->GetPropertiesInOrderDisplayed();
 		OnDisplayedPropertiesChanged.Execute();
 	}
 }
@@ -83,7 +89,7 @@ void FDetailsDiff::HandlePropertiesChanged()
 TArray<FPropertySoftPath> FDetailsDiff::GetDisplayedProperties() const
 {
 	TArray<FPropertySoftPath> Ret;
-	Algo::Copy(DifferingProperties, Ret);
+	Algo::Copy(DisplayedProperties, Ret);
 	return Ret;
 }
 
diff --git a/Engine/Source/Editor/Kismet/Private/DiffControl.cpp b/Engine/Source/Editor/Kismet/Private/DiffControl.cpp
--- a/Engine/Source/Editor/Kismet/Private/DiffControl.cpp
+++ b/Engine/Source/Editor/Kismet/Private/DiffControl.cpp
@@ -12,6 +12,9 @@
 #include "Widgets/Input/SButton.h"
 #include "Widgets/Input/SCheckBox.h"
 #include "Framework/Application/SlateApplication.h"
+#include "AsyncTreeDifferences.h"
+#include "DetailTreeNode.h"
+#include "AsyncDetailViewDiff.h"
 
 #define LOCTEXT_NAMESPACE "SBlueprintDif"
 
@@ -80,7 +83,26 @@ TSharedRef<SWidget> IDiffControl::GenerateObjectDiffWidget(FSingleObjectDiffEntr
 void IDiffControl::EnableComments(TWeakPtr<STreeView<TSharedPtr<FBlueprintDifferenceTreeEntry>>> TreeView,
 	const UObject* OldObject, const UObject* NewObject)
 {
-	ReviewCommentsDiffControl = MakeShared<FReviewCommentsDiffControl>(OldObject, NewObject, TreeView);
+	const UPackage* Package = OldObject ? OldObject->GetPackage() : NewObject->GetPackage();
+	const FPackagePath PackagePath = FPackagePath::FromPackageNameChecked(Package->GetName());
+	ReviewCommentsDiffControl = MakeShared<FReviewCommentsDiffControl>(PackagePath.GetLocalFullPath(), TreeView);
+}
+
+void IDiffControl::EnableComments(TWeakPtr<STreeView<TSharedPtr<FBlueprintDifferenceTreeEntry>>> TreeView, const TArray<const UObject*>& Objects)
+{
+	for (const UObject* Object : Objects)
+	{
+		if (!ensure(Object))
+		{
+			continue;
+		}
+		if (const UPackage* Package = Object->GetPackage())
+		{
+			const FPackagePath PackagePath = FPackagePath::FromPackageNameChecked(Package->GetName());
+			ReviewCommentsDiffControl = MakeShared<FReviewCommentsDiffControl>(PackagePath.GetLocalFullPath(), TreeView);
+			return;
+		}
+	}
 }
 
 void IDiffControl::GenerateCategoryCommentTreeEntries(TArray<TSharedPtr<FBlueprintDifferenceTreeEntry>>& OutChildrenEntries,
@@ -266,35 +288,30 @@ void FSCSDiffControl::EnableComments(TWeakPtr<STreeView<TSharedPtr<FBlueprintDif
 }
 
 FDetailsDiffControl::FDetailsDiffControl(const UObject* InOldObject, const UObject* InNewObject,
-	FOnDiffEntryFocused InSelectionCallback, bool bPopulateOutTreeEntries): SelectionCallback(InSelectionCallback)
-	                                                                        , OldDetails(InOldObject, FDetailsDiff::FOnDisplayedPropertiesChanged())
-	                                                                        , NewDetails(InNewObject, FDetailsDiff::FOnDisplayedPropertiesChanged())
-	                                                                        , bPopulateOutTreeEntries(bPopulateOutTreeEntries)
-{
-	OldDetails.DiffAgainst(NewDetails, DifferingProperties, true);
-
-	TSet<FPropertyPath> PropertyPaths;
-	Algo::Transform(DifferingProperties, PropertyPaths,
-	                [&InOldObject](const FSingleObjectDiffEntry& DiffEntry)
-	                {
-		                return DiffEntry.Identifier.ResolvePath(InOldObject);
-	                });
-
-	OldDetails.DetailsWidget()->UpdatePropertyAllowList(PropertyPaths);
-
-	PropertyPaths.Reset();
-	Algo::Transform(DifferingProperties, PropertyPaths,
-	                [&InNewObject](const FSingleObjectDiffEntry& DiffEntry)
-	                {
-		                return DiffEntry.Identifier.ResolvePath(InNewObject);
-	                });
-
-	NewDetails.DetailsWidget()->UpdatePropertyAllowList(PropertyPaths);
-		
-	// Sync the scrolling between the left and right panels
-	const TSharedRef<IDetailsView> OldDetailsView = OldDetails.DetailsWidget();
-	const TSharedRef<IDetailsView> NewDetailsView = NewDetails.DetailsWidget();
-	FDetailsDiff::LinkScrolling(OldDetails, NewDetails, GetLinkedScrollRateAttribute(OldDetailsView, NewDetailsView));;
+	FOnDiffEntryFocused InSelectionCallback, bool bPopulateOutTreeEntries)
+	: SelectionCallback(InSelectionCallback)
+	, bPopulateOutTreeEntries(bPopulateOutTreeEntries)
+{
+	if (InOldObject)
+	{
+		InsertObject(InOldObject, true);
+	}
+	if (InNewObject)
+	{
+		InsertObject(InNewObject, false);
+	}
+}
+
+void FDetailsDiffControl::Tick()
+{
+	for (auto& [Object, PropertyTreeDiffs] : PropertyTreeDifferences)
+	{
+		if (PropertyTreeDiffs.Left)
+		{
+			constexpr float MaxTickTimeMs = 0.01f;
+			PropertyTreeDiffs.Left->Tick(MaxTickTimeMs);
+		}
+	}
 }
 
 void FDetailsDiffControl::GenerateTreeEntries(TArray<TSharedPtr<FBlueprintDifferenceTreeEntry>>& OutTreeEntries, TArray<TSharedPtr<FBlueprintDifferenceTreeEntry>>& OutRealDifferences)
@@ -309,6 +326,31 @@ void FDetailsDiffControl::GenerateTreeEntriesWithoutComments(
 	TArray<TSharedPtr<FBlueprintDifferenceTreeEntry>>& OutTreeEntries,
 	TArray<TSharedPtr<FBlueprintDifferenceTreeEntry>>& OutRealDifferences)
 {
+	TArray<FSingleObjectDiffEntry> DifferingProperties;
+	
+	for (int32 LeftIndex = 0; LeftIndex < ObjectDisplayOrder.Num() - 1; ++LeftIndex)
+	{
+		const UObject* LeftObject = ObjectDisplayOrder[LeftIndex];
+		if (!ensure(LeftObject))
+		{
+			continue;
+		}
+		const TSharedPtr<FAsyncDetailViewDiff> Diff = PropertyTreeDifferences[LeftObject].Right;
+		Diff->FlushQueue(); // make sure differences are fully up to date
+		Diff->GetPropertyDifferences(DifferingProperties);
+	}
+
+	for (auto&[Object, DetailsDiff] : DetailsDiffs)
+	{
+		Algo::Transform(DifferingProperties, PropertyAllowList,
+        [&Object = Object](const FSingleObjectDiffEntry& DiffEntry)
+        {
+        	return DiffEntry.Identifier.ResolvePath(Object);
+        });
+        
+        DetailsDiff.DetailsWidget()->UpdatePropertyAllowList(PropertyAllowList);
+	}
+	
 	for (const FSingleObjectDiffEntry& Difference : DifferingProperties)
 	{
 		TSharedPtr<FBlueprintDifferenceTreeEntry> Entry = MakeShared<FBlueprintDifferenceTreeEntry>(
@@ -325,14 +367,76 @@ void FDetailsDiffControl::GenerateTreeEntriesWithoutComments(
 
 void FDetailsDiffControl::EnableComments(TWeakPtr<STreeView<TSharedPtr<FBlueprintDifferenceTreeEntry>>> TreeView)
 {
-	IDiffControl::EnableComments(TreeView, OldDetails.GetDisplayedObject(), NewDetails.GetDisplayedObject());
+	IDiffControl::EnableComments(TreeView, ObjectDisplayOrder);
+}
+
+TSharedRef<IDetailsView> FDetailsDiffControl::InsertObject(const UObject* Object, bool bScrollbarOnLeft, int32 Index)
+{
+	FDetailsDiff DetailsDiff(Object, FDetailsDiff::FOnDisplayedPropertiesChanged(), bScrollbarOnLeft);
+	const TSharedRef<IDetailsView> DetailsView = DetailsDiff.DetailsWidget();
+	DetailsView->UpdatePropertyAllowList(PropertyAllowList);
+
+	if (Index == INDEX_NONE)
+	{
+		Index = ObjectDisplayOrder.Num();
+	}
+	
+	DetailsDiffs.Add(Object, DetailsDiff);
+	ObjectDisplayOrder.Insert(Object, Index);
+	PropertyTreeDifferences.Add(Object, {});
+	
+	// set up interaction with left panel
+	if (ObjectDisplayOrder.IsValidIndex(Index - 1))
+	{
+		const UObject* OtherObject = ObjectDisplayOrder[Index - 1];
+		FDetailsDiff& OtherDetailsDiff = DetailsDiffs[OtherObject];
+		const TSharedRef<IDetailsView> OtherDetailsView = OtherDetailsDiff.DetailsWidget();
+		const auto ScrollRate = GetLinkedScrollRateAttribute(OtherDetailsView, DetailsView);
+		FDetailsDiff::LinkScrolling(OtherDetailsDiff, DetailsDiff, ScrollRate);
+
+		PropertyTreeDifferences[OtherObject].Right = MakeShared<FAsyncDetailViewDiff>(OtherDetailsView, DetailsView);
+		PropertyTreeDifferences[Object].Left = PropertyTreeDifferences[OtherObject].Right;
+	}
+	// Set up interaction with right panel
+	if (ObjectDisplayOrder.IsValidIndex(Index + 1))
+	{
+		const UObject* OtherObject = ObjectDisplayOrder[Index + 1];
+		FDetailsDiff& OtherDetailsDiff = DetailsDiffs[OtherObject];
+		const TSharedRef<IDetailsView> OtherDetailsView = OtherDetailsDiff.DetailsWidget();
+		const auto ScrollRate = GetLinkedScrollRateAttribute(DetailsView, OtherDetailsView);
+		FDetailsDiff::LinkScrolling(DetailsDiff, OtherDetailsDiff, ScrollRate);
+		
+		const TSharedPtr<FAsyncDetailViewDiff> DifferencesWithRight = MakeShared<FAsyncDetailViewDiff>(OtherDetailsView, DetailsView);
+
+		PropertyTreeDifferences[OtherObject].Left = MakeShared<FAsyncDetailViewDiff>(DetailsView, OtherDetailsView);
+		PropertyTreeDifferences[Object].Right = PropertyTreeDifferences[OtherObject].Left;
+	}
+
+	return DetailsView;
+}
+
+TSharedRef<IDetailsView> FDetailsDiffControl::GetDetailsWidget(const UObject* Object) const
+{
+	return DetailsDiffs[Object].DetailsWidget();
+}
+
+TSharedPtr<FAsyncDetailViewDiff> FDetailsDiffControl::GetDifferencesWithLeft(const UObject* Object) const
+{
+	return PropertyTreeDifferences[Object].Left;
+}
+
+TSharedPtr<FAsyncDetailViewDiff> FDetailsDiffControl::GetDifferencesWithRight(const UObject* Object) const
+{
+	return PropertyTreeDifferences[Object].Right;
 }
 
 void FDetailsDiffControl::OnSelectDiffEntry(FPropertySoftPath PropertyName)
 {
 	SelectionCallback.ExecuteIfBound();
-	OldDetails.HighlightProperty(PropertyName);
-	NewDetails.HighlightProperty(PropertyName);
+	for (auto&[Object, DetailsDiff] : DetailsDiffs)
+	{
+		DetailsDiff.HighlightProperty(PropertyName);
+	}
 }
 
 TAttribute<TArray<FVector2f>> FDetailsDiffControl::GetLinkedScrollRateAttribute(const TSharedRef<IDetailsView>& OldDetailsView, const TSharedRef<IDetailsView>& NewDetailsView)
@@ -340,62 +444,42 @@ TAttribute<TArray<FVector2f>> FDetailsDiffControl::GetLinkedScrollRateAttribute(
 	return TAttribute<TArray<FVector2f>>::CreateRaw(this, &FDetailsDiffControl::GetLinkedScrollRate, OldDetailsView, NewDetailsView);
 }
 
-TArray<FVector2f> FDetailsDiffControl::GetLinkedScrollRate(TSharedRef<IDetailsView> OldDetailsView, TSharedRef<IDetailsView> NewDetailsView) const
+TArray<FVector2f> FDetailsDiffControl::GetLinkedScrollRate(TSharedRef<IDetailsView> LeftDetailsView, TSharedRef<IDetailsView> RightDetailsView) const
 {
-	TArray<TPair<int32, FPropertyPath>> OldProperties = OldDetailsView->GetPropertyRowNumbers();
-	TArray<TPair<int32, FPropertyPath>> NewProperties = NewDetailsView->GetPropertyRowNumbers();
-
-	const int32 OldRowCount = OldDetailsView->CountRows();
-	const int32 NewRowCount = NewDetailsView->CountRows();
-
-	// use caching to avoid O(n^2) LCS calculation every frame
-	if(LinkedScrollRateCache.OldProperties != OldProperties || LinkedScrollRateCache.NewProperties != NewProperties)
-	{
-		const auto PropertyPathsEqual = [](const TPair<int32, FPropertyPath>& A, const TPair<int32, FPropertyPath>& B)
-		{
-			return A.Value == B.Value;
-		};
-		
-		// Find the Longest Common Subsequence between both sets of properties
-		const TArray<TArray<int32>> LCS = CalculateLCSTable(OldProperties, NewProperties, PropertyPathsEqual);
+	const UObject* LeftObject = LeftDetailsView->GetSelectedObjects()[0].Get();
 	
-		// Using the LCS Table, we can determine which lines should match one another while scrolling. For example, an element
-		// of FixedPoints may be {12.f, 5.f} meaning that line 12 of the left panel is the same as line 5 of the right panel
-		TArray<FVector2f> FixedPoints;
+	TArray<FIntVector2> MatchingRows;
 
-		int32 I = OldProperties.Num();
-		int32 J = NewProperties.Num();
-		while (I > 0 && J > 0)
+	// iterate matching rows of both details panels simultaneously
+	auto [LeftRowCount, RightRowCount] = PropertyTreeDifferences[LeftObject].Right->ForEachRow(
+		[&MatchingRows](const TUniquePtr<FAsyncDetailViewDiff::DiffNodeType>& DiffNode, int32 LeftRow, int32 RightRow)->ETreeTraverseControl
 		{
-			if (OldProperties[I - 1].Value == NewProperties[J - 1].Value)
+			// if both trees share this row, sync scrolling here
+			if (DiffNode->ValueA.IsValid() && DiffNode->ValueB.IsValid())
 			{
-				const int32 OldLineNum = OldProperties[I - 1].Key;
-				const int32 NewLineNum = NewProperties[J - 1].Key;
-				// add two endpoints for every matched property because the elements have a width of 1 in the panels
-				FixedPoints.Add(FVector2f((OldLineNum + 1.f) / OldRowCount, (NewLineNum + 1.f) / NewRowCount));
-				FixedPoints.Add(FVector2f((float)OldLineNum / OldRowCount, (float)NewLineNum / NewRowCount));
-				--I;
-				--J;
-			}
-			else if (LCS[I - 1][J] <= LCS[I][J	 - 1])
-			{
-				--J;
-			}
-			else
-			{
-				--I;
+				if (MatchingRows.IsEmpty() || (MatchingRows.Last().X != LeftRow && MatchingRows.Last().Y != RightRow))
+				{
+					MatchingRows.Emplace(LeftRow, RightRow);
+				}
 			}
+			return ETreeTraverseControl::Continue;
 		}
-		FixedPoints.Add({0.f, 0.f});
-		Algo::Reverse(FixedPoints);
-		FixedPoints.Add({1.f, 1.f});
-        	
-		LinkedScrollRateCache.OldProperties = OldProperties;
-		LinkedScrollRateCache.NewProperties = NewProperties;
-		LinkedScrollRateCache.ScrollRate = FixedPoints;
+	);
+
+	TArray<FVector2f> FixedPoints;
+	FixedPoints.Emplace(0.f, 0.f);
+	
+	// normalize fixed points
+	for (FIntVector2& MatchingRow : MatchingRows)
+	{
+		FixedPoints.Emplace(
+			StaticCast<float>(MatchingRow.X) / LeftRowCount,
+			StaticCast<float>(MatchingRow.Y) / RightRowCount
+		);
 	}
+	FixedPoints.Emplace(1.f, 1.f);
 	
-	return LinkedScrollRateCache.ScrollRate;
+	return FixedPoints;
 }
 
 
@@ -1410,15 +1494,10 @@ FReply FCommentDraftTreeEntry::OnCommentPostClicked()
 	return FReply::Unhandled();
 }
 
-FReviewCommentsDiffControl::FReviewCommentsDiffControl(const UObject* InOldObject, const UObject* InNewObject,
-                                                       TWeakPtr<STreeView<TSharedPtr<FBlueprintDifferenceTreeEntry>>> TreeView)
-	: OldObject(InOldObject)
-	, NewObject(InNewObject)
+FReviewCommentsDiffControl::FReviewCommentsDiffControl(const FString& InCommentFilePath, TWeakPtr<STreeView<TSharedPtr<FBlueprintDifferenceTreeEntry>>> TreeView)
+	: CommentFilePath(InCommentFilePath)
 	, CommentsTreeView(TreeView)
 {
-	const UPackage* Package = NewObject ? NewObject->GetPackage() : OldObject->GetPackage();
-	const FPackagePath PackagePath = FPackagePath::FromPackageNameChecked(Package->GetName());
-	CommentFilePath = PackagePath.GetLocalFullPath();
 }
 
 void FReviewCommentsDiffControl::GenerateCommentThreadRecursive(const FReviewComment& Comment,
diff --git a/Engine/Source/Editor/Kismet/Private/DiffControl.h b/Engine/Source/Editor/Kismet/Private/DiffControl.h
--- a/Engine/Source/Editor/Kismet/Private/DiffControl.h
+++ b/Engine/Source/Editor/Kismet/Private/DiffControl.h
@@ -20,6 +20,7 @@ class UEdGraph;
 struct FEdGraphEditAction;
 class SCheckBox;
 
+class FAsyncDetailViewDiff;
 
 /** Interface responsible for generating FBlueprintDifferenceTreeEntry's for visual diff tools */
 class KISMET_API IDiffControl
@@ -27,6 +28,8 @@ class KISMET_API IDiffControl
 public:
 	virtual ~IDiffControl() {}
 
+	virtual void Tick() {};
+	
 	/** Adds widgets to the tree of differences to show */
 	virtual void GenerateTreeEntries(TArray< TSharedPtr<FBlueprintDifferenceTreeEntry> >& OutTreeEntries, TArray< TSharedPtr<FBlueprintDifferenceTreeEntry> >& OutRealDifferences) = 0;
 	
@@ -36,6 +39,7 @@ public:
 
 	// to support comment posting, set this to the tree view that contains the comments.
 	void EnableComments(TWeakPtr<STreeView<TSharedPtr<FBlueprintDifferenceTreeEntry>>> TreeView, const UObject* OldObject, const UObject* NewObject);
+	void EnableComments(TWeakPtr<STreeView<TSharedPtr<FBlueprintDifferenceTreeEntry>>> TreeView, const TArray<const UObject*>& Objects);
 	virtual void EnableComments(TWeakPtr<STreeView<TSharedPtr<FBlueprintDifferenceTreeEntry>>> TreeView) {}
 	
 protected:
@@ -93,37 +97,42 @@ class KISMET_API FDetailsDiffControl : public TSharedFromThis<FDetailsDiffContro
 public:
 	FDetailsDiffControl(const UObject* InOldObject, const UObject* InNewObject, FOnDiffEntryFocused InSelectionCallback, bool bPopulateOutTreeEntries);
 
+	virtual void Tick() override;
+	
 	virtual void GenerateTreeEntries(TArray< TSharedPtr<FBlueprintDifferenceTreeEntry> >& OutTreeEntries, TArray< TSharedPtr<FBlueprintDifferenceTreeEntry> >& OutRealDifferences) override;
 	void GenerateTreeEntriesWithoutComments(TArray< TSharedPtr<FBlueprintDifferenceTreeEntry> >& OutTreeEntries, TArray< TSharedPtr<FBlueprintDifferenceTreeEntry> >& OutRealDifferences);
 
-	TSharedRef<SWidget> OldDetailsWidget() { return OldDetails.DetailsWidget(); }
-	TSharedRef<SWidget> NewDetailsWidget() { return NewDetails.DetailsWidget(); }
+	TSharedRef<IDetailsView> InsertObject(const UObject* Object, bool bScrollbarOnLeft = false, int32 Index = INDEX_NONE);
+	TSharedRef<IDetailsView> GetDetailsWidget(const UObject* Object) const;
+	TSharedPtr<FAsyncDetailViewDiff> GetDifferencesWithLeft(const UObject* Object) const;
+	TSharedPtr<FAsyncDetailViewDiff> GetDifferencesWithRight(const UObject* Object) const;
 
 	// to support comment posting, set this to the tree view that contains the comments.
 	virtual void EnableComments(TWeakPtr<STreeView<TSharedPtr<FBlueprintDifferenceTreeEntry>>> TreeView) override;
-
 protected:
 	virtual void OnSelectDiffEntry(FPropertySoftPath PropertyName);
 	TAttribute<TArray<FVector2f>> GetLinkedScrollRateAttribute(const TSharedRef<IDetailsView>& OldDetailsView, const TSharedRef<IDetailsView>& NewDetailsView);
 	
 
 	// helper function that analyzes two details views and determines the rate they should scroll relative to one another to be in sync
-	TArray<FVector2f> GetLinkedScrollRate(TSharedRef<IDetailsView> OldDetailsView, TSharedRef<IDetailsView> NewDetailsView) const;
+	TArray<FVector2f> GetLinkedScrollRate(TSharedRef<IDetailsView> LeftDetailsView, TSharedRef<IDetailsView> RightDetailsView) const;
 
 	FOnDiffEntryFocused SelectionCallback;
-	FDetailsDiff OldDetails;
-	FDetailsDiff NewDetails;
+	
+	TMap<const UObject*,FDetailsDiff> DetailsDiffs;
+	TArray<const UObject*> ObjectDisplayOrder;
+	struct FPropertyTreeDiffPairs
+	{
+		TSharedPtr<FAsyncDetailViewDiff> Left;
+		TSharedPtr<FAsyncDetailViewDiff> Right;
+	};
+	TMap<const UObject*, FPropertyTreeDiffPairs> PropertyTreeDifferences;
 
-	TArray<FSingleObjectDiffEntry> DifferingProperties;
 	TArray< TSharedPtr<FBlueprintDifferenceTreeEntry> > Children;
 	const bool bPopulateOutTreeEntries;
+	
+	TSet<FPropertyPath> PropertyAllowList;
 
-	mutable struct
-	{
-		TArray<TPair<int32, FPropertyPath>> OldProperties;
-		TArray<TPair<int32, FPropertyPath>> NewProperties;
-		TArray<FVector2f> ScrollRate;
-	} LinkedScrollRateCache;
 };
 
 /** Override for CDO special case */
@@ -349,7 +358,7 @@ private:
 /** Category list item for a graph*/
 struct FReviewCommentsDiffControl : public TSharedFromThis<FReviewCommentsDiffControl>, IDiffControl
 {
-	FReviewCommentsDiffControl(const UObject* InOldObject, const UObject* InNewObject, TWeakPtr<STreeView<TSharedPtr<FBlueprintDifferenceTreeEntry>>> TreeView);
+	FReviewCommentsDiffControl(const FString& InCommentFilePath, TWeakPtr<STreeView<TSharedPtr<FBlueprintDifferenceTreeEntry>>> TreeView);
 
 	virtual void GenerateTreeEntries(TArray< TSharedPtr<FBlueprintDifferenceTreeEntry> >& OutTreeEntries, TArray< TSharedPtr<FBlueprintDifferenceTreeEntry> >& OutRealDifferences) override;
 
@@ -367,9 +376,6 @@ protected:
 		const TMap<int32, TArray<const FReviewComment*>>& CommentReplyMap,
 		TArray<TSharedPtr<FBlueprintDifferenceTreeEntry>>& OutTreeEntries);
 	
-	const UObject* OldObject;
-	const UObject* NewObject;
-	
 	FString CommentFilePath;
 	FString CommentCategory;
 	TWeakPtr<STreeView<TSharedPtr<FBlueprintDifferenceTreeEntry>>> CommentsTreeView;
diff --git a/Engine/Source/Editor/Kismet/Private/SBlueprintDiff.cpp b/Engine/Source/Editor/Kismet/Private/SBlueprintDiff.cpp
--- a/Engine/Source/Editor/Kismet/Private/SBlueprintDiff.cpp
+++ b/Engine/Source/Editor/Kismet/Private/SBlueprintDiff.cpp
@@ -36,6 +36,7 @@
 #include "Misc/CString.h"
 #include "PropertyEditorDelegates.h"
 #include "ReviewComments.h"
+#include "SDetailsSplitter.h"
 #include "SKismetInspector.h"
 #include "SMyBlueprint.h"
 #include "SlateOptMacros.h"
@@ -1101,21 +1102,29 @@ SBlueprintDiff::FDiffControl SBlueprintDiff::GenerateDefaultsPanel()
 	NewDiffControl->EnableComments(DifferencesTreeView.ToWeakPtr());
 	NewDiffControl->GenerateTreeEntries(PrimaryDifferencesList, RealDifferences);
 
+	const TSharedRef<SDetailsSplitter> Splitter = SNew(SDetailsSplitter);
+	if (A)
+	{
+		Splitter->AddSlot(
+			SDetailsSplitter::Slot()
+			.Value(0.5f)
+			.DetailsView(NewDiffControl->GetDetailsWidget(A))
+			.DifferencesWithRightPanel(NewDiffControl.ToSharedRef(), &FDetailsDiffControl::GetDifferencesWithRight, A)
+		);
+	}
+	if (B)
+	{
+		Splitter->AddSlot(
+			SDetailsSplitter::Slot()
+			.Value(0.5f)
+			.DetailsView(NewDiffControl->GetDetailsWidget(B))
+			.DifferencesWithRightPanel(NewDiffControl.ToSharedRef(), &FDetailsDiffControl::GetDifferencesWithRight, B)
+		);
+	}
+
 	SBlueprintDiff::FDiffControl Ret;
 	Ret.DiffControl = NewDiffControl;
-	Ret.Widget = SNew(SSplitter)
-		.PhysicalSplitterHandleSize(10.0f)
-		+ SSplitter::Slot()
-		.Value(0.5f)
-		[
-			NewDiffControl->OldDetailsWidget()
-		]
-		+ SSplitter::Slot()
-		.Value(0.5f)
-		[
-			NewDiffControl->NewDetailsWidget()
-		];
-
+	Ret.Widget = Splitter;
 	return Ret;
 }
 
@@ -1125,21 +1134,29 @@ SBlueprintDiff::FDiffControl SBlueprintDiff::GenerateClassSettingsPanel()
 	NewDiffControl->EnableComments(DifferencesTreeView.ToWeakPtr());
 	NewDiffControl->GenerateTreeEntries(PrimaryDifferencesList, RealDifferences);
 
+	const TSharedRef<SDetailsSplitter> Splitter = SNew(SDetailsSplitter);
+	if (PanelOld.Blueprint)
+	{
+		Splitter->AddSlot(
+			SDetailsSplitter::Slot()
+			.Value(0.5f)
+			.DetailsView(NewDiffControl->GetDetailsWidget(PanelOld.Blueprint))
+			.DifferencesWithRightPanel(NewDiffControl.ToSharedRef(), &FDetailsDiffControl::GetDifferencesWithRight, Cast<UObject>(PanelOld.Blueprint))
+		);
+	}
+	if (PanelNew.Blueprint)
+	{
+		Splitter->AddSlot(
+			SDetailsSplitter::Slot()
+			.Value(0.5f)
+			.DetailsView(NewDiffControl->GetDetailsWidget(PanelNew.Blueprint))
+			.DifferencesWithRightPanel(NewDiffControl.ToSharedRef(), &FDetailsDiffControl::GetDifferencesWithRight, Cast<UObject>(PanelNew.Blueprint))
+		);
+	}
+	
 	SBlueprintDiff::FDiffControl Ret;
 	Ret.DiffControl = NewDiffControl;
-	Ret.Widget = SNew(SSplitter)
-		.PhysicalSplitterHandleSize(10.0f)
-		+ SSplitter::Slot()
-		.Value(0.5f)
-		[
-			NewDiffControl->OldDetailsWidget()
-		]
-		+ SSplitter::Slot()
-		.Value(0.5f)
-		[
-			NewDiffControl->NewDetailsWidget()
-		];
-
+	Ret.Widget = Splitter;
 	return Ret;
 }
 
@@ -1169,8 +1186,9 @@ SBlueprintDiff::FDiffControl SBlueprintDiff::GenerateComponentsPanel()
 
 SBlueprintDiff::FDiffControl SBlueprintDiff::GenerateGeneralFileCommentEntries()
 {
-	TSharedPtr<FReviewCommentsDiffControl> NewDiffControl = MakeShared<FReviewCommentsDiffControl>(
-		PanelOld.Blueprint, PanelNew.Blueprint, DifferencesTreeView.ToWeakPtr());
+	const UPackage* Package = PanelOld.Blueprint ? PanelOld.Blueprint->GetPackage() : PanelNew.Blueprint->GetPackage();
+	const FPackagePath PackagePath = FPackagePath::FromPackageNameChecked(Package->GetName());
+	const TSharedPtr<FReviewCommentsDiffControl> NewDiffControl = MakeShared<FReviewCommentsDiffControl>(PackagePath.GetLocalFullPath(), DifferencesTreeView.ToWeakPtr());
 	NewDiffControl->GenerateTreeEntries(PrimaryDifferencesList, RealDifferences);
 
 	SBlueprintDiff::FDiffControl Ret;
diff --git a/Engine/Source/Editor/Kismet/Private/SDetailsDiff.cpp b/Engine/Source/Editor/Kismet/Private/SDetailsDiff.cpp
--- a/Engine/Source/Editor/Kismet/Private/SDetailsDiff.cpp
+++ b/Engine/Source/Editor/Kismet/Private/SDetailsDiff.cpp
@@ -1,6 +1,7 @@
 // Copyright Epic Games, Inc. All Rights Reserved.
 
 #include "SDetailsDiff.h"
+#include "AsyncDetailViewDiff.h"
 #include "Editor.h"
 #include "Widgets/Layout/SSplitter.h"
 #include "Widgets/SOverlay.h"
@@ -12,17 +13,16 @@
 #include "K2Node_MathExpression.h"
 #include "Kismet2/KismetEditorUtilities.h"
 #include "DetailsDiff.h"
-#include "GraphDiffControl.h"
-#include "WidgetBlueprint.h"
+#include "DetailTreeNode.h"
 #include "HAL/PlatformApplicationMisc.h"
 #include "Framework/Application/SlateApplication.h"
 #include "SBlueprintDiff.h"
 #include "DiffControl.h"
-#include "PropertyEditorModule.h"
-#include "Modules/ModuleManager.h"
 #include "IDetailsView.h"
+#include "SDetailsSplitter.h"
 
 #include "Subsystems/AssetEditorSubsystem.h"
+#include "Widgets/DeclarativeSyntaxSupport.h"
 
 #define LOCTEXT_NAMESPACE "SDetailsDif"
 
@@ -44,8 +44,6 @@ void SDetailsDiff::Construct( const FArguments& InArgs)
 	PanelOld.bShowAssetName = InArgs._ShowAssetNames;
 	PanelNew.bShowAssetName = InArgs._ShowAssetNames;
 
-	bLockViews = true;
-
 	if (InArgs._ParentWindow.IsValid())
 	{
 		WeakParentWindow = InArgs._ParentWindow;
@@ -182,6 +180,11 @@ SDetailsDiff::~SDetailsDiff()
 	}
 }
 
+void SDetailsDiff::Tick(const FGeometry& AllottedGeometry, const double InCurrentTime, const float InDeltaTime)
+{
+	ModePanels[CurrentMode].DiffControl->Tick();
+}
+
 void SDetailsDiff::OnCloseAssetEditor(UObject* Asset, EAssetEditorCloseReason CloseReason)
 {
 	if (PanelOld.Object == Asset || PanelNew.Object == Asset || CloseReason == EAssetEditorCloseReason::CloseAllAssetEditors)
@@ -297,26 +300,6 @@ void SDetailsDiff::GenerateDifferencesList()
 	PrimaryDifferencesList.Empty();
 	RealDifferences.Empty();
 	ModePanels.Empty();
-	
-	const auto CreateInspector = [](const UObject* Object) {
-		FPropertyEditorModule& EditModule = FModuleManager::Get().GetModuleChecked<FPropertyEditorModule>("PropertyEditor");
-	
-		FNotifyHook* NotifyHook = nullptr;
-	
-		FDetailsViewArgs DetailsViewArgs;
-		DetailsViewArgs.NameAreaSettings = FDetailsViewArgs::HideNameArea;
-		DetailsViewArgs.bHideSelectionTip = true;
-		DetailsViewArgs.NotifyHook = NotifyHook;
-		DetailsViewArgs.ViewIdentifier = FName("ObjectInspector");
-		TSharedRef<IDetailsView> DetailsView = EditModule.CreateDetailView( DetailsViewArgs );
-		DetailsView->SetObject(const_cast<UObject*>(Object)); // TODO: @jordan.hoffmann apologize to the const gods for your sins
-		
-		return DetailsView;
-	};
-
-	// TODO: construct DetailsView of PanelOld and PanelNew
-	PanelOld.DetailsView = CreateInspector(PanelOld.Object);
-	PanelNew.DetailsView = CreateInspector(PanelOld.Object);
 
 	// Now that we have done the diffs, create the panel widgets
 	// (we're currently only generating the details panel but we can add more as needed)
@@ -327,25 +310,36 @@ void SDetailsDiff::GenerateDifferencesList()
 
 SDetailsDiff::FDiffControl SDetailsDiff::GenerateDetailsPanel()
 {
-	TSharedPtr<FDetailsDiffControl> NewDiffControl = MakeShared<FDetailsDiffControl>(PanelOld.Object, PanelNew.Object, FOnDiffEntryFocused::CreateRaw(this, &SDetailsDiff::SetCurrentMode, DetailsMode), true);
+	const TSharedPtr<FDetailsDiffControl> NewDiffControl = MakeShared<FDetailsDiffControl>(PanelOld.Object, PanelNew.Object, FOnDiffEntryFocused::CreateRaw(this, &SDetailsDiff::SetCurrentMode, DetailsMode), true);
 	NewDiffControl->EnableComments(DifferencesTreeView.ToWeakPtr());
 	NewDiffControl->GenerateTreeEntries(PrimaryDifferencesList, RealDifferences);
 
+	
+	TWeakPtr<SDetailsSplitter> DetailsSplitter;
+
+	const TSharedRef<SDetailsSplitter> Splitter = SNew(SDetailsSplitter);
+	if (PanelOld.Object)
+	{
+		Splitter->AddSlot(
+			SDetailsSplitter::Slot()
+			.Value(0.5f)
+			.DetailsView(NewDiffControl->GetDetailsWidget(PanelOld.Object))
+			.DifferencesWithRightPanel(NewDiffControl.ToSharedRef(), &FDetailsDiffControl::GetDifferencesWithRight, Cast<UObject>(PanelOld.Object))
+		);
+	}
+	if (PanelNew.Object)
+	{
+		Splitter->AddSlot(
+			SDetailsSplitter::Slot()
+			.Value(0.5f)
+			.DetailsView(NewDiffControl->GetDetailsWidget(PanelNew.Object))
+			.DifferencesWithRightPanel(NewDiffControl.ToSharedRef(), &FDetailsDiffControl::GetDifferencesWithRight, Cast<UObject>(PanelNew.Object))
+		);
+	}
+
 	SDetailsDiff::FDiffControl Ret;
 	Ret.DiffControl = NewDiffControl;
-	Ret.Widget = SNew(SSplitter)
-		.PhysicalSplitterHandleSize(10.0f)
-		+ SSplitter::Slot()
-		.Value(0.5f)
-		[
-			NewDiffControl->OldDetailsWidget()
-		]
-		+ SSplitter::Slot()
-		.Value(0.5f)
-		[
-			NewDiffControl->NewDetailsWidget()
-		];
-
+	Ret.Widget = Splitter;
 	return Ret;
 }
 
@@ -377,10 +371,6 @@ void SDetailsDiff::SetCurrentMode(FName NewMode)
 
 	if (FoundControl)
 	{
-		// Reset inspector view
-		PanelOld.DetailsView->SetObject(nullptr);
-		PanelNew.DetailsView->SetObject(nullptr);
-
 		ModeContents->SetContent(FoundControl->Widget.ToSharedRef());
 	}
 	else
diff --git a/Engine/Source/Editor/Kismet/Public/DetailsDiff.h b/Engine/Source/Editor/Kismet/Public/DetailsDiff.h
--- a/Engine/Source/Editor/Kismet/Public/DetailsDiff.h
+++ b/Engine/Source/Editor/Kismet/Public/DetailsDiff.h
@@ -20,7 +20,7 @@ class KISMET_API FDetailsDiff
 public:
 	DECLARE_DELEGATE( FOnDisplayedPropertiesChanged );
 
-	FDetailsDiff( const UObject* InObject, FOnDisplayedPropertiesChanged InOnDisplayedPropertiesChanged );
+	FDetailsDiff( const UObject* InObject, FOnDisplayedPropertiesChanged InOnDisplayedPropertiesChanged, bool bScrollbarOnLeft = false );
 	~FDetailsDiff();
 
 	/** Attempt to highlight the property with the given path, may not always succeed */
@@ -41,12 +41,14 @@ public:
 	/** Link the two details panels so they scroll in sync with one another */
 	static void LinkScrolling(FDetailsDiff& LeftPanel, FDetailsDiff& RightPanel, const TAttribute<TArray<FVector2f>>& ScrollRate);
 
+	static TSharedRef<IDetailsView> CreateDetailsView(const UObject* InObject, TSharedPtr<SScrollBar> ExternalScrollbar = nullptr, bool bScrollbarOnLeft = false);
+
 private:
 	void HandlePropertiesChanged();
 
 	FOnDisplayedPropertiesChanged OnDisplayedPropertiesChanged;
 
-	TArray< FPropertyPath > DifferingProperties;
+	TArray< FPropertyPath > DisplayedProperties;
 	const UObject* DisplayedObject;
 
 	TSharedPtr< class IDetailsView > DetailsView;
diff --git a/Engine/Source/Editor/Kismet/Public/SCSDiff.h b/Engine/Source/Editor/Kismet/Public/SCSDiff.h
--- a/Engine/Source/Editor/Kismet/Public/SCSDiff.h
+++ b/Engine/Source/Editor/Kismet/Public/SCSDiff.h
@@ -39,5 +39,5 @@ private:
 	TSharedPtr< class SKismetInspector > Inspector;
 
 	/** Blueprint we are inspecting */
-	UBlueprint* Blueprint;
+	UBlueprint* Blueprint = nullptr;
 };
diff --git a/Engine/Source/Editor/Kismet/Public/SDetailsDiff.h b/Engine/Source/Editor/Kismet/Public/SDetailsDiff.h
--- a/Engine/Source/Editor/Kismet/Public/SDetailsDiff.h
+++ b/Engine/Source/Editor/Kismet/Public/SDetailsDiff.h
@@ -12,6 +12,7 @@ class FSpawnTabArgs;
 class FTabManager;
 class IDiffControl;
 class FUICommandList;
+class SDetailsSplitter;
 enum class EAssetEditorCloseReason : uint8;
 
 /** Panel used to display the details */
@@ -20,9 +21,6 @@ struct KISMET_API FDetailsDiffPanel
 	/** The asset that owns the details view we are showing */
 	const UObject*				Object = nullptr;
 
-	/** The details view associated with the asset */
-	TSharedPtr<class IDetailsView>	DetailsView;
-
 	/** Revision information for this asset */
 	FRevisionInfo					RevisionInfo;
 
@@ -52,6 +50,8 @@ public:
 	void Construct(const FArguments& InArgs);
 	virtual ~SDetailsDiff();
 
+	virtual void Tick(const FGeometry& AllottedGeometry, const double InCurrentTime, const float InDeltaTime) override;
+
 	/** Called when user clicks on an entry in the listview of differences */
 	void OnDiffListSelectionChanged(TSharedPtr<struct FDiffResultItem> TheDiff);
 
diff --git a/Engine/Source/Editor/PropertyEditor/Private/SDetailsSplitter.cpp b/Engine/Source/Editor/PropertyEditor/Private/SDetailsSplitter.cpp
new file mode 100644
--- /dev/null
+++ b/Engine/Source/Editor/PropertyEditor/Private/SDetailsSplitter.cpp
@@ -0,0 +1,763 @@
+// Copyright Epic Games, Inc. All Rights Reserved.
+#include "SDetailsSplitter.h"
+
+#include "AsyncDetailViewDiff.h"
+#include "DetailTreeNode.h"
+#include "IDetailsViewPrivate.h"
+#include "PropertyNode.h"
+#include "Styling/StyleColors.h"
+#include "Framework/Application/SlateApplication.h"
+
+// converts from a container helper like FScriptArrayHelper to a property like FArrayProperty
+template<typename HelperType>
+using TContainerPropertyType =
+	std::conditional_t<std::is_same_v<HelperType, FScriptArrayHelper>, FArrayProperty,
+	std::conditional_t<std::is_same_v<HelperType, FScriptMapHelper>, FMapProperty,
+	std::conditional_t<std::is_same_v<HelperType, FScriptSetHelper>, FSetProperty,
+	void>>>;
+
+template<typename HelperType>
+static bool TryGetSourceContainer(TSharedPtr<FDetailTreeNode> DetailsNode, FPropertyNode*& OutPropNode, TUniquePtr<HelperType>& OutResolvedProperty)
+{
+	using ContainerPropertyType = TContainerPropertyType<HelperType>;
+	if (TryGetSourceContainer<ContainerPropertyType>(DetailsNode, OutPropNode))
+	{
+		const FPropertySoftPath SoftPropertyPath(FPropertyNode::CreatePropertyPath(OutPropNode->AsShared()).Get());
+		const UObject* Object = DetailsNode->GetDetailsView()->GetSelectedObjects()[0].Get();
+		const FResolvedProperty Resolved = SoftPropertyPath.Resolve(Object);
+		const ContainerPropertyType* ContainerProperty = CastFieldChecked<ContainerPropertyType>(Resolved.Property);
+		OutResolvedProperty = MakeUnique<HelperType>(ContainerProperty, ContainerProperty->template ContainerPtrToValuePtr<UObject*>(Resolved.Object));
+		return true;
+	}
+	return false;
+}
+
+template<typename HelperType>
+static bool TryGetDestinationContainer(TSharedPtr<FDetailTreeNode> DetailsNode, FPropertyNode*& OutPropNode, TUniquePtr<HelperType>& OutContainerHelper, int32& OutInsertIndex)
+{
+	using ContainerPropertyType = TContainerPropertyType<HelperType>;
+	if (TryGetDestinationContainer<ContainerPropertyType>(DetailsNode, OutPropNode, OutInsertIndex))
+	{
+		const FPropertySoftPath SoftPropertyPath(FPropertyNode::CreatePropertyPath(OutPropNode->AsShared()).Get());
+		const UObject* Object = DetailsNode->GetDetailsView()->GetSelectedObjects()[0].Get();
+		const FResolvedProperty Resolved = SoftPropertyPath.Resolve(Object);
+		const ContainerPropertyType* ContainerProperty = CastFieldChecked<ContainerPropertyType>(Resolved.Property);
+		OutContainerHelper = MakeUnique<HelperType>(ContainerProperty, ContainerProperty->template ContainerPtrToValuePtr<UObject*>(Resolved.Object));
+		return true;
+	}
+	return false;
+}
+
+template<typename ContainerPropertyType>
+static bool TryGetSourceContainer(TSharedPtr<FDetailTreeNode> DetailsNode, FPropertyNode*& OutPropNode)
+{
+	if ((OutPropNode = DetailsNode->GetPropertyNode().Get()) == nullptr)
+	{
+		return false;
+	}
+	if ((OutPropNode = OutPropNode->GetParentNode()) == nullptr)
+	{
+		return false;
+	}
+	if (CastField<ContainerPropertyType>(OutPropNode->GetProperty()) != nullptr)
+	{
+		return true;
+	}
+	return false;
+}
+
+template<typename ContainerPropertyType>
+static bool TryGetDestinationContainer(TSharedPtr<FDetailTreeNode> DetailsNode, FPropertyNode*& OutPropNode, int32& OutInsertIndex)
+{
+	if ((OutPropNode = DetailsNode->GetPropertyNode().Get()) == nullptr)
+	{
+		return false;
+	}
+	if (CastField<ContainerPropertyType>(OutPropNode->GetProperty()) != nullptr)
+	{
+		OutInsertIndex = 0;
+		return true;
+	}
+	OutInsertIndex = OutPropNode->GetArrayIndex() + 1;
+	while((OutPropNode = OutPropNode->GetParentNode()) != nullptr)
+	{
+		if (CastField<ContainerPropertyType>(OutPropNode->GetProperty()) != nullptr)
+		{
+			return true;
+		}
+		OutInsertIndex = OutPropNode->GetArrayIndex() + 1;
+	}
+	return false;
+}
+
+static void CopyPropertyValueForInsert(TSharedPtr<FDetailTreeNode> SourceDetailsNode, TSharedPtr<FDetailTreeNode> DestinationDetailsNode)
+{
+	TUniquePtr<FScriptArrayHelper> SourceArray;
+	FPropertyNode* SourceArrayPropertyNode;
+	if (TryGetSourceContainer(SourceDetailsNode, SourceArrayPropertyNode, SourceArray))
+	{
+		int32 InsertIndex;
+		TUniquePtr<FScriptArrayHelper> DestinationArray;
+		FPropertyNode* DestinationArrayPropertyNode;
+		if (ensure(TryGetDestinationContainer(DestinationDetailsNode, DestinationArrayPropertyNode, DestinationArray, InsertIndex)))
+		{
+			DestinationArray->InsertValues(InsertIndex, 1);
+			const void* SourceData = SourceArray->GetElementPtr(SourceDetailsNode->GetPropertyNode()->GetArrayIndex());
+			void* DestinationData = DestinationArray->GetElementPtr(InsertIndex);
+			const FArrayProperty* ArrayProperty = CastFieldChecked<FArrayProperty>(DestinationArrayPropertyNode->GetProperty());
+			const FProperty* ElementProperty = ArrayProperty->Inner;
+			ElementProperty->CopySingleValue(DestinationData, SourceData);
+		}
+	}
+	
+	TUniquePtr<FScriptSetHelper> SourceSet;
+	FPropertyNode* SourceSetPropertyNode;
+	if (TryGetSourceContainer(SourceDetailsNode, SourceSetPropertyNode, SourceSet))
+	{
+		int32 InsertIndex;
+		TUniquePtr<FScriptSetHelper> DestinationSet;
+		FPropertyNode* DestinationPropertyNode;
+		if (ensure(TryGetDestinationContainer(DestinationDetailsNode, DestinationPropertyNode, DestinationSet, InsertIndex)))
+		{
+			const void* SourceData = SourceSet->GetElementPtr(SourceDetailsNode->GetPropertyNode()->GetArrayIndex());
+			DestinationSet->AddElement(SourceData);
+		}
+	}
+	
+	TUniquePtr<FScriptMapHelper> SourceMap;
+	FPropertyNode* SourceMapPropertyNode;
+	if (TryGetSourceContainer(SourceDetailsNode, SourceMapPropertyNode, SourceMap))
+	{
+		int32 InsertIndex;
+		TUniquePtr<FScriptMapHelper> DestinationMap;
+		FPropertyNode* DestinationPropertyNode;
+		if (ensure(TryGetDestinationContainer(DestinationDetailsNode, DestinationPropertyNode, DestinationMap, InsertIndex)))
+		{
+			const void* SourceKey = SourceMap->GetKeyPtr(SourceDetailsNode->GetPropertyNode()->GetArrayIndex());
+			const void* SourceVal = SourceMap->GetValuePtr(SourceDetailsNode->GetPropertyNode()->GetArrayIndex());
+			DestinationMap->AddPair(SourceKey, SourceVal);
+		}
+	}
+}
+
+static void CopyPropertyValue(TSharedPtr<FDetailTreeNode> SourceDetailsNode, TSharedPtr<FDetailTreeNode> DestinationDetailsNode, ETreeDiffResult Diff)
+{
+	switch(Diff)
+	{
+		// traditional copy
+		case ETreeDiffResult::DifferentValues:
+			break;
+
+		// insert
+		case ETreeDiffResult::MissingFromTree1:
+		case ETreeDiffResult::MissingFromTree2:
+			CopyPropertyValueForInsert(SourceDetailsNode, DestinationDetailsNode);
+			return;
+		
+		// no difference
+		case ETreeDiffResult::Invalid:
+		case ETreeDiffResult::Identical:
+		default:
+			return;
+	}
+	
+	FResolvedProperty SourceResolved;
+	FResolvedProperty DestinationResolved;
+	bool bCopySingleValue = false;
+	
+	if (SourceDetailsNode)
+	{
+		const UObject* Object = SourceDetailsNode->GetDetailsView()->GetSelectedObjects()[0].Get();
+		const FPropertyPath PropertyPath = SourceDetailsNode->GetPropertyPath();
+		SourceResolved = FPropertySoftPath(PropertyPath).Resolve(Object);
+		bCopySingleValue = PropertyPath.GetLeafMostProperty().ArrayIndex != INDEX_NONE;
+	}
+	if (DestinationDetailsNode)
+	{
+		const UObject* Object = DestinationDetailsNode->GetDetailsView()->GetSelectedObjects()[0].Get();
+		const FPropertyPath PropertyPath = DestinationDetailsNode->GetPropertyPath();
+		DestinationResolved = FPropertySoftPath(PropertyPath).Resolve(Object);
+		bCopySingleValue = PropertyPath.GetLeafMostProperty().ArrayIndex != INDEX_NONE;
+	}
+		
+	if (SourceResolved.Property && DestinationResolved.Property && SourceResolved.Object && DestinationResolved.Object)
+	{
+		const void* SourceData = SourceResolved.Property->ContainerPtrToValuePtr<void*>(SourceResolved.Object);
+		void* DestinationData = (void*)DestinationResolved.Property->ContainerPtrToValuePtr<void*>(DestinationResolved.Object);
+		if (bCopySingleValue)
+		{
+			DestinationResolved.Property->CopySingleValue(DestinationData, SourceData);
+		}
+		else
+		{
+			DestinationResolved.Property->CopyCompleteValue(DestinationData, SourceData);
+		}
+	}
+}
+
+// note: DestinationDetailsNode is the node before the position in the tree you wish to insert
+static bool CanCopyPropertyValueForInsert(TSharedPtr<FDetailTreeNode> SourceDetailsNode, TSharedPtr<FDetailTreeNode> DestinationDetailsNode)
+{
+	FPropertyNode* SourceArrayPropertyNode;
+	if (TryGetSourceContainer<FArrayProperty>(SourceDetailsNode, SourceArrayPropertyNode))
+	{
+		FPropertyNode* DestinationArrayPropertyNode;
+		int32 InsertIndex;
+		if (TryGetDestinationContainer<FArrayProperty>(DestinationDetailsNode, DestinationArrayPropertyNode, InsertIndex))
+		{
+			return true;
+		}
+		return false;
+	}
+	
+	FPropertyNode* SourceSetPropertyNode;
+	if (TryGetSourceContainer<FSetProperty>(SourceDetailsNode, SourceSetPropertyNode))
+	{
+		FPropertyNode* DestinationSetPropertyNode;
+		int32 InsertIndex;
+		if (TryGetDestinationContainer<FSetProperty>(DestinationDetailsNode, DestinationSetPropertyNode, InsertIndex))
+		{
+			return true;
+		}
+		return false;
+	}
+	
+	FPropertyNode* SourceMapPropertyNode;
+	if (TryGetSourceContainer<FMapProperty>(SourceDetailsNode, SourceMapPropertyNode))
+	{
+		FPropertyNode* DestinationMapPropertyNode;
+		int32 InsertIndex;
+		if (TryGetDestinationContainer<FMapProperty>(DestinationDetailsNode, DestinationMapPropertyNode, InsertIndex))
+		{
+			return true;
+		}
+		return false;
+	}
+
+	// you can only insert into containers
+	return false;
+}
+
+static bool CanCopyPropertyValue(TSharedPtr<FDetailTreeNode> SourceDetailsNode, TSharedPtr<FDetailTreeNode> DestinationDetailsNode, ETreeDiffResult Diff)
+{
+	switch(Diff)
+	{
+	// traditional copy
+	case ETreeDiffResult::DifferentValues:
+		return true;
+
+	// insert
+	case ETreeDiffResult::MissingFromTree1:
+	case ETreeDiffResult::MissingFromTree2:
+		return CanCopyPropertyValueForInsert(SourceDetailsNode, DestinationDetailsNode);
+		
+	// no difference
+	case ETreeDiffResult::Invalid:
+	case ETreeDiffResult::Identical:
+	default:
+		return false;
+	}
+}
+
+void SDetailsSplitter::Construct(const FArguments& InArgs)
+{
+	Splitter = SNew(SSplitter).PhysicalSplitterHandleSize(5.f);
+	for(const FSlot::FSlotArguments& SlotArgs : InArgs._Slots)
+	{
+		AddSlot(SlotArgs);
+	}
+	
+	ChildSlot
+	[
+		Splitter.ToSharedRef()
+	];
+}
+
+void SDetailsSplitter::AddSlot(const FSlot::FSlotArguments& SlotArgs, int32 Index)
+{
+	if (Index == INDEX_NONE)
+	{
+		Index = Panels.Num();
+	}
+	
+	Splitter->AddSlot(Index)
+		.Value(SlotArgs._Value)
+	[
+		SNew(SBox).Padding(15.f,0.f, 15.f,0.f)
+		[
+			SlotArgs._DetailsView.ToSharedRef()
+		]
+	];
+	Panels.Insert({
+		SlotArgs._DetailsView,
+		SlotArgs._IsReadonly,
+		SlotArgs._DifferencesWithRightPanel,
+	}, Index);
+}
+
+SDetailsSplitter::FSlot::FSlotArguments SDetailsSplitter::Slot()
+{
+	return FSlot::FSlotArguments(MakeUnique<FSlot>());
+}
+
+int32 SDetailsSplitter::OnPaint(const FPaintArgs& Args, const FGeometry& AllottedGeometry,
+                                const FSlateRect& MyCullingRect, FSlateWindowElementList& OutDrawElements, int32 LayerId,
+                                const FWidgetStyle& InWidgetStyle, bool bParentEnabled) const
+{
+	int32 MaxLayerId = LayerId;
+	MaxLayerId = Splitter->Paint(Args,AllottedGeometry,MyCullingRect,OutDrawElements,MaxLayerId,InWidgetStyle,bParentEnabled);
+	++MaxLayerId;
+
+	TMap<FSlateRect, FLinearColor> RowHighlights;
+	for (int32 LeftIndex = 0; LeftIndex < Panels.Num(); ++LeftIndex)
+	{
+		const FPanel& LeftPanel = Panels[LeftIndex];
+		if (!LeftPanel.DiffRight.IsBound())
+		{
+			continue;
+		}
+		
+		if (const TSharedPtr<FAsyncDetailViewDiff> Diff = LeftPanel.DiffRight.Get())
+		{
+			FSlateRect PrevLeftPropertyRect;
+			FSlateRect PrevRightPropertyRect;
+			
+			TSharedPtr<FDetailTreeNode> LastSeenRightDetailsNode;
+			TSharedPtr<FDetailTreeNode> LastSeenLeftDetailsNode;
+			
+			Diff->ForEachRow([&](const TUniquePtr<FAsyncDetailViewDiff::DiffNodeType>& DiffNode, int32, int32)->ETreeTraverseControl
+			{
+				const FLinearColor Color = FLinearColor(0.f,1.f,1.f);
+				
+				FSlateRect LeftPropertyRect;
+				if (const TSharedPtr<FDetailTreeNode> LeftDetailNode = DiffNode->ValueA.Pin())
+				{
+					LastSeenLeftDetailsNode = LeftDetailNode;
+					// if the other tree doesn't have a matching node, treat this and all it's children as a single group
+					const bool bIncludeChildren = !DiffNode->ValueB.IsValid();
+					LeftPropertyRect = LeftPanel.DetailsView->GetPaintSpacePropertyBounds(LeftDetailNode.ToSharedRef(), bIncludeChildren);
+				}
+				if (!LeftPropertyRect.IsValid() && PrevLeftPropertyRect.IsValid())
+				{
+					LeftPropertyRect = PrevLeftPropertyRect;
+					LeftPropertyRect.Top = LeftPropertyRect.Bottom;
+				}
+				if (LeftPropertyRect.IsValid() && DiffNode->DiffResult != ETreeDiffResult::Identical)
+				{
+					RowHighlights.Add(LeftPropertyRect, Color);
+				}
+		
+				const int32 RightIndex = LeftIndex + 1;
+				FSlateRect RightPropertyRect;
+				if (Panels.IsValidIndex(RightIndex))
+				{
+					const FPanel& RightPanel = Panels[RightIndex];
+					if (const TSharedPtr<FDetailTreeNode> RightDetailNode = DiffNode->ValueB.Pin())
+					{
+						LastSeenRightDetailsNode = RightDetailNode;
+						// if the other tree doesn't have a matching node, treat this and all it's children as a single group
+						const bool bIncludeChildren = !DiffNode->ValueA.IsValid();
+						RightPropertyRect = RightPanel.DetailsView->GetPaintSpacePropertyBounds(RightDetailNode.ToSharedRef(), bIncludeChildren);
+					}
+					if (!RightPropertyRect.IsValid() && PrevRightPropertyRect.IsValid())
+					{
+						RightPropertyRect = PrevRightPropertyRect;
+						RightPropertyRect.Top = RightPropertyRect.Bottom;
+					}
+					if (RightPropertyRect.IsValid() && DiffNode->DiffResult != ETreeDiffResult::Identical)
+					{
+						RowHighlights.Add(RightPropertyRect, Color);
+					}
+					
+					if (LeftPropertyRect.IsValid() && RightPropertyRect.IsValid() && DiffNode->DiffResult != ETreeDiffResult::Identical)
+					{
+						FLinearColor FillColor = Color.Desaturate(.3f) * FLinearColor(0.04f,0.04f,0.04f);
+						FillColor.A = 0.3f;
+			
+						FLinearColor OutlineColor = Color;
+						OutlineColor.A = 0.7f;
+						PaintPropertyConnector(OutDrawElements, MaxLayerId, LeftPropertyRect, RightPropertyRect, FillColor, OutlineColor);
+						++MaxLayerId;
+
+						if (!RightPanel.IsReadonly.Get(true) && CanCopyPropertyValue(LastSeenLeftDetailsNode, LastSeenRightDetailsNode, DiffNode->DiffResult))
+						{
+							PaintCopyPropertyButton(OutDrawElements, MaxLayerId, DiffNode, LeftPropertyRect, RightPropertyRect, EPropertyCopyDirection::CopyLeftToRight);
+						}
+						if (!LeftPanel.IsReadonly.Get(true) && CanCopyPropertyValue(LastSeenRightDetailsNode, LastSeenLeftDetailsNode, DiffNode->DiffResult))
+						{
+							PaintCopyPropertyButton(OutDrawElements, MaxLayerId, DiffNode, LeftPropertyRect, RightPropertyRect, EPropertyCopyDirection::CopyRightToLeft);
+						}
+						++MaxLayerId;
+						
+					}
+				}
+		
+				PrevLeftPropertyRect = MoveTemp(LeftPropertyRect);
+				PrevRightPropertyRect = MoveTemp(RightPropertyRect);
+
+				// only traverse children if both trees have them
+				if (!DiffNode->ValueA.IsValid() || !DiffNode->ValueB.IsValid())
+				{
+					return ETreeTraverseControl::SkipChildren;
+				}
+				return ETreeTraverseControl::Continue;
+				
+			});
+		}
+	}
+
+	for (const auto& [PropertyRect, Color] : RowHighlights)
+	{
+		FLinearColor FillColor = Color.Desaturate(.3) * FLinearColor(0.04f,0.04f,0.04f);
+		FillColor.A = 0.3f;
+		
+		FPaintGeometry Geometry(
+			PropertyRect.GetTopLeft() + FVector2D{0.f,2.f},
+			PropertyRect.GetSize() - FVector2D{0.f,4.f},
+			1.f
+		);
+		
+		FSlateDrawElement::MakeBox(
+			OutDrawElements,
+			LayerId + 8, // draw in front of background but behind text, buttons, etc
+			Geometry,
+			FAppStyle::GetBrush("WhiteBrush"),
+			ESlateDrawEffect::None,
+			FillColor
+		);
+	}
+	
+	return MaxLayerId;
+}
+
+FReply SDetailsSplitter::OnMouseMove(const FGeometry& MyGeometry, const FPointerEvent& MouseEvent)
+{
+	HoveredCopyButton = {};
+	const FVector2D MousePosition = MouseEvent.GetScreenSpacePosition();
+	
+	for (int32 LeftIndex = 0; LeftIndex < Panels.Num() - 1; ++LeftIndex)
+	{
+		const FPanel& LeftPanel = Panels[LeftIndex];
+		if (!LeftPanel.DiffRight.IsBound())
+		{
+			continue;
+		}
+
+		
+		const int32 RightIndex = LeftIndex + 1;
+		const FPanel& RightPanel = Panels[RightIndex];
+		if (LeftPanel.IsReadonly.Get(false) && RightPanel.IsReadonly.Get(false))
+		{
+			continue;
+		}
+		if (const TSharedPtr<FAsyncDetailViewDiff> Diff = LeftPanel.DiffRight.Get())
+		{
+			
+			TSharedPtr<FDetailTreeNode> LastSeenRightDetailsNode;
+			TSharedPtr<FDetailTreeNode> LastSeenLeftDetailsNode;
+			Diff->ForEachRow([&LastSeenLeftDetailsNode, &LastSeenRightDetailsNode, &LeftPanel, &RightPanel, &MousePosition, &HoveredCopyButton = HoveredCopyButton]
+			(const TUniquePtr<FAsyncDetailViewDiff::DiffNodeType>& DiffNode, int32, int32)->ETreeTraverseControl
+			{
+				
+				FSlateRect LeftPropertyRect;
+				if (const TSharedPtr<FDetailTreeNode> LeftDetailNode = DiffNode->ValueA.Pin())
+				{
+					LastSeenLeftDetailsNode = LeftDetailNode;
+					const bool bIncludeChildren = !DiffNode->ValueB.IsValid();
+					LeftPropertyRect = LeftPanel.DetailsView->GetTickSpacePropertyBounds(LeftDetailNode.ToSharedRef(), bIncludeChildren);
+				}
+		
+				FSlateRect RightPropertyRect;
+				if (const TSharedPtr<FDetailTreeNode> RightDetailNode = DiffNode->ValueB.Pin())
+				{
+					LastSeenRightDetailsNode = RightDetailNode;
+					const bool bIncludeChildren = !DiffNode->ValueA.IsValid();
+					RightPropertyRect = RightPanel.DetailsView->GetTickSpacePropertyBounds(RightDetailNode.ToSharedRef(), bIncludeChildren);
+				}
+				
+				if (DiffNode->DiffResult == ETreeDiffResult::Identical)
+				{
+					return ETreeTraverseControl::Continue;
+				}
+
+				if (LeftPropertyRect.IsValid() && !RightPanel.IsReadonly.Get(false))
+				{
+					const FSlateRect CopyButtonZoneLeftToRight = FSlateRect(
+						LeftPropertyRect.Right,
+						LeftPropertyRect.Top,
+						LeftPropertyRect.Right + 15.f,
+						LeftPropertyRect.Bottom
+					);
+
+					if (CopyButtonZoneLeftToRight.ContainsPoint(MousePosition))
+					{
+						if (CanCopyPropertyValue(LastSeenLeftDetailsNode, LastSeenRightDetailsNode, DiffNode->DiffResult))
+						{
+							HoveredCopyButton = {
+								LastSeenLeftDetailsNode,
+								LastSeenRightDetailsNode,
+								DiffNode->DiffResult,
+								EPropertyCopyDirection::CopyLeftToRight
+							};
+						}
+						return ETreeTraverseControl::Break;
+					}
+				}
+				
+				if (RightPropertyRect.IsValid() && !LeftPanel.IsReadonly.Get(false))
+				{
+					const FSlateRect CopyButtonZoneRightToLeft = FSlateRect(
+						RightPropertyRect.Left - 15.f,
+						RightPropertyRect.Top,
+						RightPropertyRect.Left,
+						RightPropertyRect.Bottom
+					);
+
+					if (CopyButtonZoneRightToLeft.ContainsPoint(MousePosition))
+					{
+						if (CanCopyPropertyValue(LastSeenRightDetailsNode, LastSeenLeftDetailsNode, DiffNode->DiffResult))
+						{
+							HoveredCopyButton = {
+								LastSeenRightDetailsNode,
+								LastSeenLeftDetailsNode,
+								DiffNode->DiffResult,
+								EPropertyCopyDirection::CopyRightToLeft
+							};
+						}
+						return ETreeTraverseControl::Break;
+					}
+				}
+				
+				return ETreeTraverseControl::Continue;
+			});
+
+			if (HoveredCopyButton.CopyDirection != EPropertyCopyDirection::Copy_None)
+			{
+				break;
+			}
+		}
+	}
+	return SCompoundWidget::OnMouseMove(MyGeometry, MouseEvent);
+}
+
+void SDetailsSplitter::OnMouseLeave(const FPointerEvent& MouseEvent)
+{
+	HoveredCopyButton = {};
+	SCompoundWidget::OnMouseLeave(MouseEvent);
+}
+
+FReply SDetailsSplitter::OnMouseButtonDown(const FGeometry& MyGeometry, const FPointerEvent& MouseEvent)
+{
+	if (MouseEvent.GetEffectingButton() != EKeys::LeftMouseButton)
+	{
+		return FReply::Unhandled();
+	}
+	
+	if (HoveredCopyButton.CopyDirection == EPropertyCopyDirection::Copy_None)
+	{
+		return FReply::Unhandled();
+	}
+
+	CopyPropertyValue(
+		HoveredCopyButton.SourceDetailsNode.Pin(),
+		HoveredCopyButton.DestinationDetailsNode.Pin(),
+		HoveredCopyButton.DiffResult);
+	return FReply::Handled();
+}
+
+void SDetailsSplitter::PaintPropertyConnector(FSlateWindowElementList& OutDrawElements, int32 LayerId, const FSlateRect& LeftPropertyRect,
+	const FSlateRect& RightPropertyRect, const FLinearColor& FillColor, const FLinearColor& OutlineColor) const
+{
+	FVector2D TopLeft = LeftPropertyRect.GetTopRight();
+	FVector2D BottomLeft = LeftPropertyRect.GetBottomRight();
+	
+	FVector2D TopRight = RightPropertyRect.GetTopLeft();
+	FVector2D BottomRight = RightPropertyRect.GetBottomLeft();
+
+	{
+		constexpr float YPadding = 2.f;
+		if (BottomLeft.Y - TopLeft.Y > YPadding * 2.f)
+		{
+			BottomLeft.Y -= YPadding;
+			TopLeft.Y += YPadding;
+		}
+		if (BottomRight.Y - TopRight.Y > YPadding * 2.f)
+		{
+			BottomRight.Y -= YPadding;
+			TopRight.Y += YPadding;
+		}
+	}
+	
+	TArray<FSlateVertex> FillVerts;
+	TArray<SlateIndex> FillIndices;
+	TArray<FVector2D> TopBoarderLine;
+	TArray<FVector2D> BottomBoarderLine;
+
+	auto AddVert = [&FillVerts, &FillIndices](const FVector2D& Position, const FColor& VertColor)
+	{
+		FillVerts.Add(FSlateVertex::Make<ESlateVertexRounding::Disabled>(
+			FSlateRenderTransform(),
+			FVector2f(Position),
+			{0.f,0.f},
+			VertColor
+		));
+		if (FillVerts.Num() >= 3)
+		{
+			FillIndices.Add(FillVerts.Num()-3);
+			FillIndices.Add(FillVerts.Num()-2);
+			FillIndices.Add(FillVerts.Num()-1);
+		}
+	};
+
+	// interpolate between left and right corners and add vertices to the mesh
+	constexpr float StepSize = 1.f / 30.f;
+	constexpr float InterpBoarder = .3f; // make room for buttons
+
+		FillVerts.Empty();
+		FillIndices.Empty();
+		TopBoarderLine.Empty();
+		BottomBoarderLine.Empty();
+
+	float Alpha = 0.f;
+	while (true)
+	{
+		constexpr float Exp = 3.f;
+		const float TopAlpha = FMath::Clamp(Alpha, 0.f, 1.f);
+		const float BottomAlpha = FMath::Clamp(Alpha, 0.f, 1.f);
+		const double TopX = FMath::Lerp(TopLeft.X, TopRight.X, TopAlpha);
+		const double BottomX = FMath::Lerp(TopLeft.X, TopRight.X, BottomAlpha);
+		
+		constexpr float InterpolationRange = 1.f - 2.f * InterpBoarder;
+
+		const float TopTransformedAlpha = FMath::Clamp((TopAlpha - InterpBoarder) / InterpolationRange, 0.f, 1.f);
+		const float BottomTransformedAlpha = FMath::Clamp((BottomAlpha - InterpBoarder) / InterpolationRange, 0.f, 1.f);
+		
+		const double TopY = FMath::InterpEaseInOut(TopLeft.Y, TopRight.Y, TopTransformedAlpha, Exp);
+		const double BottomY = FMath::InterpEaseInOut(BottomLeft.Y, BottomRight.Y, BottomTransformedAlpha, Exp);
+
+		FLinearColor ColumnColor = FillColor;
+		if (Alpha <= 0.5f)
+		{
+			ColumnColor.A = FMath::InterpEaseOut(FillColor.A, 1.f, Alpha * 2.f, 2.f);
+		}
+		else
+		{
+			ColumnColor.A = FMath::InterpEaseIn( 1.f,FillColor.A, (Alpha - 0.5f) * 2.f, 2.f);
+		}
+	
+		TopBoarderLine.Emplace(TopX,TopY);
+		FLinearColor TopColor = ColumnColor.Desaturate(0.5f);
+		AddVert(TopBoarderLine.Last(), TopColor.ToFColorSRGB());
+	
+		BottomBoarderLine.Emplace(BottomX,BottomY);
+		FLinearColor BottomColor = ColumnColor;
+		AddVert(BottomBoarderLine.Last(), BottomColor.ToFColorSRGB());
+	
+		if (Alpha >= 1.f)
+		{
+			break;
+		}
+		Alpha = FMath::Clamp(Alpha + StepSize, 0.f, 1.f);
+	}
+	
+	const FSlateResourceHandle ResourceHandle = FSlateApplication::Get().GetRenderer()->GetResourceHandle(*FAppStyle::GetBrush("WhiteBrush"));
+	FSlateDrawElement::MakeCustomVerts(
+		OutDrawElements,
+		LayerId,
+		ResourceHandle,
+		FillVerts,
+		FillIndices,
+		nullptr,
+		0,
+		0,
+		ESlateDrawEffect::None
+	);
+
+	FSlateDrawElement::MakeLines(
+		OutDrawElements,
+		LayerId,
+		FPaintGeometry(),
+		TopBoarderLine,
+		ESlateDrawEffect::None,
+		OutlineColor,
+		true,
+		0.5f
+	);
+
+	FSlateDrawElement::MakeLines(
+		OutDrawElements,
+		LayerId,
+		FPaintGeometry(),
+		BottomBoarderLine,
+		ESlateDrawEffect::None,
+		OutlineColor,
+		true,
+		0.5f
+	);
+}
+
+void SDetailsSplitter::PaintCopyPropertyButton(FSlateWindowElementList& OutDrawElements, int32 LayerId, const TUniquePtr<FAsyncDetailViewDiff::DiffNodeType>& DiffNode,
+	const FSlateRect& LeftPropertyRect, const FSlateRect& RightPropertyRect, EPropertyCopyDirection CopyDirection) const
+{
+	constexpr float ButtonScale = 15.f;
+	FPaintGeometry Geometry;
+	const FSlateBrush* Brush = nullptr;
+	FLinearColor ButtonColor = FStyleColors::Foreground.GetSpecifiedColor();
+	switch (CopyDirection)
+	{
+	case EPropertyCopyDirection::CopyLeftToRight:
+		if (LeftPropertyRect.GetSize().Y < UE_SMALL_NUMBER)
+		{
+			// space is too small to fit a button
+			return;
+		}
+		Geometry = FPaintGeometry(
+			FVector2D(LeftPropertyRect.Right, LeftPropertyRect.GetCenter().Y - 0.5f * ButtonScale),
+			FVector2D(ButtonScale),
+			1.f
+		);
+		Brush = FAppStyle::GetBrush("BlueprintDif.CopyPropertyRight");
+		
+		if (HoveredCopyButton.CopyDirection == EPropertyCopyDirection::CopyLeftToRight)
+		{
+			if (HoveredCopyButton.SourceDetailsNode == DiffNode->ValueA)
+			{
+				ButtonColor = FStyleColors::ForegroundHover.GetSpecifiedColor();
+			}
+		}
+		break;
+	case EPropertyCopyDirection::CopyRightToLeft:
+		if (RightPropertyRect.GetSize().Y < UE_SMALL_NUMBER)
+		{
+			// space is too small to fit a button
+			return;
+		}
+		Geometry = FPaintGeometry(
+			FVector2D(RightPropertyRect.Left - ButtonScale, RightPropertyRect.GetCenter().Y - 0.5f * ButtonScale),
+			FVector2D(ButtonScale),
+			1.f
+		);
+		Brush = FAppStyle::GetBrush("BlueprintDif.CopyPropertyLeft");
+		
+		if (HoveredCopyButton.CopyDirection == EPropertyCopyDirection::CopyRightToLeft)
+		{
+			if (HoveredCopyButton.SourceDetailsNode == DiffNode->ValueB)
+			{
+				ButtonColor = FStyleColors::ForegroundHover.GetSpecifiedColor();
+			}
+		}
+		break;
+	default:
+		return;
+	}
+	
+		
+	FSlateDrawElement::MakeBox(
+		OutDrawElements,
+		LayerId, // draw in front of background but behind text, buttons, etc
+		Geometry,
+		Brush,
+		ESlateDrawEffect::None,
+		ButtonColor
+	);
+}
\ No newline at end of file
diff --git a/Engine/Source/Editor/PropertyEditor/Private/SDetailsViewBase.cpp b/Engine/Source/Editor/PropertyEditor/Private/SDetailsViewBase.cpp
--- a/Engine/Source/Editor/PropertyEditor/Private/SDetailsViewBase.cpp
+++ b/Engine/Source/Editor/PropertyEditor/Private/SDetailsViewBase.cpp
@@ -329,6 +329,90 @@ void SDetailsViewBase::HighlightProperty(const FPropertyPath& Property)
 	CurrentlyHighlightedNode = TreeNode;
 }
 
+static void ExpandPaintSpacePropertyBoundsRecursive(const TArray<TSharedRef<FDetailTreeNode>>& TreeNodes, TSharedPtr<SDetailTree> DetailTree, FSlateRect& InOutRect)
+{
+	for(const TSharedRef<FDetailTreeNode>& TreeNode : TreeNodes)
+	{
+		if (const TSharedPtr<ITableRow> Row = DetailTree->WidgetFromItem(TreeNode))
+		{
+			if (const TSharedPtr<SWidget> Widget = Row->AsWidget())
+			{
+				FSlateRect ChildRect = Widget->GetPaintSpaceGeometry().GetLayoutBoundingRect();
+				InOutRect = InOutRect.IsValid() ? InOutRect.Expand(ChildRect) : ChildRect;
+			}
+		}
+		
+		if (!TreeNode->IsLeaf() && DetailTree->IsItemExpanded(TreeNode))
+		{
+			TArray<TSharedRef<FDetailTreeNode>> Children;
+			TreeNode->GetChildren(Children);
+			ExpandPaintSpacePropertyBoundsRecursive(Children, DetailTree, InOutRect);
+		}
+	}
+}
+
+FSlateRect SDetailsViewBase::GetPaintSpacePropertyBounds(const TSharedRef<FDetailTreeNode>& InDetailTreeNode, bool bIncludeChildren) const
+{
+	const TSharedRef<FDetailTreeNode> DetailTreeNode = StaticCastSharedRef<FDetailTreeNode>(InDetailTreeNode);
+	FSlateRect Result{-1,-1,-1,-1};
+	if (const TSharedPtr<ITableRow> Row = DetailTree->WidgetFromItem(DetailTreeNode))
+	{
+		if (const TSharedPtr<SWidget> Widget = Row->AsWidget())
+		{
+			Result = Widget->GetPaintSpaceGeometry().GetLayoutBoundingRect();
+		}
+	}
+	if (bIncludeChildren && !DetailTreeNode->IsLeaf() && DetailTree->IsItemExpanded(DetailTreeNode))
+	{
+		TArray<TSharedRef<FDetailTreeNode>> Children;
+		DetailTreeNode->GetChildren(Children);
+		ExpandPaintSpacePropertyBoundsRecursive(Children, DetailTree, Result);
+	}
+	return Result;
+}
+
+static void ExpandTickSpacePropertyBoundsRecursive(const TArray<TSharedRef<FDetailTreeNode>>& TreeNodes, TSharedPtr<SDetailTree> DetailTree, FSlateRect& InOutRect)
+{
+	for(const TSharedRef<FDetailTreeNode>& TreeNode : TreeNodes)
+	{
+		if (const TSharedPtr<ITableRow> Row = DetailTree->WidgetFromItem(TreeNode))
+		{
+			if (const TSharedPtr<SWidget> Widget = Row->AsWidget())
+			{
+				FSlateRect ChildRect = Widget->GetTickSpaceGeometry().GetLayoutBoundingRect();
+				InOutRect = InOutRect.IsValid() ? InOutRect.Expand(ChildRect) : ChildRect;
+			}
+		}
+		
+		if (!TreeNode->IsLeaf() && DetailTree->IsItemExpanded(TreeNode))
+		{
+			TArray<TSharedRef<FDetailTreeNode>> Children;
+			TreeNode->GetChildren(Children);
+			ExpandTickSpacePropertyBoundsRecursive(Children, DetailTree, InOutRect);
+		}
+	}
+}
+
+FSlateRect SDetailsViewBase::GetTickSpacePropertyBounds(const TSharedRef<FDetailTreeNode>& InDetailTreeNode, bool bIncludeChildren) const
+{
+	const TSharedRef<FDetailTreeNode> DetailTreeNode = StaticCastSharedRef<FDetailTreeNode>(InDetailTreeNode);
+	FSlateRect Result{-1,-1,-1,-1};
+	if (const TSharedPtr<ITableRow> Row = DetailTree->WidgetFromItem(DetailTreeNode))
+	{
+		if (const TSharedPtr<SWidget> Widget = Row->AsWidget())
+		{
+			Result = Widget->GetTickSpaceGeometry().GetLayoutBoundingRect();
+		}
+	}
+	if (bIncludeChildren && !DetailTreeNode->IsLeaf() && DetailTree->IsItemExpanded(DetailTreeNode))
+	{
+		TArray<TSharedRef<FDetailTreeNode>> Children;
+		DetailTreeNode->GetChildren(Children);
+		ExpandTickSpacePropertyBoundsRecursive(Children, DetailTree, Result);
+	}
+	return Result;
+}
+
 bool SDetailsViewBase::IsAncestorCollapsed(const TSharedRef<IDetailTreeNode>& Node) const
 {
 	// in order for a node to be expanded, all it's parents have to be expanded
diff --git a/Engine/Source/Editor/PropertyEditor/Private/SDetailsViewBase.h b/Engine/Source/Editor/PropertyEditor/Private/SDetailsViewBase.h
--- a/Engine/Source/Editor/PropertyEditor/Private/SDetailsViewBase.h
+++ b/Engine/Source/Editor/PropertyEditor/Private/SDetailsViewBase.h
@@ -21,7 +21,6 @@
 #include "PropertyPath.h"
 #include "PropertyRowGenerator.h"
 #include "StringPrefixTree.h"
-#include "Widgets/Layout/SSplitter.h"
 #include "Widgets/SWindow.h"
 #include "Widgets/Views/STableRow.h"
 #include "Widgets/Views/STableViewBase.h"
@@ -88,6 +87,8 @@ public:
 	virtual TArray<TPair<int32, FPropertyPath>> GetPropertyRowNumbers() const override;
 	virtual int32 CountRows() const override;
 	virtual void HighlightProperty(const FPropertyPath& Property) override;
+	virtual FSlateRect GetPaintSpacePropertyBounds(const TSharedRef<FDetailTreeNode>& InDetailTreeNode, bool bIncludeChildren = true) const override;
+	virtual FSlateRect GetTickSpacePropertyBounds(const TSharedRef<FDetailTreeNode>& InDetailTreeNode, bool bIncludeChildren = true) const override;
 	virtual bool IsAncestorCollapsed(const TSharedRef<IDetailTreeNode>& Node) const override;
 	virtual void ShowAllAdvancedProperties() override;
 	virtual void SetOnDisplayedPropertiesChanged(FOnDisplayedPropertiesChanged InOnDisplayedPropertiesChangedDelegate) override;
diff --git a/Engine/Source/Editor/PropertyEditor/Public/DetailsViewArgs.h b/Engine/Source/Editor/PropertyEditor/Public/DetailsViewArgs.h
--- a/Engine/Source/Editor/PropertyEditor/Public/DetailsViewArgs.h
+++ b/Engine/Source/Editor/PropertyEditor/Public/DetailsViewArgs.h
@@ -50,6 +50,8 @@ struct FDetailsViewArgs
 	TArray<TSharedRef<IClassViewerFilter>> ClassViewerFilters;
 	/** Externally constructed Scrollbar (internally constructed if not provided) */
 	TSharedPtr<SScrollBar> ExternalScrollbar;
+	/** whether to display scrollbar on the left or the right of the panel */
+	EHorizontalAlignment ScrollbarAlignment;
 
 	/** Identifier for this details view; NAME_None if this view is anonymous */
 	FName ViewIdentifier;
@@ -110,6 +112,7 @@ public:
 
 	FDetailsViewArgs()
 		: DefaultsOnlyVisibility(EEditDefaultsOnlyNodeVisibility::Show)
+		, ScrollbarAlignment(HAlign_Right)
 		, ViewIdentifier(NAME_None)
 		, NotifyHook(nullptr)
 		, NameAreaSettings(ActorsUseNameArea)
diff --git a/Engine/Source/Editor/PropertyEditor/Public/IDetailsView.h b/Engine/Source/Editor/PropertyEditor/Public/IDetailsView.h
--- a/Engine/Source/Editor/PropertyEditor/Public/IDetailsView.h
+++ b/Engine/Source/Editor/PropertyEditor/Public/IDetailsView.h
@@ -196,6 +196,16 @@ public:
 	 */
 	virtual TArray<TPair<int32, FPropertyPath>> GetPropertyRowNumbers() const  = 0;
 
+	/**
+	 * Returns the boundary (in paint space) of a detail tree node if it's currently visible Otherwise return an invalid boundary
+	 */
+	virtual FSlateRect GetPaintSpacePropertyBounds(const TSharedRef<FDetailTreeNode>& InDetailTreeNode, bool IncludeChildren = true) const {return FSlateRect();}
+	
+	/**
+	 * Returns the boundary (in tick space) of a detail tree node if it's currently visible Otherwise return an invalid boundary
+	 */
+	virtual FSlateRect GetTickSpacePropertyBounds(const TSharedRef<FDetailTreeNode>& InDetailTreeNode, bool IncludeChildren = true) const {return FSlateRect();}
+
 	/**
 	 * Returns true if the provided detail tree node is part of a collapsed/hidden subsection of the details panel
 	 */
diff --git a/Engine/Source/Editor/PropertyEditor/Public/SDetailsSplitter.h b/Engine/Source/Editor/PropertyEditor/Public/SDetailsSplitter.h
new file mode 100644
--- /dev/null
+++ b/Engine/Source/Editor/PropertyEditor/Public/SDetailsSplitter.h
@@ -0,0 +1,77 @@
+// Copyright Epic Games, Inc. All Rights Reserved.
+#pragma once
+#include "AsyncDetailViewDiff.h"
+#include "DiffUtils.h"
+#include "Templates/SharedPointer.h"
+#include "IDetailsView.h"
+
+// Splitter that allows you to provide an FAsyncDetailViewDiff to connect like-properties between two or more details panels
+class PROPERTYEDITOR_API SDetailsSplitter : public SCompoundWidget
+{
+public:
+	class FSlot : public TSlotBase<FSlot>
+	{
+	public:
+		FSlot() = default;
+
+		SLATE_SLOT_BEGIN_ARGS(FSlot, TSlotBase<FSlot>)
+			/** When the RuleSize is set to FractionOfParent, the size of the slot is the Value percentage of its parent size. */
+			SLATE_ATTRIBUTE(float, Value)
+			SLATE_ARGUMENT(TSharedPtr<IDetailsView>, DetailsView)
+			SLATE_ATTRIBUTE(bool, IsReadonly)
+			SLATE_ATTRIBUTE(TSharedPtr<FAsyncDetailViewDiff>, DifferencesWithRightPanel)
+		SLATE_SLOT_END_ARGS()
+	};
+	static FSlot::FSlotArguments Slot();
+	
+	SLATE_BEGIN_ARGS(SDetailsSplitter)
+	{}
+		SLATE_SLOT_ARGUMENT(FSlot, Slots)
+	SLATE_END_ARGS()
+
+	void Construct(const FArguments& InArgs);
+
+	virtual int32 OnPaint(const FPaintArgs& Args, const FGeometry& AllottedGeometry, const FSlateRect& MyCullingRect,
+	                      FSlateWindowElementList& OutDrawElements, int32 LayerId, const FWidgetStyle& InWidgetStyle,
+	                      bool bParentEnabled) const override;
+
+	virtual FReply OnMouseMove(const FGeometry& MyGeometry, const FPointerEvent& MouseEvent) override;
+	virtual void OnMouseLeave(const FPointerEvent& MouseEvent) override;
+	virtual FReply OnMouseButtonDown(const FGeometry& MyGeometry, const FPointerEvent& MouseEvent) override;
+	
+	void AddSlot(const FSlot::FSlotArguments& SlotArgs, int32 Index = INDEX_NONE);
+
+private:
+	
+	struct FPanel
+	{
+		TSharedPtr<IDetailsView> DetailsView;
+		TAttribute<bool> IsReadonly;
+		TAttribute<TSharedPtr<FAsyncDetailViewDiff>> DiffRight;
+	};
+	
+	enum class EPropertyCopyDirection
+	{
+		Copy_None,
+		CopyLeftToRight,
+		CopyRightToLeft,
+	};
+	
+	struct FCopyPropertyButton
+	{
+		TWeakPtr<FDetailTreeNode> SourceDetailsNode;
+		TWeakPtr<FDetailTreeNode> DestinationDetailsNode;
+		ETreeDiffResult DiffResult = ETreeDiffResult::Invalid;
+		EPropertyCopyDirection CopyDirection = EPropertyCopyDirection::Copy_None;
+	};
+	
+	void PaintPropertyConnector(FSlateWindowElementList& OutDrawElements, int32 LayerId, const FSlateRect& LeftPropertyRect, 
+		const FSlateRect& RightPropertyRect, const FLinearColor& FillColor, const FLinearColor& OutlineColor) const;
+	
+	void PaintCopyPropertyButton(FSlateWindowElementList& OutDrawElements, int32 LayerId, const TUniquePtr<TDiffNode<TWeakPtr<FDetailTreeNode>>>& DiffNode,
+		const FSlateRect& LeftPropertyRect, const FSlateRect& RightPropertyRect, EPropertyCopyDirection CopyDirection) const;
+	
+	TSharedPtr<SSplitter> Splitter;
+	TArray<FPanel> Panels;
+	FCopyPropertyButton HoveredCopyButton;
+};
diff --git a/Engine/Source/Editor/UnrealEd/Private/DiffUtils.cpp b/Engine/Source/Editor/UnrealEd/Private/DiffUtils.cpp
--- a/Engine/Source/Editor/UnrealEd/Private/DiffUtils.cpp
+++ b/Engine/Source/Editor/UnrealEd/Private/DiffUtils.cpp
@@ -127,13 +127,21 @@ FResolvedProperty FPropertySoftPath::Resolve(const UStruct* Struct, const void*
 			}
 			else if( const FSetProperty* SetProperty = CastField<FSetProperty>(Property) )
 			{
-				checkf(false, TEXT("Set Indexing not supported yet"));
-				// TODO: @jordan.hoffmann: handle indexing for sets
+				FScriptSetHelper SetHelper(SetProperty, Property->ContainerPtrToValuePtr<UObject*>(CurrentBlock));
+				if (SetHelper.IsValidIndex(PropertyIndex))
+				{
+					NextProperty = SetProperty->ElementProp;
+					NextBlock = SetHelper.GetElementPtr(PropertyIndex);
+				}
 			}
 			else if( const FMapProperty* MapProperty = CastField<FMapProperty>(Property) )
 			{
-				checkf(false, TEXT("Map Indexing not supported yet"));
-				// TODO: @jordan.hoffmann: handle indexing for maps
+				FScriptMapHelper MapHelper(MapProperty, Property->ContainerPtrToValuePtr<UObject*>(CurrentBlock));
+				if (MapHelper.IsValidIndex(PropertyIndex))
+				{
+					NextProperty = MapProperty->ValueProp;
+					NextBlock = MapHelper.GetValuePtr(PropertyIndex);
+				}
 			}
 		}
 		
diff --git a/Engine/Source/Runtime/Slate/Private/Widgets/Layout/LinkableScrollBar.cpp b/Engine/Source/Runtime/Slate/Private/Widgets/Layout/LinkableScrollBar.cpp
--- a/Engine/Source/Runtime/Slate/Private/Widgets/Layout/LinkableScrollBar.cpp
+++ b/Engine/Source/Runtime/Slate/Private/Widgets/Layout/LinkableScrollBar.cpp
@@ -5,10 +5,22 @@
 #include "Widgets/Layout/SScrollBarTrack.h"
 #include "Framework/FInvertiblePiecewiseLinearFunction.h"
 
+
+
+enum class ELinkId : uint8
+{
+	LeftLink,
+	RightLink,
+	Disabled
+};
+
 void SLinkableScrollBar::SetState(float InOffsetFraction, float InThumbSizeFraction, bool bCallOnUserScrolled)
 {
+	const bool bLinkedRight = LinkedScrollBarRight.IsValid() && ScrollSyncRateRight.IsSet();
+	const bool bLinkedLeft = LinkedScrollBarLeft.IsValid() && ScrollSyncRateLeft.IsSet();
+	
 	// fallback to default behavior if this isn't linked to another scrollbar
-	if (!LinkedScrollBar.IsValid() || !ScrollSyncRate.IsSet())
+	if (!bLinkedRight && !bLinkedLeft)
 	{
 		SScrollBar::SetState(InOffsetFraction, InThumbSizeFraction, bCallOnUserScrolled);
 		return;
@@ -27,46 +39,64 @@ void SLinkableScrollBar::SetState(float InOffsetFraction, float InThumbSizeFract
 	// if another scrollbar is linked to this one, Scroll it aswell
 	if (bDirty)
 	{
-		const TSharedPtr<SScrollBar> OtherScrollbar = LinkedScrollBar.Pin();
-		const FInvertiblePiecewiseLinearFunction ScrollRate = FInvertiblePiecewiseLinearFunction(ScrollSyncRate.Get());
-		
-		float OtherScrollOffset = -1.f;
-		switch(LinkId)
+		auto UpdateLinkedScrollbar = [&](TSharedPtr<SScrollBar> OtherScrollbar, const FInvertiblePiecewiseLinearFunction ScrollRate, ELinkId LinkId)
 		{
-		case ELinkId::LeftLink:
-			OtherScrollOffset = ScrollRate.SolveForY(InOffsetFraction);
-			break;
-		case ELinkId::RightLink:
-			OtherScrollOffset = ScrollRate.SolveForX(InOffsetFraction);
-			break;
+			
+			float OtherScrollOffset = -1.f;
+			switch(LinkId)
+			{
+			case ELinkId::LeftLink:
+				OtherScrollOffset = ScrollRate.SolveForX(InOffsetFraction);
+				break;
+			case ELinkId::RightLink:
+				OtherScrollOffset = ScrollRate.SolveForY(InOffsetFraction);
+				break;
 		
-		case ELinkId::Disabled:;
-		default: ;
-		}
+			case ELinkId::Disabled:;
+			default: ;
+			}
 
-		if (NewDistanceFromTop > PrevDistanceFromTop) // scrolling down
-		{
-			if (FMath::IsNearlyZero(OtherScrollbar->DistanceFromBottom(), 0.01f))
+			if (NewDistanceFromTop > PrevDistanceFromTop) // scrolling down
+				{
+				if (FMath::IsNearlyZero(OtherScrollbar->DistanceFromBottom(), 0.01f))
+				{
+					return;
+				}
+				}
+		
+			if (OtherScrollOffset <= 1.f && OtherScrollOffset >= 0.f)
 			{
-				return;
+				OtherScrollbar->SetState(OtherScrollOffset, OtherScrollbar->ThumbSizeFraction(), true);
 			}
+		};
+
+		if (bLinkedLeft)
+		{
+			UpdateLinkedScrollbar(
+				LinkedScrollBarLeft.Pin(),
+				FInvertiblePiecewiseLinearFunction(ScrollSyncRateLeft.Get()),
+				ELinkId::LeftLink
+			);
 		}
-		
-		if (OtherScrollOffset <= 1.f && OtherScrollOffset >= 0.f)
+
+		if (bLinkedRight)
 		{
-			OtherScrollbar->SetState(OtherScrollOffset, OtherScrollbar->ThumbSizeFraction(), true);
+			UpdateLinkedScrollbar(
+				LinkedScrollBarRight.Pin(),
+				FInvertiblePiecewiseLinearFunction(ScrollSyncRateRight.Get()),
+				ELinkId::RightLink
+			);
 		}
+		
 	}
 }
 
 void SLinkableScrollBar::LinkScrollBars(TSharedRef<SLinkableScrollBar> Left, TSharedRef<SLinkableScrollBar> Right,
 	TAttribute<TArray<FVector2f>> ScrollSyncRate)
 {
-	Left->LinkId = ELinkId::LeftLink;
-	Left->LinkedScrollBar = Right.ToWeakPtr();
-	Left->ScrollSyncRate = ScrollSyncRate;
+	Left->LinkedScrollBarRight = Right.ToWeakPtr();
+	Left->ScrollSyncRateRight = ScrollSyncRate;
 	
-	Right->LinkId = ELinkId::RightLink;
-	Right->LinkedScrollBar = Left.ToWeakPtr();
-	Right->ScrollSyncRate = ScrollSyncRate;
+	Right->LinkedScrollBarLeft = Left.ToWeakPtr();
+	Right->ScrollSyncRateLeft = ScrollSyncRate;
 }
diff --git a/Engine/Source/Runtime/Slate/Public/Widgets/Layout/LinkableScrollBar.h b/Engine/Source/Runtime/Slate/Public/Widgets/Layout/LinkableScrollBar.h
--- a/Engine/Source/Runtime/Slate/Public/Widgets/Layout/LinkableScrollBar.h
+++ b/Engine/Source/Runtime/Slate/Public/Widgets/Layout/LinkableScrollBar.h
@@ -8,22 +8,17 @@
 class SLATE_API SLinkableScrollBar : public SScrollBar
 {
 public:
-	enum class ELinkId : uint8
-	{
-		LeftLink,
-		RightLink,
-		Disabled
-	};
 	
 	virtual void SetState(float InOffsetFraction, float InThumbSizeFraction, bool bCallOnUserScrolled = false) override;
 
 	static void LinkScrollBars(TSharedRef<SLinkableScrollBar> Left, TSharedRef<SLinkableScrollBar> Right, TAttribute<TArray<FVector2f>> ScrollSyncRate);
 
 private:
-	TWeakPtr<SLinkableScrollBar> LinkedScrollBar;
-	ELinkId LinkId = ELinkId::Disabled; // LeftLink is associated w/ x value in SyncScrollRate while RightLink is associated w/ y
+	TWeakPtr<SLinkableScrollBar> LinkedScrollBarRight;
+	TWeakPtr<SLinkableScrollBar> LinkedScrollBarLeft;
 
 	// list of scroll distance pairs that determine variable scroll rate.
 	// scroll values of Left panel will match X components while the right panel will match Y components
-	TAttribute<TArray<FVector2f>> ScrollSyncRate;
+	TAttribute<TArray<FVector2f>> ScrollSyncRateRight;
+	TAttribute<TArray<FVector2f>> ScrollSyncRateLeft;
 };
diff --git a/Engine/Source/Runtime/SlateCore/Public/Layout/SlateRect.h b/Engine/Source/Runtime/SlateCore/Public/Layout/SlateRect.h
--- a/Engine/Source/Runtime/SlateCore/Public/Layout/SlateRect.h
+++ b/Engine/Source/Runtime/SlateCore/Public/Layout/SlateRect.h
@@ -355,6 +355,17 @@ public:
 	* @return	True if initialized successfully
 	*/
 	bool InitFromString(const FString& InSourceString);
+
+	
+	friend FORCEINLINE uint32 GetTypeHash(const FSlateRect& Key)
+	{
+		uint32 Hash = 0;
+		Hash = HashCombine(Hash, GetTypeHash(Key.Left));
+		Hash = HashCombine(Hash, GetTypeHash(Key.Right));
+		Hash = HashCombine(Hash, GetTypeHash(Key.Top));
+		Hash = HashCombine(Hash, GetTypeHash(Key.Bottom));
+		return Hash;
+	}
 };
 
 /**
