[Crash Fix] FText property comparison crash in diff tool

diff --git a/Engine/Source/Editor/PropertyEditor/Private/AsyncDetailViewDiff.cpp b/Engine/Source/Editor/PropertyEditor/Private/AsyncDetailViewDiff.cpp
--- a/Engine/Source/Editor/PropertyEditor/Private/AsyncDetailViewDiff.cpp
+++ b/Engine/Source/Editor/PropertyEditor/Private/AsyncDetailViewDiff.cpp
@@ -41,7 +41,7 @@ bool TreeDiffSpecification::AreValuesEqual<TWeakPtr<FDetailTreeNode>>(const TWea
 			return false;
 		}
 	    const void* DataA = ResolvedA.Property->ContainerPtrToValuePtr<void*>(ResolvedA.Object);
-	    const void* DataB = ResolvedA.Property->ContainerPtrToValuePtr<void*>(ResolvedB.Object);
+	    const void* DataB = ResolvedB.Property->ContainerPtrToValuePtr<void*>(ResolvedB.Object);
 		return ResolvedA.Property->Identical(DataA, DataB, PPF_DeepComparison);
 	}
 	if (!ResolvedA.Property && !ResolvedB.Property)
