[CrashFix] Crash when reviewing CL 25127136 when comparing property map keys
#preflag 6441dbed0206a6e20fe227b0

diff --git a/Engine/Source/Editor/PropertyEditor/Private/AsyncDetailViewDiff.cpp b/Engine/Source/Editor/PropertyEditor/Private/AsyncDetailViewDiff.cpp
--- a/Engine/Source/Editor/PropertyEditor/Private/AsyncDetailViewDiff.cpp
+++ b/Engine/Source/Editor/PropertyEditor/Private/AsyncDetailViewDiff.cpp
@@ -5,17 +5,23 @@
 #include "IDetailsViewPrivate.h"
 #include "DiffUtils.h"
 
-static FResolvedProperty GetResolvedProperty(const TSharedRef<FDetailTreeNode>& TreeNode)
+static const UObject* GetObject(const TSharedPtr<FDetailTreeNode>& TreeNode)
 {
-	const FPropertyPath PropertyPath = TreeNode->GetPropertyPath();
-	if (PropertyPath.IsValid())
+	if (const IDetailsViewPrivate* DetailsView = TreeNode->GetDetailsView())
 	{
-		if (const IDetailsViewPrivate* DetailsView = TreeNode->GetDetailsView())
+		return DetailsView->GetSelectedObjects()[0].Get();
+	}
+	return nullptr;
+}
+
+static FResolvedProperty GetResolvedProperty(const TSharedPtr<FPropertyNode>& PropertyNode, const UObject* Object)
+{
+	if (PropertyNode)
+	{
+		const TSharedRef<FPropertyPath> PropertyPath = FPropertyNode::CreatePropertyPath(PropertyNode.ToSharedRef());
+		if (PropertyPath->IsValid())
 		{
-			if (const UObject* Object = DetailsView->GetSelectedObjects()[0].Get())
-			{
-				return FPropertySoftPath(PropertyPath).Resolve(Object);
-			}
+			return FPropertySoftPath(*PropertyPath).Resolve(Object);
 		}
 	}
 	return FResolvedProperty();
@@ -31,8 +37,8 @@ bool TreeDiffSpecification::AreValuesEqual<TWeakPtr<FDetailTreeNode>>(const TWea
 		return PinnedTreeNodeA == PinnedTreeNodeB;
 	}
 
-	const FResolvedProperty ResolvedA = GetResolvedProperty(PinnedTreeNodeA.ToSharedRef());
-    const FResolvedProperty ResolvedB = GetResolvedProperty(PinnedTreeNodeB.ToSharedRef());
+	const FResolvedProperty ResolvedA = GetResolvedProperty(PinnedTreeNodeA->GetPropertyNode(), GetObject(PinnedTreeNodeA));
+    const FResolvedProperty ResolvedB = GetResolvedProperty(PinnedTreeNodeB->GetPropertyNode(), GetObject(PinnedTreeNodeB));
 	if (ResolvedA.Property && ResolvedB.Property)
 	{
 		// property nodes
@@ -54,41 +60,58 @@ bool TreeDiffSpecification::AreValuesEqual<TWeakPtr<FDetailTreeNode>>(const TWea
 
 static bool MapKeysMatch(const TSharedPtr<FDetailTreeNode>& TreeNodeA, const TSharedPtr<FDetailTreeNode>& TreeNodeB, int32 KeyIndexA, int32 KeyIndexB)
 {
-	const FMapProperty* MapPropertyA = CastField<FMapProperty>(TreeNodeA->GetPropertyNode()->GetProperty());
-	const FMapProperty* MapPropertyB = CastField<FMapProperty>(TreeNodeB->GetPropertyNode()->GetProperty());
-	if (MapPropertyA && MapPropertyB)
+	const TSharedPtr<FPropertyNode> MapPropertyNodeA = TreeNodeA->GetPropertyNode()->GetParentNode()->AsShared();
+	const TSharedPtr<FPropertyNode> MapPropertyNodeB = TreeNodeB->GetPropertyNode()->GetParentNode()->AsShared();
+	if (!MapPropertyNodeA || !MapPropertyNodeB)
 	{
-		const FResolvedProperty ResolvedMapA = GetResolvedProperty(TreeNodeA.ToSharedRef());
-		const FResolvedProperty ResolvedMapB = GetResolvedProperty(TreeNodeB.ToSharedRef());
-		FScriptMapHelper MapHelperA(MapPropertyA, MapPropertyA->ContainerPtrToValuePtr<UObject*>(ResolvedMapA.Object));
-		FScriptMapHelper MapHelperB(MapPropertyB, MapPropertyB->ContainerPtrToValuePtr<UObject*>(ResolvedMapB.Object));
+		return false;
+	}
+	const FMapProperty* MapPropertyA = CastField<FMapProperty>(MapPropertyNodeA->GetProperty());
+	const FMapProperty* MapPropertyB = CastField<FMapProperty>(MapPropertyNodeB->GetProperty());
+	if (!MapPropertyA || !MapPropertyB)
+	{
+		return false;
+	}
+	
+	const FResolvedProperty ResolvedMapA = GetResolvedProperty(MapPropertyNodeA, GetObject(TreeNodeA));
+	const FResolvedProperty ResolvedMapB = GetResolvedProperty(MapPropertyNodeB, GetObject(TreeNodeB));
+	FScriptMapHelper MapHelperA(MapPropertyA, MapPropertyA->ContainerPtrToValuePtr<UObject*>(ResolvedMapA.Object));
+	FScriptMapHelper MapHelperB(MapPropertyB, MapPropertyB->ContainerPtrToValuePtr<UObject*>(ResolvedMapB.Object));
 
-		const void* KeyA = MapHelperA.GetKeyPtr(KeyIndexA);
-		const void* KeyB = MapHelperB.GetKeyPtr(KeyIndexB);
-		if (MapPropertyA->KeyProp->SameType(MapPropertyB->KeyProp))
-		{
-			return MapPropertyA->KeyProp->Identical(KeyA, KeyB, PPF_DeepComparison);
-		}
+	const void* KeyA = MapHelperA.GetKeyPtr(KeyIndexA);
+	const void* KeyB = MapHelperB.GetKeyPtr(KeyIndexB);
+	if (MapPropertyA->KeyProp->SameType(MapPropertyB->KeyProp))
+	{
+		return MapPropertyA->KeyProp->Identical(KeyA, KeyB, PPF_DeepComparison);
 	}
 	return false;
 }
 
 static bool SetKeysMatch(const TSharedPtr<FDetailTreeNode>& TreeNodeA, const TSharedPtr<FDetailTreeNode>& TreeNodeB, int32 KeyIndexA, int32 KeyIndexB)
 {
-	const FSetProperty* SetPropertyA = CastField<FSetProperty>(TreeNodeA->GetPropertyNode()->GetProperty());
-	const FSetProperty* SetPropertyB = CastField<FSetProperty>(TreeNodeB->GetPropertyNode()->GetProperty());
-	if (SetPropertyA && SetPropertyB)
+	const TSharedPtr<FPropertyNode> SetPropertyNodeA = TreeNodeA->GetPropertyNode()->GetParentNode()->AsShared();
+	const TSharedPtr<FPropertyNode> SetPropertyNodeB = TreeNodeB->GetPropertyNode()->GetParentNode()->AsShared();
+	if (!SetPropertyNodeA || !SetPropertyNodeB)
 	{
-		const FResolvedProperty ResolvedSetA = GetResolvedProperty(TreeNodeA.ToSharedRef());
-		const FResolvedProperty ResolvedSetB = GetResolvedProperty(TreeNodeB.ToSharedRef());
-		FScriptSetHelper SetHelperA(SetPropertyA, SetPropertyA->ContainerPtrToValuePtr<UObject*>(ResolvedSetA.Object));
-		FScriptSetHelper SetHelperB(SetPropertyB, SetPropertyB->ContainerPtrToValuePtr<UObject*>(ResolvedSetB.Object));
-		const void* KeyA = SetHelperA.GetElementPtr(KeyIndexA);
-		const void* KeyB = SetHelperB.GetElementPtr(KeyIndexB);
-		if (SetPropertyA->ElementProp->SameType(SetPropertyB->ElementProp))
-		{
-			return SetPropertyA->ElementProp->Identical(KeyA, KeyB, PPF_DeepComparison);
-		}
+		return false;
+	}
+	const FSetProperty* SetPropertyA = CastField<FSetProperty>(SetPropertyNodeA->GetProperty());
+	const FSetProperty* SetPropertyB = CastField<FSetProperty>(SetPropertyNodeB->GetProperty());
+	if (!SetPropertyA || !SetPropertyB)
+	{
+		return false;
+	}
+	
+	const FResolvedProperty ResolvedSetA = GetResolvedProperty(SetPropertyNodeA, GetObject(TreeNodeA));
+	const FResolvedProperty ResolvedSetB = GetResolvedProperty(SetPropertyNodeB, GetObject(TreeNodeB));
+	FScriptSetHelper SetHelperA(SetPropertyA, SetPropertyA->ContainerPtrToValuePtr<UObject*>(ResolvedSetA.Object));
+	FScriptSetHelper SetHelperB(SetPropertyB, SetPropertyB->ContainerPtrToValuePtr<UObject*>(ResolvedSetB.Object));
+	
+	const void* KeyA = SetHelperA.GetElementPtr(KeyIndexA);
+	const void* KeyB = SetHelperB.GetElementPtr(KeyIndexB);
+	if (SetPropertyA->ElementProp->SameType(SetPropertyB->ElementProp))
+	{
+		return SetPropertyA->ElementProp->Identical(KeyA, KeyB, PPF_DeepComparison);
 	}
 	return false;
 }
@@ -123,8 +146,8 @@ bool TreeDiffSpecification::AreMatching<TWeakPtr<FDetailTreeNode>>(const TWeakPt
 			if (CastField<FMapProperty>(ParentPropertyA) || CastField<FMapProperty>(ParentPropertyB))
 			{
 				return MapKeysMatch(
-					PinnedTreeNodeA->GetParentNode().Pin(),
-					PinnedTreeNodeB->GetParentNode().Pin(),
+					PinnedTreeNodeA,
+					PinnedTreeNodeB,
 					ArrayIndexA,
 					ArrayIndexB
 				);
@@ -132,8 +155,8 @@ bool TreeDiffSpecification::AreMatching<TWeakPtr<FDetailTreeNode>>(const TWeakPt
 			if (ParentPropertyA->IsA<FSetProperty>() || ParentPropertyB->IsA<FSetProperty>())
 			{
 				return SetKeysMatch(
-					PinnedTreeNodeA->GetParentNode().Pin(),
-					PinnedTreeNodeB->GetParentNode().Pin(),
+					PinnedTreeNodeA,
+					PinnedTreeNodeB,
 					ArrayIndexA,
 					ArrayIndexB
 				);
