Improved standalone structure property node handling
- Added IStructureProvider which allows to edit multiple instances of structs at a time
- Improved default value handling in standalone structure nodes
- Fixed bug where expanded nodes did not retain their expansion state when structure nodes were rebuild (e.g. when editing a component)
- Removed caching from Instanced struct customization

diff --git a/Engine/Plugins/Experimental/StructUtils/Source/StructUtilsEditor/Private/Customizations/InstancedStructDetails.cpp b/Engine/Plugins/Experimental/StructUtils/Source/StructUtilsEditor/Private/Customizations/InstancedStructDetails.cpp
--- a/Engine/Plugins/Experimental/StructUtils/Source/StructUtilsEditor/Private/Customizations/InstancedStructDetails.cpp
+++ b/Engine/Plugins/Experimental/StructUtils/Source/StructUtilsEditor/Private/Customizations/InstancedStructDetails.cpp
@@ -16,9 +16,139 @@
 #include "InstancedStruct.h"
 #include "IPropertyUtilities.h"
 #include "Widgets/Layout/SBox.h"
+#include "IStructureDataProvider.h"
+#include "GameFramework/Actor.h"
+#include "Misc/ConfigCacheIni.h"
 
 #define LOCTEXT_NAMESPACE "StructUtilsEditor"
 
+////////////////////////////////////
+
+class FInstancedStructProvider : public IStructureDataProvider
+{
+public:
+	FInstancedStructProvider() = default;
+	
+	explicit FInstancedStructProvider(const TSharedPtr<IPropertyHandle>& InStructProperty)
+		: StructProperty(InStructProperty)
+	{
+	}
+	
+	virtual ~FInstancedStructProvider() override {}
+
+	virtual bool IsValid() const override
+	{
+		bool bHasValidData = false;
+		EnumerateInstances([&bHasValidData](const UScriptStruct* ScriptStruct, uint8* Memory, UPackage* Package)
+		{
+			if (ScriptStruct && Memory)
+			{
+				bHasValidData = true;
+				return false; // Stop
+			}
+			return true; // Continue
+		});
+
+		return bHasValidData;
+	}
+	
+	virtual const UStruct* GetBaseStructure() const override
+	{
+		const UScriptStruct* CommonStruct = nullptr;
+		EnumerateInstances([&CommonStruct](const UScriptStruct* ScriptStruct, uint8* Memory, UPackage* Package)
+		{
+			if (!CommonStruct && ScriptStruct)
+			{
+				CommonStruct = ScriptStruct;
+				return false; // Stop
+			}
+			return true; // Continue
+		});
+
+		return CommonStruct;
+	}
+
+	virtual void GetInstances(TArray<TSharedPtr<FStructOnScope>>& OutInstances) const override
+	{
+		// The returned instances need to be compatible with base structure.
+		// This function returns empty instances in case they are not compatible, with the idea that we have as many instances as we have outer objects.
+		const UScriptStruct* CommonStruct = nullptr;
+		EnumerateInstances([&OutInstances, &CommonStruct](const UScriptStruct* ScriptStruct, uint8* Memory, UPackage* Package)
+		{
+			TSharedPtr<FStructOnScope> Result;
+			
+			// Only return structs that are same type.
+			if (CommonStruct == nullptr || ScriptStruct == CommonStruct)
+			{
+				Result = MakeShared<FStructOnScope>(ScriptStruct, Memory);
+				CommonStruct = ScriptStruct;
+			}
+
+			Result->SetPackage(Package);
+
+			OutInstances.Add(Result);
+
+			return true; // Continue
+		});
+	}
+
+	virtual bool IsPropertyIndirection() const override
+	{
+		return true;
+	}
+
+	virtual uint8* GetValueBaseAddress(uint8* ParentValueAddress, const UStruct* ExpectedType) const override
+	{
+		if (!ParentValueAddress)
+		{
+			return nullptr;
+		}
+
+		FInstancedStruct& InstancedStruct = *reinterpret_cast<FInstancedStruct*>(ParentValueAddress);
+		if (InstancedStruct.GetScriptStruct() != ExpectedType)
+		{
+			return nullptr;
+		}
+
+		return InstancedStruct.GetMutableMemory();
+	}
+	
+protected:
+
+	void EnumerateInstances(TFunctionRef<bool(const UScriptStruct* ScriptStruct, uint8* Memory, UPackage* Package)> InFunc) const
+	{
+		if (!StructProperty.IsValid())
+		{
+			return;
+		}
+		
+		TArray<UPackage*> Packages;
+		StructProperty->GetOuterPackages(Packages);
+
+		StructProperty->EnumerateRawData([&InFunc, &Packages](void* RawData, const int32 DataIndex, const int32 /*NumDatas*/)
+		{
+			const UScriptStruct* ScriptStruct = nullptr;
+			uint8* Memory = nullptr;
+			UPackage* Package = nullptr;
+			if (FInstancedStruct* InstancedStruct = static_cast<FInstancedStruct*>(RawData))
+			{
+				ScriptStruct = InstancedStruct->GetScriptStruct();
+				Memory = InstancedStruct->GetMutableMemory();
+				if (ensureMsgf(Packages.IsValidIndex(DataIndex), TEXT("Expecting packges and raw data to match.")))
+				{
+					Package = Packages[DataIndex];
+				}
+			}
+
+			return InFunc(ScriptStruct, Memory, Package);
+		});
+	}
+	
+	TSharedPtr<IPropertyHandle> StructProperty;
+};
+
+////////////////////////////////////
+
 class FInstancedStructFilter : public IStructViewerFilter
 {
 public:
@@ -59,70 +189,46 @@ public:
 	}
 };
 
-namespace UE::StructUtils::Private
-{
+////////////////////////////////////
 
-const UScriptStruct* GetCommonScriptStruct(TSharedPtr<IPropertyHandle> StructProperty)
+namespace UE::StructUtils::Private
 {
-	const UScriptStruct* CommonStructType = nullptr;
-
-	StructProperty->EnumerateConstRawData([&CommonStructType](const void* RawData, const int32 /*DataIndex*/, const int32 /*NumDatas*/)
-	{
-		if (RawData)
-		{
-			const FInstancedStruct* InstancedStruct = static_cast<const FInstancedStruct*>(RawData);
-
-			const UScriptStruct* StructTypePtr = InstancedStruct->GetScriptStruct();
-			if (CommonStructType && CommonStructType != StructTypePtr)
-			{
-				// Multiple struct types on the sources - show nothing set
-				CommonStructType = nullptr;
-				return false;
-			}
-			CommonStructType = StructTypePtr;
-		}
-
-		return true;
-	});
-
-	return CommonStructType;
-}
 
-void SetInstancedStructProperty(TSharedPtr<IPropertyHandle> StructProperty, const FInstancedStruct& InstancedStructToSet, const bool bAllowStructMismatch)
+FPropertyAccess::Result GetCommonScriptStruct(TSharedPtr<IPropertyHandle> StructProperty, const UScriptStruct*& OutCommonStruct)
 {
-	// Note: We use the ExportText/SetPerObjectValues flow here (rather then CopyScriptStruct) 
-	// so that PropagatePropertyChange is called for the underlying FInstancedStruct property
-	TArray<FString> NewInstancedStructValues;
-	StructProperty->EnumerateRawData([bAllowStructMismatch, &InstancedStructToSet, &NewInstancedStructValues](const void* RawData, const int32 /*DataIndex*/, const int32 /*NumDatas*/)
+	bool bHasResult = false;
+	bool bHasMultipleValues = false;
+	
+	StructProperty->EnumerateConstRawData([&OutCommonStruct, &bHasResult, &bHasMultipleValues](const void* RawData, const int32 /*DataIndex*/, const int32 /*NumDatas*/)
 	{
-		if (RawData)
+		if (const FInstancedStruct* InstancedStruct = static_cast<const FInstancedStruct*>(RawData))
 		{
-			const FInstancedStruct* InstancedStruct = static_cast<const FInstancedStruct*>(RawData);
+			const UScriptStruct* Struct = InstancedStruct->GetScriptStruct();
 
-			// Only copy the data if this source is still using the expected struct type or we allow a mismatch
-			const UScriptStruct* StructTypePtr = InstancedStruct->GetScriptStruct();
-			if (bAllowStructMismatch || StructTypePtr == InstancedStructToSet.GetScriptStruct())
+			if (!bHasResult)
 			{
-				FString& NewInstancedStructValue = NewInstancedStructValues.AddDefaulted_GetRef();
-				FInstancedStruct::StaticStruct()->ExportText(NewInstancedStructValue, &InstancedStructToSet, &InstancedStructToSet, nullptr, PPF_None, nullptr);
+				OutCommonStruct = Struct;
 			}
-			else
+			else if (OutCommonStruct != Struct)
 			{
-				FString& NewInstancedStructValue = NewInstancedStructValues.AddDefaulted_GetRef();
-				FInstancedStruct::StaticStruct()->ExportText(NewInstancedStructValue, InstancedStruct, InstancedStruct, nullptr, PPF_None, nullptr);
+				bHasMultipleValues = true;
 			}
+
+			bHasResult = true;
 		}
-		else
-		{
-			NewInstancedStructValues.AddDefaulted();
-		}
+
 		return true;
 	});
 
-	StructProperty->SetPerObjectValues(NewInstancedStructValues);
+	if (bHasMultipleValues)
+	{
+		return FPropertyAccess::MultipleValues;
+	}
+	
+	return bHasResult ? FPropertyAccess::Success : FPropertyAccess::Fail;
 }
 
-}
+} // UE::StructUtils::Private
 
 ////////////////////////////////////
 
@@ -137,119 +243,42 @@ FInstancedStructDataDetails::FInstancedStructDataDetails(TSharedPtr<IPropertyHan
 	StructProperty = InStructProperty;
 }
 
-void FInstancedStructDataDetails::OnStructValuePreChange()
+void FInstancedStructDataDetails::SetOnRebuildChildren(FSimpleDelegate InOnRegenerateChildren)
 {
-	PreChangeOuterObjectNames.Reset();
-
-	// Forward the change event to the real struct handle
-	if (StructProperty && StructProperty->IsValidHandle())
-	{
-		StructProperty->NotifyPreChange();
-
-		// Store the outer objects for the Actor Component special case (see FInstancedStructDataDetails::OnStructValuePostChange).
-		TArray<UObject*> OuterObjects;
-		StructProperty->GetOuterObjects(OuterObjects);
-		
-		for (UObject* Outer : OuterObjects)
-		{
-			if (Outer != nullptr)
-			{
-				PreChangeOuterObjectNames.Add(Outer->GetPathName());
-			}
-		}
-	}
+	OnRegenerateChildren = InOnRegenerateChildren;
 }
 
-void FInstancedStructDataDetails::OnStructValuePostChange()
+TArray<TWeakObjectPtr<const UStruct>> FInstancedStructDataDetails::GetInstanceTypes() const
 {
-	// Forward the change event to the real struct handle
-	if (StructProperty && StructProperty->IsValidHandle())
+	TArray<TWeakObjectPtr<const UStruct>> Result;
+	
+	StructProperty->EnumerateConstRawData([&Result](const void* RawData, const int32 /*DataIndex*/, const int32 /*NumDatas*/)
 	{
-		TGuardValue<bool> HandlingStructValuePostChangeGuard(bIsHandlingStructValuePostChange, true);
-
-		// When an InstancedStruct is on an Actor Component, the details customization gets rebuild between
-		// the OnStructValuePreChange() and OnStructValuePostChange() calls due to AActor::RerunConstructionScripts().
-		// When the new details panel is build, it will clear the outer objects of the StructProperty property handle, and then setting the value will fail.
-		// To overcome that, we restore the outer objects based on the objects stored in OnStructValuePreChange().
-		if (StructProperty->GetNumOuterObjects() == 0 && PreChangeOuterObjectNames.Num() > 0)
+		TWeakObjectPtr<const UStruct>& Type = Result.AddDefaulted_GetRef();
+		if (const FInstancedStruct* InstancedStruct = static_cast<const FInstancedStruct*>(RawData))
 		{
-			TArray<UObject*> OuterObjects;
-			for (const FString& ObjectPathName : PreChangeOuterObjectNames)
-			{
-				UObject* OuterObject = FindObject<UObject>(nullptr, *ObjectPathName);
-				if (OuterObject)
-				{
-					OuterObjects.Add(OuterObject);
-				}
-			}
-			
-			StructProperty->ReplaceOuterObjects(OuterObjects);
+			Result.Add(InstancedStruct->GetScriptStruct());
 		}
-		
-		// Copy the modified struct data back to the source instances
+		else
 		{
-			FInstancedStruct TmpInstancedStruct;
-			if (StructInstanceData)
-			{
-				TmpInstancedStruct.InitializeAs(Cast<UScriptStruct>(StructInstanceData->GetStruct()), StructInstanceData->GetStructMemory());
-			}
-			UE::StructUtils::Private::SetInstancedStructProperty(StructProperty, TmpInstancedStruct, /*bAllowStructMismatch*/false);
+			Result.Add(nullptr);
 		}
+		return true;
+	});
 
-		StructProperty->NotifyPostChange(EPropertyChangeType::ValueSet);
-		StructProperty->NotifyFinishedChangingProperties();
-	}
+	return Result;
 }
 
 void FInstancedStructDataDetails::OnStructHandlePostChange()
 {
-	if (!bIsHandlingStructValuePostChange)
-	{
-		// External change; force a sync next Tick
-		LastSyncEditableInstanceFromSourceSeconds = 0.0;
-	}
-}
-
-void FInstancedStructDataDetails::SyncEditableInstanceFromSource(bool* OutStructMismatch)
-{
-	if (OutStructMismatch)
+	if (StructProvider.IsValid())
 	{
-		*OutStructMismatch = false;
-	}
-
-	if (StructProperty && StructProperty->IsValidHandle())
-	{
-		const UScriptStruct* ExpectedStructType = StructInstanceData ? Cast<UScriptStruct>(StructInstanceData->GetStruct()) : nullptr;
-		StructProperty->EnumerateConstRawData([this, ExpectedStructType, OutStructMismatch](const void* RawData, const int32 /*DataIndex*/, const int32 NumDatas)
+		TArray<TWeakObjectPtr<const UStruct>> InstanceTypes = GetInstanceTypes();
+		if (InstanceTypes != CachedInstanceTypes)
 		{
-			if (RawData && NumDatas == 1)
-			{
-				const FInstancedStruct* InstancedStruct = static_cast<const FInstancedStruct*>(RawData);
-
-				// Only copy the data if this source is still using the expected struct type
-				const UScriptStruct* StructTypePtr = InstancedStruct->GetScriptStruct();
-				if (StructTypePtr == ExpectedStructType)
-				{
-					if (StructTypePtr)
-					{
-						StructTypePtr->CopyScriptStruct(StructInstanceData->GetStructMemory(), InstancedStruct->GetMemory());
-					}
-				}
-				else if (OutStructMismatch)
-				{
-					*OutStructMismatch = true;
-				}
-			}
-			return false;
-		});
+			OnRegenerateChildren.ExecuteIfBound();
+		}
 	}
-
-	LastSyncEditableInstanceFromSourceSeconds = FPlatformTime::Seconds();
-}
-
-void FInstancedStructDataDetails::SetOnRebuildChildren(FSimpleDelegate InOnRegenerateChildren)
-{
-	OnRegenerateChildren = InOnRegenerateChildren;
 }
 
 void FInstancedStructDataDetails::GenerateHeaderRowContent(FDetailWidgetRow& NodeRow)
@@ -259,59 +288,28 @@ void FInstancedStructDataDetails::GenerateHeaderRowContent(FDetailWidgetRow& Nod
 
 void FInstancedStructDataDetails::GenerateChildContent(IDetailChildrenBuilder& ChildBuilder)
 {
-	// Create a struct instance to edit, for the common struct type of the sources being edited
-	StructInstanceData.Reset();
-	if (const UScriptStruct* CommonStructType = UE::StructUtils::Private::GetCommonScriptStruct(StructProperty))
-	{
-		StructInstanceData = MakeShared<FStructOnScope>(CommonStructType);
-
-		// Make sure the struct also has a valid package set, so that properties that rely on this (like FText) work correctly
-		{
-			TArray<UPackage*> OuterPackages;
-			StructProperty->GetOuterPackages(OuterPackages);
-			if (OuterPackages.Num() > 0)
-			{
-				StructInstanceData->SetPackage(OuterPackages[0]);
-			}
-		}
-
-		bool bStructMismatch = false;
-		SyncEditableInstanceFromSource(&bStructMismatch);
-	}
-
 	// Add the rows for the struct
-	if (StructInstanceData)
+	TSharedRef<FInstancedStructProvider> NewStructProvider = MakeShared<FInstancedStructProvider>(StructProperty);
+	
+	TArray<TSharedPtr<IPropertyHandle>> ChildProperties = StructProperty->AddChildStructure(NewStructProvider);
+	for (TSharedPtr<IPropertyHandle> ChildHandle : ChildProperties)
 	{
-		FSimpleDelegate OnStructValuePreChangeDelegate = FSimpleDelegate::CreateSP(this, &FInstancedStructDataDetails::OnStructValuePreChange);
-		FSimpleDelegate OnStructValuePostChangeDelegate = FSimpleDelegate::CreateSP(this, &FInstancedStructDataDetails::OnStructValuePostChange);
+		IDetailPropertyRow& Row = ChildBuilder.AddProperty(ChildHandle.ToSharedRef());
+		OnChildRowAdded(Row);
+	}
 
-		TArray<TSharedPtr<IPropertyHandle>> ChildProperties = StructProperty->AddChildStructure(StructInstanceData.ToSharedRef());
-		for (TSharedPtr<IPropertyHandle> ChildHandle : ChildProperties)
-		{
-			ChildHandle->SetOnPropertyValuePreChange(OnStructValuePreChangeDelegate);
-			ChildHandle->SetOnChildPropertyValuePreChange(OnStructValuePreChangeDelegate);
-			ChildHandle->SetOnPropertyValueChanged(OnStructValuePostChangeDelegate);
-			ChildHandle->SetOnChildPropertyValueChanged(OnStructValuePostChangeDelegate);
+	StructProvider = NewStructProvider;
 
-			IDetailPropertyRow& Row = ChildBuilder.AddProperty(ChildHandle.ToSharedRef());
-			OnChildRowAdded(Row);
-		}
-	}
+	CachedInstanceTypes = GetInstanceTypes();
 }
 
 void FInstancedStructDataDetails::Tick(float DeltaTime)
 {
-	if (LastSyncEditableInstanceFromSourceSeconds + 0.1 < FPlatformTime::Seconds())
+	// If the instance types change (e.g. due to selecting new struct type), we'll need to update the layout.
+	TArray<TWeakObjectPtr<const UStruct>> InstanceTypes = GetInstanceTypes();
+	if (InstanceTypes != CachedInstanceTypes)
 	{
-		bool bStructMismatch = false;
-		SyncEditableInstanceFromSource(&bStructMismatch);
-
-		if (bStructMismatch)
-		{
-			// If the editable struct no longer has the same struct type as the underlying source, 
-			// then we need to refresh to update the child property rows for the new type
-			OnRegenerateChildren.ExecuteIfBound();
-		}
+		OnRegenerateChildren.ExecuteIfBound();
 	}
 }
 
@@ -427,17 +425,34 @@ void FInstancedStructDetails::CustomizeChildren(TSharedRef<class IPropertyHandle
 
 FText FInstancedStructDetails::GetDisplayValueString() const
 {
-	const UScriptStruct* ScriptStruct = UE::StructUtils::Private::GetCommonScriptStruct(StructProperty);
-	if (ScriptStruct)
+	const UScriptStruct* CommonStruct = nullptr;
+	const FPropertyAccess::Result Result = UE::StructUtils::Private::GetCommonScriptStruct(StructProperty, CommonStruct);
+	
+	if (Result == FPropertyAccess::Success)
+	{
+		if (CommonStruct)
+		{
+			return CommonStruct->GetDisplayNameText();
+		}
+		return LOCTEXT("NullScriptStruct", "None");
+	}
+	if (Result == FPropertyAccess::MultipleValues)
 	{
-		return ScriptStruct->GetDisplayNameText();
+		return LOCTEXT("MultipleValues", "Multiple Values");
 	}
-	return LOCTEXT("NullScriptStruct", "None");
+	
+	return FText::GetEmpty();
 }
 
 const FSlateBrush* FInstancedStructDetails::GetDisplayValueIcon() const
 {
-	return FSlateIconFinder::FindIconBrushForClass(UScriptStruct::StaticClass());
+	const UScriptStruct* CommonStruct = nullptr;
+	if (UE::StructUtils::Private::GetCommonScriptStruct(StructProperty, CommonStruct) == FPropertyAccess::Success)
+	{
+		return FSlateIconFinder::FindIconBrushForClass(UScriptStruct::StaticClass());
+	}
+	
+	return nullptr;
 }
 
 TSharedRef<SWidget> FInstancedStructDetails::GenerateStructPicker()
@@ -486,12 +501,14 @@ void FInstancedStructDetails::OnStructPicked(const UScriptStruct* InStruct)
 
 		StructProperty->NotifyPreChange();
 
-		// Copy the modified struct data back to the source instances
+		StructProperty->EnumerateRawData([InStruct](void* RawData, const int32 /*DataIndex*/, const int32 /*NumDatas*/)
 		{
-			FInstancedStruct TmpInstancedStruct;
-			TmpInstancedStruct.InitializeAs(InStruct);
-			UE::StructUtils::Private::SetInstancedStructProperty(StructProperty, TmpInstancedStruct, /*bAllowStructMismatch*/true);
-		}
+			if (FInstancedStruct* InstancedStruct = static_cast<FInstancedStruct*>(RawData))
+			{
+				InstancedStruct->InitializeAs(InStruct);
+			}
+			return true;
+		});
 
 		StructProperty->NotifyPostChange(EPropertyChangeType::ValueSet);
 		StructProperty->NotifyFinishedChangingProperties();
diff --git a/Engine/Plugins/Experimental/StructUtils/Source/StructUtilsEditor/Public/InstancedStructDetails.h b/Engine/Plugins/Experimental/StructUtils/Source/StructUtilsEditor/Public/InstancedStructDetails.h
--- a/Engine/Plugins/Experimental/StructUtils/Source/StructUtilsEditor/Public/InstancedStructDetails.h
+++ b/Engine/Plugins/Experimental/StructUtils/Source/StructUtilsEditor/Public/InstancedStructDetails.h
@@ -11,6 +11,7 @@ class IDetailPropertyRow;
 class FStructOnScope;
 class SWidget;
 struct FInstancedStruct;
+class FInstancedStructProvider;
 
 /**
  * Type customization for FInstancedStruct.
@@ -79,22 +80,19 @@ public:
 	virtual void OnChildRowAdded(IDetailPropertyRow& ChildRow) {}
 
 private:
-	/** Pre/Post change notifications for struct value changes */
-	void OnStructValuePreChange();
-	void OnStructValuePostChange();
 	void OnStructHandlePostChange();
 
-	/** Sync the current state of the editable struct instance from the source instance(s) */
-	void SyncEditableInstanceFromSource(bool* OutStructMismatch = nullptr);
+	/** Returns type of the instanced struct for each instance/object being edited. */
+	TArray<TWeakObjectPtr<const UStruct>> GetInstanceTypes() const;
 
-	/** Outer objects at the time of PreChangeProperties */
-	TArray<FString> PreChangeOuterObjectNames;
+	/** Cached instance types, used to invalidate the layout when types change. */
+	TArray<TWeakObjectPtr<const UStruct>> CachedInstanceTypes;
 	
 	/** Handle to the struct property being edited */
 	TSharedPtr<IPropertyHandle> StructProperty;
 
-	/** Struct instance that is being edited; this is a copy of the source struct data to avoid lifetime issues when the underlying source is updated/deleted */
-	TSharedPtr<FStructOnScope> StructInstanceData;
+	/** Struct provider for the structs. */
+	TSharedPtr<FInstancedStructProvider> StructProvider;
 
 	/** Delegate that can be used to refresh the child rows of the current struct (eg, when changing struct type) */
 	FSimpleDelegate OnRegenerateChildren;
diff --git a/Engine/Source/Editor/PropertyEditor/Private/DetailLayoutBuilderImpl.cpp b/Engine/Source/Editor/PropertyEditor/Private/DetailLayoutBuilderImpl.cpp
--- a/Engine/Source/Editor/PropertyEditor/Private/DetailLayoutBuilderImpl.cpp
+++ b/Engine/Source/Editor/PropertyEditor/Private/DetailLayoutBuilderImpl.cpp
@@ -840,7 +840,7 @@ void FDetailLayoutBuilderImpl::GetStructsBeingCustomized( TArray< TSharedPtr<FSt
 				FStructurePropertyNode* StructureNode = ParentComplexProperty ? ParentComplexProperty->AsStructureNode() : nullptr;
 				if(StructureNode)
 				{
-					OutStructs.Add(StructureNode->GetStructData());
+					StructureNode->GetAllStructureData(OutStructs);
 				}
 			}
 		}
@@ -848,7 +848,7 @@ void FDetailLayoutBuilderImpl::GetStructsBeingCustomized( TArray< TSharedPtr<FSt
 
 	if(RootStructNode)
 	{
-		OutStructs.Add(RootStructNode->GetStructData());
+		RootStructNode->GetAllStructureData(OutStructs);
 	}
 }
 
diff --git a/Engine/Source/Editor/PropertyEditor/Private/ItemPropertyNode.cpp b/Engine/Source/Editor/PropertyEditor/Private/ItemPropertyNode.cpp
--- a/Engine/Source/Editor/PropertyEditor/Private/ItemPropertyNode.cpp
+++ b/Engine/Source/Editor/PropertyEditor/Private/ItemPropertyNode.cpp
@@ -21,7 +21,7 @@ FItemPropertyNode::~FItemPropertyNode()
 
 }
 
-uint8* FItemPropertyNode::GetValueBaseAddress(uint8* StartAddress, bool bIsSparseData) const
+uint8* FItemPropertyNode::GetValueBaseAddress(uint8* StartAddress, bool bIsSparseData, bool bIsStruct) const
 {
 	const FProperty* MyProperty = GetProperty();
 	const TSharedPtr<FPropertyNode> ParentNode = ParentNodeWeakPtr.Pin();
@@ -31,7 +31,7 @@ uint8* FItemPropertyNode::GetValueBaseAddress(uint8* StartAddress, bool bIsSpars
 		const FSetProperty* OuterSetProp = MyProperty->GetOwner<FSetProperty>();
 		const FMapProperty* OuterMapProp = MyProperty->GetOwner<FMapProperty>();
 
-		uint8* ValueBaseAddress = ParentNode->GetValueBaseAddress(StartAddress, bIsSparseData);
+		uint8* ValueBaseAddress = ParentNode->GetValueBaseAddress(StartAddress, bIsSparseData, bIsStruct);
 
 		if (OuterArrayProp != nullptr)
 		{
@@ -68,11 +68,11 @@ uint8* FItemPropertyNode::GetValueBaseAddress(uint8* StartAddress, bool bIsSpars
 		}
 		else
 		{
-			uint8* ValueAddress = ParentNode->GetValueAddress(StartAddress, bIsSparseData);
+			uint8* ValueAddress = ParentNode->GetValueAddress(StartAddress, bIsSparseData, bIsStruct);
 			if (ValueAddress != nullptr && ParentNode->GetProperty() != MyProperty)
 			{
 				// if this is not a fixed size array (in which the parent property and this property are the same), we need to offset from the property (otherwise, the parent already did that for us)
-				ValueAddress = Property->ContainerPtrToValuePtr<uint8>(ValueAddress);
+				ValueAddress = MyProperty->ContainerPtrToValuePtr<uint8>(ValueAddress);
 			}
 
 			if (ValueAddress != nullptr)
@@ -86,9 +86,9 @@ uint8* FItemPropertyNode::GetValueBaseAddress(uint8* StartAddress, bool bIsSpars
 	return nullptr;
 }
 
-uint8* FItemPropertyNode::GetValueAddress(uint8* StartAddress, bool bIsSparseData) const
+uint8* FItemPropertyNode::GetValueAddress(uint8* StartAddress, bool bIsSparseData, bool bIsStruct) const
 {
-	uint8* Result = GetValueBaseAddress(StartAddress, bIsSparseData);
+	uint8* Result = GetValueBaseAddress(StartAddress, bIsSparseData, bIsStruct);
 	return Result;
 }
 
diff --git a/Engine/Source/Editor/PropertyEditor/Private/ItemPropertyNode.h b/Engine/Source/Editor/PropertyEditor/Private/ItemPropertyNode.h
--- a/Engine/Source/Editor/PropertyEditor/Private/ItemPropertyNode.h
+++ b/Engine/Source/Editor/PropertyEditor/Private/ItemPropertyNode.h
@@ -12,8 +12,8 @@ public:
 	FItemPropertyNode();
 	virtual ~FItemPropertyNode();
 
-	virtual uint8* GetValueBaseAddress(uint8* StartAddress, bool bIsSparseData) const override;
-	virtual uint8* GetValueAddress(uint8* StartAddress, bool bIsSparseData) const override;
+	virtual uint8* GetValueBaseAddress(uint8* StartAddress, bool bIsSparseData, bool bIsStruct) const override;
+	virtual uint8* GetValueAddress(uint8* StartAddress, bool bIsSparseData, bool bIsStruct) const override;
 
 	/**
 	 * Overridden function to get the derived object node
diff --git a/Engine/Source/Editor/PropertyEditor/Private/ObjectPropertyNode.cpp b/Engine/Source/Editor/PropertyEditor/Private/ObjectPropertyNode.cpp
--- a/Engine/Source/Editor/PropertyEditor/Private/ObjectPropertyNode.cpp
+++ b/Engine/Source/Editor/PropertyEditor/Private/ObjectPropertyNode.cpp
@@ -428,7 +428,7 @@ bool FObjectPropertyNode::GetReadAddressUncached(const FPropertyNode& InNode, FR
 	return true;
 }
 
-uint8* FObjectPropertyNode::GetValueBaseAddress(uint8* StartAddress, bool bIsSparseData) const
+uint8* FObjectPropertyNode::GetValueBaseAddress(uint8* StartAddress, bool bIsSparseData, bool bIsStruct) const
 {
 	uint8* Result = StartAddress;
 
diff --git a/Engine/Source/Editor/PropertyEditor/Private/ObjectPropertyNode.h b/Engine/Source/Editor/PropertyEditor/Private/ObjectPropertyNode.h
--- a/Engine/Source/Editor/PropertyEditor/Private/ObjectPropertyNode.h
+++ b/Engine/Source/Editor/PropertyEditor/Private/ObjectPropertyNode.h
@@ -152,7 +152,7 @@ protected:
 	virtual void InitBeforeNodeFlags() override;
 	virtual void InitChildNodes() override;
 	virtual bool GetQualifiedName( FString& PathPlusIndex, const bool bWithArrayIndex, const FPropertyNode* StopParent = nullptr, bool bIgnoreCategories = false ) const override;
-	virtual uint8* GetValueBaseAddress(uint8* Base, bool bIsSparseData) const override;
+	virtual uint8* GetValueBaseAddress(uint8* Base, bool bIsSparseData, bool bIsStruct) const override;
 	/**
 	 * Looks at the Objects array and creates the best base class.  Called by
 	 * Finalize(); that is, when the list of selected objects is being finalized.
diff --git a/Engine/Source/Editor/PropertyEditor/Private/PropertyHandleImpl.cpp b/Engine/Source/Editor/PropertyEditor/Private/PropertyHandleImpl.cpp
--- a/Engine/Source/Editor/PropertyEditor/Private/PropertyHandleImpl.cpp
+++ b/Engine/Source/Editor/PropertyEditor/Private/PropertyHandleImpl.cpp
@@ -75,7 +75,7 @@ void FPropertyValueImpl::EnumerateObjectsToModify( FPropertyNode* InPropertyNode
 				Object = ComplexNode->GetInstanceAsUObject(Index).Get();
 				StructAddress = InPropertyNode->GetStartAddressFromObject(Object);
 			}
-			uint8* BaseAddress = InPropertyNode->GetValueBaseAddress(StructAddress, InPropertyNode->HasNodeFlags(EPropertyNodeFlags::IsSparseClassData));
+			uint8* BaseAddress = InPropertyNode->GetValueBaseAddress(StructAddress, InPropertyNode->HasNodeFlags(EPropertyNodeFlags::IsSparseClassData), bIsStruct);
 
 			if (!InObjectsToModifyCallback(FObjectBaseAddress(Object, StructAddress, BaseAddress), Index, NumInstances))
 			{
@@ -327,6 +327,7 @@ FPropertyAccess::Result FPropertyValueImpl::ImportText( const TArray<FObjectBase
 			FProperty* Property = ParentNode ? ParentNode->GetProperty() : nullptr;
 
 			const bool bIsSparseClassData = InPropertyNode->HasNodeFlags(EPropertyNodeFlags::IsSparseClassData) != 0;
+			const bool bIsStruct = !Cur.Object;
 
 			bool bIsInContainer = false;
 
@@ -341,7 +342,7 @@ FPropertyAccess::Result FPropertyValueImpl::ImportText( const TArray<FObjectBase
 
 				if (bIsInContainer)
 				{
-					uint8* ValueBaseAddress = ParentNode->GetValueBaseAddress(Cur.StructAddress, bIsSparseClassData);
+					uint8* ValueBaseAddress = ParentNode->GetValueBaseAddress(Cur.StructAddress, bIsSparseClassData, bIsStruct);
 					
 					/**
 					 * Checks if an element has already been added to the set
@@ -397,7 +398,7 @@ FPropertyAccess::Result FPropertyValueImpl::ImportText( const TArray<FObjectBase
 
 				if (bIsInContainer)
 				{
-					uint8* ValueBaseAddress = ParentNode->GetValueBaseAddress(Cur.StructAddress, bIsSparseClassData);
+					uint8* ValueBaseAddress = ParentNode->GetValueBaseAddress(Cur.StructAddress, bIsSparseClassData, bIsStruct);
 
 					/**
 					 * Checks if a key in the map matches the specified key
@@ -455,7 +456,7 @@ FPropertyAccess::Result FPropertyValueImpl::ImportText( const TArray<FObjectBase
 
 			if (bIsInContainer)
 			{
-				uint8* Addr = ParentNode->GetValueBaseAddress(Cur.StructAddress, bIsSparseClassData);
+				uint8* Addr = ParentNode->GetValueBaseAddress(Cur.StructAddress, bIsSparseClassData, bIsStruct);
 				Property->ExportText_Direct(PreviousContainerValue, Addr, Addr, nullptr, 0);
 			}
 
@@ -521,13 +522,13 @@ FPropertyAccess::Result FPropertyValueImpl::ImportText( const TArray<FObjectBase
 				// For TMap and TSet, we need to rehash it in case a key was modified
 				if (NodeProperty->GetOwner<FMapProperty>())
 				{
-					uint8* Addr = InPropertyNode->GetParentNode()->GetValueBaseAddress(Cur.StructAddress, bIsSparseClassData);
+					uint8* Addr = InPropertyNode->GetParentNode()->GetValueBaseAddress(Cur.StructAddress, bIsSparseClassData, bIsStruct);
 					FScriptMapHelper MapHelper(NodeProperty->GetOwner<FMapProperty>(), Addr);
 					MapHelper.Rehash();
 				}
 				else if (NodeProperty->GetOwner<FSetProperty>())
 				{
-					uint8* Addr = InPropertyNode->GetParentNode()->GetValueBaseAddress(Cur.StructAddress, bIsSparseClassData);
+					uint8* Addr = InPropertyNode->GetParentNode()->GetValueBaseAddress(Cur.StructAddress, bIsSparseClassData, bIsStruct);
 					FScriptSetHelper SetHelper(NodeProperty->GetOwner<FSetProperty>(), Addr);
 					SetHelper.Rehash();
 				}
@@ -2285,7 +2286,6 @@ FString FPropertyHandleBase::GeneratePathToProperty() const
 	}
 
 	return OutPath;
-
 }
 
 TSharedRef<SWidget> FPropertyHandleBase::CreatePropertyNameWidget(const FText& NameOverride, const FText& ToolTipOverride, bool bDisplayResetToDefault, bool bDisplayText, bool bDisplayThumbnail) const
@@ -2585,7 +2585,7 @@ void FPropertyHandleBase::GetOuterPackages(TArray<UPackage*>& OuterPackages) con
 			case FComplexPropertyNode::EPT_StandaloneStructure:
 			{
 				FStructurePropertyNode* StructNode = static_cast<FStructurePropertyNode*>(ComplexNode);
-				OuterPackages.Add(StructNode->GetOwnerPackage());
+				StructNode->GetOwnerPackages(OuterPackages);
 			}
 			break;
 
@@ -3315,6 +3315,12 @@ void FPropertyHandleBase::SetIgnoreValidation(bool bInIgnore)
 }
 
 TArray<TSharedPtr<IPropertyHandle>> FPropertyHandleBase::AddChildStructure( TSharedRef<FStructOnScope> InStruct )
+{
+	TSharedRef<FStructOnScopeStructureDataProvider> StructProvider = MakeShared<FStructOnScopeStructureDataProvider>(InStruct.ToSharedPtr());
+	return AddChildStructure(StructProvider);
+}
+
+TArray<TSharedPtr<IPropertyHandle>> FPropertyHandleBase::AddChildStructure(TSharedRef<IStructureDataProvider> InStructProvider)
 {
 	TArray<TSharedPtr<IPropertyHandle>> PropertyHandles;
 
@@ -3325,7 +3331,7 @@ TArray<TSharedPtr<IPropertyHandle>> FPropertyHandleBase::AddChildStructure( TSha
 	}
 
 	TSharedPtr<FStructurePropertyNode> StructPropertyNode( new FStructurePropertyNode );
-	StructPropertyNode->SetStructure(InStruct);
+	StructPropertyNode->SetStructure(InStructProvider);
 
 	FPropertyNodeInitParams RootInitParams;
 	RootInitParams.ParentNode = PropertyNode;
@@ -3335,7 +3341,6 @@ TArray<TSharedPtr<IPropertyHandle>> FPropertyHandleBase::AddChildStructure( TSha
 	RootInitParams.bAllowChildren = true;
 	RootInitParams.bForceHiddenPropertyVisibility = FPropertySettings::Get().ShowHiddenProperties();
 	RootInitParams.bCreateCategoryNodes = false;
-	RootInitParams.IsSparseProperty = FPropertyNodeInitParams::EIsSparseDataProperty::False; // FStructurePropertyNode can't inherit the sparse data flag
 
 	StructPropertyNode->InitNode(RootInitParams);
 
@@ -4949,9 +4954,10 @@ bool FPropertyHandleSet::HasDefaultElement()
 
 		if (Addresses.Num() > 0)
 		{
-			const bool IsSparseClassData = PropNode->HasNodeFlags(EPropertyNodeFlags::IsSparseClassData) != 0;
+			const bool bIsSparseClassData = PropNode->HasNodeFlags(EPropertyNodeFlags::IsSparseClassData) != 0;
+			const bool bIsStruct = !Addresses[0].Object;
 			FSetProperty* SetProperty = CastFieldChecked<FSetProperty>(PropNode->GetProperty());
-			FScriptSetHelper SetHelper(SetProperty, PropNode->GetValueBaseAddress(Addresses[0].StructAddress, IsSparseClassData));
+			FScriptSetHelper SetHelper(SetProperty, PropNode->GetValueBaseAddress(Addresses[0].StructAddress, bIsSparseClassData, bIsStruct));
 
 			FDefaultConstructedPropertyElement DefaultElement(SetHelper.ElementProp);
 			return SetHelper.FindElementIndex(DefaultElement.GetObjAddress()) != INDEX_NONE;
@@ -5058,9 +5064,10 @@ bool FPropertyHandleMap::HasDefaultKey()
 
 		if (Addresses.Num() > 0)
 		{
-			const bool IsSparseClassData = PropNode->HasNodeFlags(EPropertyNodeFlags::IsSparseClassData) != 0;
+			const bool bIsSparseClassData = PropNode->HasNodeFlags(EPropertyNodeFlags::IsSparseClassData) != 0;
+			const bool bIsStruct = !Addresses[0].Object;
 			FMapProperty* MapProperty = CastFieldChecked<FMapProperty>(PropNode->GetProperty());
-			FScriptMapHelper MapHelper(MapProperty, PropNode->GetValueBaseAddress(Addresses[0].StructAddress, IsSparseClassData));
+			FScriptMapHelper MapHelper(MapProperty, PropNode->GetValueBaseAddress(Addresses[0].StructAddress, bIsSparseClassData, bIsStruct));
 
 			FDefaultConstructedPropertyElement DefaultKey(MapHelper.KeyProp);
 			return MapHelper.FindMapIndexWithKey(DefaultKey.GetObjAddress()) != INDEX_NONE;
diff --git a/Engine/Source/Editor/PropertyEditor/Private/PropertyHandleImpl.h b/Engine/Source/Editor/PropertyEditor/Private/PropertyHandleImpl.h
--- a/Engine/Source/Editor/PropertyEditor/Private/PropertyHandleImpl.h
+++ b/Engine/Source/Editor/PropertyEditor/Private/PropertyHandleImpl.h
@@ -503,6 +503,7 @@ public:
 	virtual bool GenerateRestrictionToolTip(const FString& Value, FText& OutTooltip) const override;
 	virtual void SetIgnoreValidation(bool bInIgnore) override;
 	virtual TArray<TSharedPtr<IPropertyHandle>> AddChildStructure(TSharedRef<FStructOnScope> ChildStructure) override;
+	virtual TArray<TSharedPtr<IPropertyHandle>> AddChildStructure(TSharedRef<IStructureDataProvider> ChildStructProvider) override;
 	virtual bool CanResetToDefault() const override;
 	virtual void ExecuteCustomResetToDefault(const FResetToDefaultOverride& InOnCustomResetToDefault) override;
 	virtual FName GetDefaultCategoryName() const override;
diff --git a/Engine/Source/Editor/PropertyEditor/Private/PropertyNode.cpp b/Engine/Source/Editor/PropertyEditor/Private/PropertyNode.cpp
--- a/Engine/Source/Editor/PropertyEditor/Private/PropertyNode.cpp
+++ b/Engine/Source/Editor/PropertyEditor/Private/PropertyNode.cpp
@@ -16,6 +16,7 @@
 #include "Misc/ScopeExit.h"
 #include "Editor.h"
 #include "ObjectPropertyNode.h"
+#include "StructurePropertyNode.h"
 #include "PropertyHandleImpl.h"
 #include "PropertyTextUtilities.h"
 #include "EditorSupportDelegates.h"
@@ -412,6 +413,38 @@ const class FObjectPropertyNode* FPropertyNode::FindObjectItemParent() const
 	return nullptr;
 }
 
+FStructurePropertyNode* FPropertyNode::FindStructureItemParent()
+{
+	FComplexPropertyNode* ComplexParent = FindComplexParent();
+	if (!ComplexParent)
+	{
+		return nullptr;
+	}
+
+	if (FStructurePropertyNode* StructureNode = ComplexParent->AsStructureNode())
+	{
+		return StructureNode;
+	}
+
+	return nullptr;
+}
+
+const FStructurePropertyNode* FPropertyNode::FindStructureItemParent() const
+{
+	const FComplexPropertyNode* ComplexParent = FindComplexParent();
+	if (!ComplexParent)
+	{
+		return nullptr;
+	}
+
+	if (const FStructurePropertyNode* StructureNode = ComplexParent->AsStructureNode())
+	{
+		return StructureNode;
+	}
+
+	return nullptr;
+}
+
 /**
  * Follows the top-most object window that contains this property window item.
  */
@@ -1531,7 +1564,7 @@ uint8* FPropertyNode::GetValueAddressFromObject(const UObject* Obj) const
 }
 
 
-uint8* FPropertyNode::GetValueBaseAddress(uint8* StartAddress, bool bIsSparseData) const
+uint8* FPropertyNode::GetValueBaseAddress(uint8* StartAddress, bool bIsSparseData, bool bIsStruct) const
 {
 	uint8* Result = NULL;
 
@@ -1544,15 +1577,15 @@ uint8* FPropertyNode::GetValueBaseAddress(uint8* StartAddress, bool bIsSparseDat
 		const TSharedPtr<FPropertyNode> ParentNode = ParentNodeWeakPtr.Pin();
 		if (ParentNode.IsValid())
 		{
-			Result = ParentNode->GetValueAddress(StartAddress, bIsSparseData);
+			Result = ParentNode->GetValueAddress(StartAddress, bIsSparseData, bIsStruct);
 		}
 	}
 	return Result;
 }
 
-uint8* FPropertyNode::GetValueAddress(uint8* StartAddress, bool bIsSparseData) const
+uint8* FPropertyNode::GetValueAddress(uint8* StartAddress, bool bIsSparseData, bool bIsStruct) const
 {
-	return GetValueBaseAddress(StartAddress, bIsSparseData);
+	return GetValueBaseAddress(StartAddress, bIsSparseData, bIsStruct);
 }
 
 
@@ -1599,6 +1632,13 @@ public:
 		check(Property);
 		check(PropertyValueRoot.OwnerObject);
 
+		// Do not cache pointers for standalone structures, as we don't have the same guarantees how the provided pointers are invalidated as we have with UObject nodes. 
+		// The default value handling for structure nodes is done in FPropertyNode::GetDefaultValueAsString() and FPropertyNode::GetDiffersFromDefault().
+		if (PropertyNode->FindStructureItemParent())
+		{
+			return;
+		}
+		
 		FPropertyNode* ParentNode = PropertyNode->GetParentNode();
 
 		// if the object specified is a class object, transfer to the CDO instead
@@ -1617,7 +1657,7 @@ public:
 
 		if (IsValidTracker())
 		{
-			 bHasDefaultValue = Private_HasDefaultValue();
+			bHasDefaultValue = Private_HasDefaultValue();
 
 			// calculate the addresses for the default object if it exists
 			if (bHasDefaultValue)
@@ -1736,7 +1776,7 @@ public:
 	uint8* GetPropertyDefaultRootAddress() const { return PropertyDefaultValueRoot.ValueAddress; }
 
 private:
-		/**
+	/**
 	 * Determines whether the property bound to this struct exists in the owning object's archetype.
 	 *
 	 * @return	true if this property exists in the owning object's archetype; false if the archetype is e.g. a
@@ -1757,7 +1797,7 @@ private:
 					UStruct* SparseClassDataArchetypeStruct = OwnerClass->GetSparseClassDataArchetypeStruct();
 
 					if (SparseClassDataStruct == SparseClassDataArchetypeStruct)
-		{
+					{
 						bResult = true;
 					}
 					else
@@ -1797,20 +1837,20 @@ private:
 			else
 			{
 				// Find the member property which contains this item's property
-			FPropertyNode* MemberPropertyNode = PropertyNode;
-			for ( ;MemberPropertyNode != NULL; MemberPropertyNode = MemberPropertyNode->GetParentNode() )
-			{
-				FProperty* MemberProperty = MemberPropertyNode->GetProperty();
-				if ( MemberProperty != NULL )
+				FPropertyNode* MemberPropertyNode = PropertyNode;
+				for ( ;MemberPropertyNode != NULL; MemberPropertyNode = MemberPropertyNode->GetParentNode() )
 				{
-					if ( MemberProperty->GetOwner<UClass>() != NULL )
+					FProperty* MemberProperty = MemberPropertyNode->GetProperty();
+					if ( MemberProperty != NULL )
 					{
-						break;
+						if ( MemberProperty->GetOwner<UClass>() != NULL )
+						{
+							break;
+						}
 					}
 				}
-			}
 				if ( MemberPropertyNode != NULL && MemberPropertyNode->GetProperty())
-			{
+				{
 					// we check to see that this property is in the defaults class
 					bResult = MemberPropertyNode->GetProperty()->IsInContainer(ParentDefault->GetClass());
 				}
@@ -2192,104 +2232,97 @@ private:
 
 };
 
-bool FPropertyNode::GetDiffersFromDefaultForObject( FPropertyItemValueDataTrackerSlate& ValueTracker, FProperty* InProperty )
-{	
-	check( InProperty );
-
-	bool bDiffersFromDefaultForObject = false;
 
-	bool bIsValidTracker = ValueTracker.IsValidTracker();
-	bool bHasDefaultValue = ValueTracker.HasDefaultValue();
-	bool bHasParent = GetParentNode() != nullptr;
+bool FPropertyNode::GetDiffersFromDefault(const uint8* PropertyValueAddress, const uint8* PropertyDefaultAddress, const uint8* DefaultPropertyValueBaseAddress, const FProperty* InProperty) const
+{
+	bool bDiffersFromDefaultValue = false;
 
-	if (bIsValidTracker && bHasDefaultValue && bHasParent)
+	if (DefaultPropertyValueBaseAddress != nullptr)
 	{
-		//////////////////////////
-		// Check the property against its default.
-		// If the property is an object property, we have to take special measures.
-
-		if (FArrayProperty* OuterArrayProperty = InProperty->GetOwner<FArrayProperty>())
+		if (const FArrayProperty* OuterArrayProperty = InProperty->GetOwner<FArrayProperty>())
 		{
-				// make sure we're not trying to compare against an element that doesn't exist
-			if (ValueTracker.GetPropertyDefaultBaseAddress() != nullptr)
+			// make sure we're not trying to compare against an element that doesn't exist
+			FScriptArrayHelper ArrayHelper(OuterArrayProperty, DefaultPropertyValueBaseAddress);
+			if (!ArrayHelper.IsValidIndex(GetArrayIndex()))
 			{
-				// make sure we're not trying to compare against an element that doesn't exist
-				FScriptArrayHelper ArrayHelper(OuterArrayProperty, ValueTracker.GetPropertyDefaultBaseAddress());
-				if (GetArrayIndex() >= ArrayHelper.Num())
-				{
-					bDiffersFromDefaultForObject = true;
-				}
+				bDiffersFromDefaultValue = true;
 			}
 		}
-		else if (FSetProperty* OuterSetProperty = InProperty->GetOwner<FSetProperty>())
+		else if (const FSetProperty* OuterSetProperty = InProperty->GetOwner<FSetProperty>())
 		{
-			if (ValueTracker.GetPropertyDefaultBaseAddress() != nullptr)
+			FScriptSetHelper SetHelper(OuterSetProperty, DefaultPropertyValueBaseAddress);
+			if (!SetHelper.IsValidIndex(GetArrayIndex()))
 			{
-				FScriptSetHelper SetHelper(OuterSetProperty, ValueTracker.GetPropertyDefaultBaseAddress());
-				if ( ValueTracker.GetPropertyDefaultBaseAddress() != NULL && (ArrayIndex < 0 || ArrayIndex >= SetHelper.Num()))
-				{
-					bDiffersFromDefaultForObject = true;
-				}
+				bDiffersFromDefaultValue = true;
 			}
 		}
-		else if (FMapProperty* OuterMapProperty = InProperty->GetOwner<FMapProperty>())
+		else if (const FMapProperty* OuterMapProperty = InProperty->GetOwner<FMapProperty>())
 		{
-			if (ValueTracker.GetPropertyDefaultBaseAddress() != nullptr)
+			FScriptMapHelper MapHelper(OuterMapProperty, DefaultPropertyValueBaseAddress);
+			if (!MapHelper.IsValidIndex(GetArrayIndex()))
 			{
-				FScriptMapHelper MapHelper(OuterMapProperty, ValueTracker.GetPropertyDefaultBaseAddress());
-				if ( ValueTracker.GetPropertyDefaultBaseAddress() != NULL && (ArrayIndex < 0 || ArrayIndex >= MapHelper.Num()))
-				{
-					bDiffersFromDefaultForObject = true;
-				}
+				bDiffersFromDefaultValue = true;
 			}
 		}
+	}
 
-		// The property is a simple field.  Compare it against the enclosing object's default for that property.
-		if (!bDiffersFromDefaultForObject)
+	if (!bDiffersFromDefaultValue)
+	{
+		uint32 PortFlags = 0;
+		if (InProperty->ContainsInstancedObjectProperty())
 		{
-			uint32 PortFlags = 0;
-			if (InProperty->ContainsInstancedObjectProperty())
-			{
-				PortFlags |= PPF_DeepComparison;
-			}
-
-			if ( ValueTracker.GetPropertyValueAddress() == NULL || ValueTracker.GetPropertyDefaultAddress() == NULL )
-			{
-				// if either are NULL, we had a dynamic array somewhere in our parent chain and the array doesn't
-				// have enough elements in either the default or the object
-				bDiffersFromDefaultForObject = true;
-			}
-			else if (GetArrayIndex() == INDEX_NONE && InProperty->ArrayDim > 1)
+			PortFlags |= PPF_DeepComparison;
+		}
+	
+		if (PropertyValueAddress == nullptr || PropertyDefaultAddress == nullptr)
+		{
+			// if either are NULL, we had a dynamic array somewhere in our parent chain and the array doesn't
+			// have enough elements in either the default or the object
+			bDiffersFromDefaultValue = true;
+		}
+		else if (GetArrayIndex() == INDEX_NONE && InProperty->ArrayDim > 1)
+		{
+			// this is a container; loop through all of its elements and see if any of them differ from the default
+			for (int32 Idx = 0; !bDiffersFromDefault && Idx < InProperty->ArrayDim; Idx++)
 			{
-				// this is a container; loop through all of its elements and see if any of them differ from the default
-				for (int32 Idx = 0; !bDiffersFromDefaultForObject && Idx < InProperty->ArrayDim; Idx++)
-				{
-					bDiffersFromDefaultForObject = !InProperty->Identical(
-						ValueTracker.GetPropertyValueAddress() + Idx * InProperty->ElementSize,
-						ValueTracker.GetPropertyDefaultAddress() + Idx * InProperty->ElementSize,
-						PortFlags
-						);
-				}
+				bDiffersFromDefaultValue = !InProperty->Identical(
+					PropertyValueAddress + Idx * InProperty->ElementSize,
+					PropertyDefaultAddress + Idx * InProperty->ElementSize,
+					PortFlags
+					);
 			}
-			else
+		}
+		else
+		{
+			// try to compare the values at the current and default property addresses
+			if( PropertyValueAddress != nullptr && PropertyDefaultAddress != nullptr )
 			{
-				// try to compare the values at the current and default property addresses
-				uint8* PropertyValueAddr = ValueTracker.GetPropertyValueAddress();
-				uint8* DefaultPropertyValueAddr = ValueTracker.GetPropertyDefaultAddress();
-
-				if( PropertyValueAddr != NULL && DefaultPropertyValueAddr != NULL )
-				{
-					bDiffersFromDefaultForObject = !InProperty->Identical(
-						PropertyValueAddr,
-						DefaultPropertyValueAddr,
-						PortFlags
-						);
-				}
+				bDiffersFromDefaultValue = !InProperty->Identical(
+					PropertyValueAddress,
+					PropertyDefaultAddress,
+					PortFlags
+					);
 			}
 		}
 	}
 
-	return bDiffersFromDefaultForObject;
+	return bDiffersFromDefaultValue;
+}
+
+bool FPropertyNode::GetDiffersFromDefaultForObject( FPropertyItemValueDataTrackerSlate& ValueTracker, FProperty* InProperty )
+{	
+	check( InProperty );
+
+	const bool bIsValidTracker = ValueTracker.IsValidTracker();
+	const bool bHasDefaultValue = ValueTracker.HasDefaultValue();
+	const bool bHasParent = GetParentNode() != nullptr;
+
+	if (bIsValidTracker && bHasDefaultValue && bHasParent)
+	{
+		return GetDiffersFromDefault(ValueTracker.GetPropertyValueAddress(), ValueTracker.GetPropertyDefaultAddress(), ValueTracker.GetPropertyDefaultBaseAddress(), InProperty);
+	}
+
+	return false;
 }
 
 /**
@@ -2308,17 +2341,67 @@ bool FPropertyNode::GetDiffersFromDefault()
 			return bDiffersFromDefault;
 		}
 
-		FObjectPropertyNode* ObjectNode = FindObjectItemParent();
-		if (ObjectNode && Property.IsValid())
+		if (const FStructurePropertyNode* StructNode = FindStructureItemParent())
+		{
+			TArray<TSharedPtr<FStructOnScope>> Structs;
+			StructNode->GetAllStructureData(Structs);
+			
+			const bool bIsSparse = HasNodeFlags(EPropertyNodeFlags::IsSparseClassData);
+			const bool bIsContainer = CastField<FArrayProperty>(Prop) || CastField<FSetProperty>(Prop) || CastField<FMapProperty>(Prop);
+			const bool bIsInsideContainerProperty = Property->GetOwner<FArrayProperty>() || Property->GetOwner<FSetProperty>() || Property->GetOwner<FMapProperty>();
+			const FPropertyNode* BaseNode = bIsInsideContainerProperty ? GetParentNode() : this;
+
+			FStructOnScope DefaultStruct;
+
+			for (int32 Index = 0; !bDiffersFromDefault && Index < Structs.Num(); Index++)
+			{
+				const TSharedPtr<FStructOnScope>& StructData = Structs[Index];
+				// Skip empty data.
+				if (!StructData.IsValid())
+				{
+					continue;
+				}
+				const UStruct* Struct = StructData->GetStruct();
+				if (!Struct)
+				{
+					continue;
+				}
+				
+				// Make an instance of the struct to be used as default value to test against.
+				if (DefaultStruct.GetStruct() != Struct)
+				{
+					DefaultStruct.Initialize(Struct);
+				}
+
+				check(DefaultStruct.IsValid());
+
+				const uint8* PropertyValueAddress = GetValueAddress(StructData->GetStructMemory(), bIsSparse, /*bIsStruct=*/true);
+				const uint8* PropertyDefaultAddress = GetValueAddress(DefaultStruct.GetStructMemory(), bIsSparse, /*bIsStruct=*/true);
+				const uint8* PropertyDefaultBaseAddress = BaseNode->GetValueBaseAddress(DefaultStruct.GetStructMemory(), bIsSparse, /*bIsStruct=*/true);
+
+				// If this is a container property, we must take special measures to use the base address of the property's value; for instance,
+				// the array property's PropertyDefaultBaseAddress points to an FScriptArray*, while PropertyDefaultAddress points to the 
+				// FScriptArray's Data pointer.
+				if (bIsContainer)
+				{
+					const uint8* PropertyValueBaseAddress = BaseNode->GetValueBaseAddress(StructData->GetStructMemory(), bIsSparse, /*bIsStruct=*/true);
+					PropertyValueAddress = PropertyValueBaseAddress;
+					PropertyDefaultAddress = PropertyDefaultBaseAddress;
+				}
+
+				bDiffersFromDefault = GetDiffersFromDefault(PropertyValueAddress, PropertyDefaultAddress, PropertyDefaultBaseAddress, Prop);
+			}
+		}
+		else if (FObjectPropertyNode* ObjectNode = FindObjectItemParent())
 		{
 			// Get an iterator for the enclosing objects.
-			for(int32 ObjIndex = 0; ObjIndex < ObjectNode->GetNumObjects(); ++ObjIndex)
+			for (int32 ObjIndex = 0; ObjIndex < ObjectNode->GetNumObjects(); ++ObjIndex)
 			{
 				UObject* Object = ObjectNode->GetUObject(ObjIndex);
 
 				TSharedPtr<FPropertyItemValueDataTrackerSlate> ValueTracker = GetValueTracker(Object, ObjIndex);
 
-				if(ValueTracker.IsValid() && Object && GetDiffersFromDefaultForObject(*ValueTracker, Property.Get()))
+				if (Object && GetDiffersFromDefaultForObject(*ValueTracker, Prop))
 				{
 					// If any object being observed differs from the result then there is no need to keep searching
 					bDiffersFromDefault = true;
@@ -2331,12 +2414,49 @@ bool FPropertyNode::GetDiffersFromDefault()
 	return bDiffersFromDefault;
 }
 
+
+FString FPropertyNode::GetDefaultValueAsString(const uint8* PropertyDefaultAddress, const FProperty* InProperty, const bool bUseDisplayName) const
+{
+	FString DefaultValue;
+
+	uint32 PortFlags = bUseDisplayName ? PPF_PropertyWindow : PPF_None;
+	if (InProperty->ContainsInstancedObjectProperty())
+	{
+		PortFlags |= PPF_DeepComparison;
+	}
+
+	if (!PropertyDefaultAddress)
+	{
+		// no default available, fall back on the default value for our primitive:
+		uint8* TempComplexPropAddr = (uint8*)FMemory::Malloc(InProperty->GetSize(), InProperty->GetMinAlignment());
+		InProperty->InitializeValue(TempComplexPropAddr);
+		ON_SCOPE_EXIT
+		{
+			InProperty->DestroyValue(TempComplexPropAddr);
+			FMemory::Free(TempComplexPropAddr);
+		};
+				
+		InProperty->ExportText_Direct(DefaultValue, TempComplexPropAddr, TempComplexPropAddr, nullptr, PPF_None);
+	}
+	else if ( GetArrayIndex() == INDEX_NONE && InProperty->ArrayDim > 1 )
+	{
+		FArrayProperty::ExportTextInnerItem(DefaultValue, InProperty, PropertyDefaultAddress, InProperty->ArrayDim,
+											PropertyDefaultAddress, InProperty->ArrayDim, nullptr, PortFlags);
+	}
+	else
+	{
+		// Port flags will cause enums to display correctly
+		InProperty->ExportTextItem_Direct(DefaultValue, PropertyDefaultAddress, PropertyDefaultAddress, nullptr, PortFlags, nullptr);
+	}
+
+	return DefaultValue;
+}
+
 FString FPropertyNode::GetDefaultValueAsStringForObject( FPropertyItemValueDataTrackerSlate& ValueTracker, UObject* InObject, FProperty* InProperty, bool bUseDisplayName)
 {
 	check( InObject );
 	check( InProperty );
 
-	bool bDiffersFromDefaultForObject = false;
 	FString DefaultValue;
 
 	// special case for Object class - no defaults to compare against
@@ -2344,36 +2464,7 @@ FString FPropertyNode::GetDefaultValueAsStringForObject( FPropertyItemValueDataT
 	{
 		if ( ValueTracker.IsValidTracker() && ValueTracker.HasDefaultValue() )
 		{
-			uint32 PortFlags = bUseDisplayName ? PPF_PropertyWindow : PPF_None;
-			
-			if (InProperty->ContainsInstancedObjectProperty())
-			{
-				PortFlags |= PPF_DeepComparison;
-			}
-
-			if ( ValueTracker.GetPropertyDefaultAddress() == NULL )
-			{
-				// no default available, fall back on the default value for our primitive:
-				uint8* TempComplexPropAddr = (uint8*)FMemory::Malloc(InProperty->GetSize(), InProperty->GetMinAlignment());
-				InProperty->InitializeValue(TempComplexPropAddr);
-				ON_SCOPE_EXIT
-				{
-					InProperty->DestroyValue(TempComplexPropAddr);
-					FMemory::Free(TempComplexPropAddr);
-				};
-					
-				InProperty->ExportText_Direct(DefaultValue, TempComplexPropAddr, TempComplexPropAddr, nullptr, PPF_None);
-			}
-			else if ( GetArrayIndex() == INDEX_NONE && InProperty->ArrayDim > 1 )
-			{
-				FArrayProperty::ExportTextInnerItem(DefaultValue, InProperty, ValueTracker.GetPropertyDefaultAddress(), InProperty->ArrayDim,
-													ValueTracker.GetPropertyDefaultAddress(), InProperty->ArrayDim, nullptr, PortFlags);
-			}
-			else
-			{
-				// Port flags will cause enums to display correctly
-				InProperty->ExportTextItem_Direct(DefaultValue, ValueTracker.GetPropertyDefaultAddress(), ValueTracker.GetPropertyDefaultAddress(), nullptr, PortFlags, nullptr);
-			}
+			DefaultValue = GetDefaultValueAsString(ValueTracker.GetPropertyDefaultAddress(), InProperty, bUseDisplayName);
 		}
 	}
 
@@ -2384,15 +2475,78 @@ FString FPropertyNode::GetDefaultValueAsString(bool bUseDisplayName)
 {
 	FString DefaultValue;
 	FString DelimitedValue;
-
 	bool bAllSame = true;
-	
-	FObjectPropertyNode* ObjectNode = FindObjectItemParent();
-	if (ObjectNode && Property.IsValid())
+
+	FProperty* Prop = GetProperty();
+	if (!Prop)
 	{
-		TArray<FString> Values;
-		Values.Reserve(ObjectNode->GetNumObjects());
+		return DefaultValue;
+	}
 
+	if (const FStructurePropertyNode* StructNode = FindStructureItemParent())
+	{
+		TArray<TSharedPtr<FStructOnScope>> Structs;
+		StructNode->GetAllStructureData(Structs);
+
+		const bool bIsSparse = HasNodeFlags(EPropertyNodeFlags::IsSparseClassData);
+		const bool bIsContainer = CastField<FArrayProperty>(Prop) || CastField<FSetProperty>(Prop) || CastField<FMapProperty>(Prop);
+		const bool bIsInsideContainerProperty = Property->GetOwner<FArrayProperty>() || Property->GetOwner<FSetProperty>() || Property->GetOwner<FMapProperty>();
+		const FPropertyNode* BaseNode = bIsInsideContainerProperty ? GetParentNode() : this;
+
+		FStructOnScope DefaultStruct;
+		FString NodeDefaultValue;
+
+		for (const TSharedPtr<FStructOnScope>& StructData : Structs)
+		{
+			if (!StructData.IsValid())
+			{
+				continue;
+			}
+			const UStruct* Struct = StructData->GetStruct();
+			if (!StructData)
+			{
+				continue;
+			}
+
+			if (DefaultStruct.GetStruct() != Struct)
+			{
+				// Make an instance of the struct to be used as default value.
+				DefaultStruct.Initialize(Struct);
+				check(DefaultStruct.IsValid());
+
+				const uint8* PropertyDefaultAddress = GetValueAddress(DefaultStruct.GetStructMemory(), bIsSparse, /*bIsStruct=*/true);
+				const uint8* PropertyDefaultBaseAddress = BaseNode->GetValueBaseAddress(DefaultStruct.GetStructMemory(), bIsSparse, /*bIsStruct=*/true);
+
+				// If this is a container property, we must take special measures to use the base address of the property's value; for instance,
+				// the array property's PropertyDefaultBaseAddress points to an FScriptArray*, while PropertyDefaultAddress points to the 
+				// FScriptArray's Data pointer.
+				if (bIsContainer)
+				{
+					PropertyDefaultAddress = PropertyDefaultBaseAddress;
+				}
+
+				NodeDefaultValue = GetDefaultValueAsString(PropertyDefaultAddress, Prop, bUseDisplayName);
+			}
+			
+			if (DefaultValue.IsEmpty())
+			{
+				DefaultValue = NodeDefaultValue;
+			}
+
+			if (DelimitedValue.Len() > 0 && NodeDefaultValue.Len() > 0)
+			{
+				DelimitedValue += TEXT(", ");
+			}
+			DelimitedValue += NodeDefaultValue;
+
+			if (!ensureAlwaysMsgf(NodeDefaultValue == DefaultValue, TEXT("Default values differ for different objects of property '%s'. First: \"%s\", Other: \"%s\""), *Prop->GetNameCPP(), *DefaultValue, *NodeDefaultValue))
+			{
+				bAllSame = false;
+			}
+		}
+	}
+	else if (FObjectPropertyNode* ObjectNode = FindObjectItemParent())
+	{
 		// Get an iterator for the enclosing objects.
 		for (int32 ObjIndex = 0; ObjIndex < ObjectNode->GetNumObjects(); ++ObjIndex)
 		{
@@ -2401,7 +2555,7 @@ FString FPropertyNode::GetDefaultValueAsString(bool bUseDisplayName)
 
 			if (Object && ValueTracker.IsValid())
 			{
-				const FString NodeDefaultValue = GetDefaultValueAsStringForObject( *ValueTracker, Object, Property.Get(), bUseDisplayName );
+				const FString NodeDefaultValue = GetDefaultValueAsStringForObject( *ValueTracker, Object, Prop, bUseDisplayName );
 
 				if (DefaultValue.IsEmpty())
 				{
@@ -2414,7 +2568,7 @@ FString FPropertyNode::GetDefaultValueAsString(bool bUseDisplayName)
 				}
 				DelimitedValue += NodeDefaultValue;
 
-				if (!ensureAlwaysMsgf(NodeDefaultValue == DefaultValue, TEXT("Default values differ for different objects of property '%s'. First: \"%s\", Other: \"%s\""), *Property->GetNameCPP(), *DefaultValue, *NodeDefaultValue))
+				if (!ensureAlwaysMsgf(NodeDefaultValue == DefaultValue, TEXT("Default values differ for different objects of property '%s'. First: \"%s\", Other: \"%s\""), *Prop->GetNameCPP(), *DefaultValue, *NodeDefaultValue))
 				{
 					bAllSame = false;
 				}
diff --git a/Engine/Source/Editor/PropertyEditor/Private/PropertyNode.h b/Engine/Source/Editor/PropertyEditor/Private/PropertyNode.h
--- a/Engine/Source/Editor/PropertyEditor/Private/PropertyNode.h
+++ b/Engine/Source/Editor/PropertyEditor/Private/PropertyNode.h
@@ -428,6 +428,12 @@ public:
 	FObjectPropertyNode* FindObjectItemParent();
 	const FObjectPropertyNode* FindObjectItemParent() const;
 
+	/**
+	 * Follows the chain of items upwards until it finds the structure property that houses this item.
+	 */
+	FStructurePropertyNode* FindStructureItemParent();
+	const FStructurePropertyNode* FindStructureItemParent() const;
+
 	/**
 	 * Follows the top-most object window that contains this property window item.
 	 */
@@ -596,7 +602,7 @@ public:
 	 *
 	 * @return	a pointer to a FProperty value or UObject.  (For dynamic arrays, you'd cast this value to an FArray*)
 	 */
-	virtual uint8* GetValueBaseAddress(uint8* StartAddress, bool bIsSparseData) const;
+	virtual uint8* GetValueBaseAddress(uint8* StartAddress, bool bIsSparseData, bool bIsStruct = false) const;
 
 	/**
 	 * Calculates the memory address for the data associated with this item's value.  For most properties, identical to GetValueBaseAddress.  For items corresponding
@@ -607,7 +613,7 @@ public:
 	 *
 	 * @return	a pointer to a FProperty value or UObject.  (For dynamic arrays, you'd cast this value to whatever type is the Inner for the dynamic array)
 	 */
-	virtual uint8* GetValueAddress(uint8* StartAddress, bool bIsSparseData) const;
+	virtual uint8* GetValueAddress(uint8* StartAddress, bool bIsSparseData, bool bIsStruct = false) const;
 
 	/**
 	 * Caclulates the memory address for the starting point of the structure that contains the property this node uses.
@@ -1103,11 +1109,11 @@ protected:
 	/** @return		The property stored at this node, to be passed to Pre/PostEditChange. */
 	FProperty*		GetStoredProperty()		{ return nullptr; }
 
+	bool GetDiffersFromDefault(const uint8* PropertyValueAddress, const uint8* PropertyDefaultAddress, const uint8* DefaultPropertyValueBaseAddress, const FProperty* InProperty) const;
 	bool GetDiffersFromDefaultForObject( FPropertyItemValueDataTrackerSlate& ValueTracker, FProperty* InProperty );
 
+	FString GetDefaultValueAsString(const uint8* PropertyDefaultAddress, const FProperty* InProperty, const bool bUseDisplayName) const;
 	FString GetDefaultValueAsStringForObject( FPropertyItemValueDataTrackerSlate& ValueTracker, UObject* InObject, FProperty* InProperty, bool bUseDisplayName );
-
-	FString GetDefaultValueAsStringForObject( FPropertyItemValueDataTrackerSlate& ValueTracker, UObject* InObject, FProperty* InProperty );
 	
 	/**
 	 * Helper function to obtain the display name for an enum property
diff --git a/Engine/Source/Editor/PropertyEditor/Private/PropertyRowGenerator.cpp b/Engine/Source/Editor/PropertyEditor/Private/PropertyRowGenerator.cpp
--- a/Engine/Source/Editor/PropertyEditor/Private/PropertyRowGenerator.cpp
+++ b/Engine/Source/Editor/PropertyEditor/Private/PropertyRowGenerator.cpp
@@ -210,6 +210,16 @@ void FPropertyRowGenerator::SetStructure(const TSharedPtr<FStructOnScope>& InStr
 	PostSetObject();
 }
 
+void FPropertyRowGenerator::SetStructure(const TSharedPtr<IStructureDataProvider>& InStructProvider)
+{
+	PreSetObject(1, /*bHasStructRoots=*/true);
+
+	check(RootPropertyNodes.Num() == 1);
+	RootPropertyNodes[0]->AsStructureNode()->SetStructure(InStructProvider);
+
+	PostSetObject();
+}
+
 const TArray<TSharedRef<IDetailTreeNode>>& FPropertyRowGenerator::GetRootTreeNodes() const
 {
 	return RootTreeNodes;
@@ -420,7 +430,7 @@ bool FPropertyRowGenerator::IsPropertyEditingEnabled() const
 void FPropertyRowGenerator::ForceRefresh()
 {
 	TArray<UObject*> NewObjectList;
-	TSharedPtr<FStructOnScope> StructData = nullptr;
+	TSharedPtr<IStructureDataProvider> StructureProvider = nullptr;
 
 	for (const TSharedPtr<FComplexPropertyNode>& ComplexRootNode : RootPropertyNodes)
 	{
@@ -438,13 +448,13 @@ void FPropertyRowGenerator::ForceRefresh()
 		}
 		else if (FStructurePropertyNode* StructRootNode = ComplexRootNode->AsStructureNode())
 		{
-			StructData = StructRootNode->GetStructData();
+			StructureProvider = StructRootNode->GetStructProvider();
 		}
 	}
-	
-	if (StructData && StructData->IsValid())
+
+	if (StructureProvider)
 	{
-		SetStructure(StructData);
+		SetStructure(StructureProvider);
 	}
 	else
 	{
@@ -483,7 +493,7 @@ void FPropertyRowGenerator::PreSetObject(int32 NumNewObjects, bool bHasStructRoo
 		else
 		{
 			FStructurePropertyNode* RootStructNode = RootNode->AsStructureNode();
-			RootStructNode->SetStructure(nullptr);
+			RootStructNode->RemoveStructure();
 		}
 		RootNode->ClearCachedReadAddresses(true);
 	}
diff --git a/Engine/Source/Editor/PropertyEditor/Private/PropertyRowGenerator.h b/Engine/Source/Editor/PropertyEditor/Private/PropertyRowGenerator.h
--- a/Engine/Source/Editor/PropertyEditor/Private/PropertyRowGenerator.h
+++ b/Engine/Source/Editor/PropertyEditor/Private/PropertyRowGenerator.h
@@ -20,6 +20,7 @@ class FComplexPropertyNode;
 class FDetailTreeNode;
 class IPropertyGenerationUtilities;
 class FStructOnScope;
+class IStructureDataProvider;
 
 struct FPropertyNodeMap
 {
@@ -82,6 +83,7 @@ public:
 	/** IPropertyRowGenerator interface */
 	virtual void SetObjects(const TArray<UObject*>& InObjects) override;
 	virtual void SetStructure(const TSharedPtr<FStructOnScope>& InStruct) override;
+	virtual void SetStructure(const TSharedPtr<IStructureDataProvider>& InStructProvider) override;
 	virtual const TArray<TWeakObjectPtr<UObject>>& GetSelectedObjects() const override { return SelectedObjects; }
 	virtual const TArray<TSharedRef<IDetailTreeNode>>& GetRootTreeNodes() const override;
 	virtual TSharedPtr<IDetailTreeNode> FindTreeNode(TSharedPtr<IPropertyHandle> PropertyHandle) const override;
diff --git a/Engine/Source/Editor/PropertyEditor/Private/SDetailsView.cpp b/Engine/Source/Editor/PropertyEditor/Private/SDetailsView.cpp
--- a/Engine/Source/Editor/PropertyEditor/Private/SDetailsView.cpp
+++ b/Engine/Source/Editor/PropertyEditor/Private/SDetailsView.cpp
@@ -1000,9 +1000,11 @@ void SDetailsView::PostSetObject(const TArray<FDetailsViewObjectRoot>& Roots)
 		RootPropertyNode->InitNode( InitParams );
 	}
 
+	UpdatePropertyMaps();
+
+	// Restore expanded items after handling property customization as some items that require expansion come from the customizations.
 	RestoreAllExpandedItems();
 
-	UpdatePropertyMaps();
 	UpdateFilteredDetails();
 }
 
diff --git a/Engine/Source/Editor/PropertyEditor/Private/SStructureDetailsView.cpp b/Engine/Source/Editor/PropertyEditor/Private/SStructureDetailsView.cpp
--- a/Engine/Source/Editor/PropertyEditor/Private/SStructureDetailsView.cpp
+++ b/Engine/Source/Editor/PropertyEditor/Private/SStructureDetailsView.cpp
@@ -28,7 +28,7 @@ SStructureDetailsView::~SStructureDetailsView()
 
 UStruct* SStructureDetailsView::GetBaseScriptStruct() const
 {
-	const UStruct* Struct = StructData.IsValid() ? StructData->GetStruct() : NULL;
+	const UStruct* Struct = StructProvider.IsValid() ? StructProvider->GetBaseStructure() : nullptr;
 	return const_cast<UStruct*>(Struct);
 }
 
@@ -209,10 +209,16 @@ void SStructureDetailsView::Construct(const FArguments& InArgs)
 
 void SStructureDetailsView::SetStructureData(TSharedPtr<FStructOnScope> InStructData)
 {
+	SetStructureProvider(InStructData ? MakeShared<FStructOnScopeStructureDataProvider>(InStructData) : TSharedPtr<IStructureDataProvider>());
+}
+
+void SStructureDetailsView::SetStructureProvider(TSharedPtr<IStructureDataProvider> InStructProvider)
+{
+
 	TSharedPtr<FComplexPropertyNode> RootNode = GetRootNode();
 	//PRE SET
 	SaveExpandedItems(RootNode.ToSharedRef() );
-	RootNode->AsStructureNode()->SetStructure(nullptr);
+	RootNode->AsStructureNode()->RemoveStructure();
 	RootNodesPendingKill.Add(RootNode);
 
 	RootNodes.Empty(1);
@@ -222,9 +228,9 @@ void SStructureDetailsView::SetStructureData(TSharedPtr<FStructOnScope> InStruct
 	RootNodes.Add(RootNode);
 
 	//SET
-	StructData = InStructData;
-	RootNode->AsStructureNode()->SetStructure(StructData);
-	if (!StructData.IsValid())
+	StructProvider = InStructProvider;
+	RootNode->AsStructureNode()->SetStructure(StructProvider);
+	if (!StructProvider.IsValid())
 	{
 		bIsLocked = false;
 	}
@@ -240,8 +246,8 @@ void SStructureDetailsView::SetStructureData(TSharedPtr<FStructOnScope> InStruct
 	}
 
 	FPropertyNodeInitParams InitParams;
-	InitParams.ParentNode = NULL;
-	InitParams.Property = NULL;
+	InitParams.ParentNode = nullptr;
+	InitParams.Property = nullptr;
 	InitParams.ArrayOffset = 0;
 	InitParams.ArrayIndex = INDEX_NONE;
 	InitParams.bAllowChildren = true;
@@ -265,7 +271,7 @@ void SStructureDetailsView::SetCustomName(const FText& Text)
 
 void SStructureDetailsView::ForceRefresh()
 {
-	SetStructureData(StructData);
+	SetStructureProvider(StructProvider);
 }
 
 void SStructureDetailsView::ClearSearch()
@@ -297,7 +303,7 @@ const FSelectedActorInfo& SStructureDetailsView::GetSelectedActorInfo() const
 bool SStructureDetailsView::IsConnected() const
 {
 	const FStructurePropertyNode* RootNode = GetRootNode().IsValid() ? GetRootNode()->AsStructureNode() : nullptr;
-	return StructData.IsValid() && StructData->IsValid() && RootNode && RootNode->HasValidStructData();
+	return StructProvider.IsValid() && StructProvider->IsValid() && RootNode && RootNode->HasValidStructData();
 }
 
 FRootPropertyNodeList& SStructureDetailsView::GetRootNodes()
@@ -319,13 +325,14 @@ void SStructureDetailsView::CustomUpdatePropertyMap(TSharedPtr<FDetailLayoutBuil
 {
 	FName StructCategoryName = NAME_None;
 	
-	if (StructData.IsValid() && StructData->GetStruct())
+	const UStruct* Struct = StructProvider.IsValid() ? StructProvider->GetBaseStructure() : nullptr;
+	if (Struct)
 	{
 		TArray<FName> CategoryNames;
 		InDetailLayout->GetCategoryNames(CategoryNames);
 
 		int32 StructCategoryNameIndex = INDEX_NONE;
-		CategoryNames.Find(StructData->GetStruct()->GetFName(), StructCategoryNameIndex);
+		CategoryNames.Find(Struct->GetFName(), StructCategoryNameIndex);
 		StructCategoryName = StructCategoryNameIndex != INDEX_NONE ? CategoryNames[StructCategoryNameIndex] : NAME_None;
 	}
 	
@@ -335,7 +342,7 @@ void SStructureDetailsView::CustomUpdatePropertyMap(TSharedPtr<FDetailLayoutBuil
 EVisibility SStructureDetailsView::GetPropertyEditingVisibility() const
 {
 	const FStructurePropertyNode* RootNode = GetRootNode().IsValid() ? GetRootNode()->AsStructureNode() : nullptr;
-	return StructData.IsValid() && StructData->IsValid() && RootNode && RootNode->HasValidStructData() ? EVisibility::Visible : EVisibility::Collapsed;
+	return StructProvider.IsValid() && StructProvider->IsValid() && RootNode && RootNode->HasValidStructData() ? EVisibility::Visible : EVisibility::Collapsed;
 }
 
 #undef LOCTEXT_NAMESPACE
diff --git a/Engine/Source/Editor/PropertyEditor/Private/SStructureDetailsView.h b/Engine/Source/Editor/PropertyEditor/Private/SStructureDetailsView.h
--- a/Engine/Source/Editor/PropertyEditor/Private/SStructureDetailsView.h
+++ b/Engine/Source/Editor/PropertyEditor/Private/SStructureDetailsView.h
@@ -13,6 +13,7 @@
 class AActor;
 class FDetailLayoutBuilderImpl;
 class IDetailRootObjectCustomization;
+class IStructureDataProvider;
 
 class SStructureDetailsView : public SDetailsViewBase, public IStructureDetailsView
 {
@@ -65,6 +66,8 @@ public:
 
 	virtual void SetStructureData(TSharedPtr<FStructOnScope> InStructData) override;
 
+	virtual void SetStructureProvider(TSharedPtr<IStructureDataProvider> StructProvider) override;
+
 	virtual FOnFinishedChangingProperties& GetOnFinishedChangingPropertiesDelegate() override
 	{
 		return OnFinishedChangingProperties();
@@ -116,7 +119,7 @@ protected:
 	EVisibility GetPropertyEditingVisibility() const;
 
 private:
-	TSharedPtr<class FStructOnScope> StructData;
+	TSharedPtr<IStructureDataProvider> StructProvider;
 	FRootPropertyNodeList RootNodes;
 	FText CustomName;
 };
diff --git a/Engine/Source/Editor/PropertyEditor/Private/StructurePropertyNode.cpp b/Engine/Source/Editor/PropertyEditor/Private/StructurePropertyNode.cpp
--- a/Engine/Source/Editor/PropertyEditor/Private/StructurePropertyNode.cpp
+++ b/Engine/Source/Editor/PropertyEditor/Private/StructurePropertyNode.cpp
@@ -10,7 +10,7 @@ void FStructurePropertyNode::InitChildNodes()
 	const bool bShouldShowHiddenProperties = !!HasNodeFlags(EPropertyNodeFlags::ShouldShowHiddenProperties);
 	const bool bShouldShowDisableEditOnInstance = !!HasNodeFlags(EPropertyNodeFlags::ShouldShowDisableEditOnInstance);
 
-	const UStruct* Struct = StructData.IsValid() ? StructData->GetStruct() : nullptr;
+	const UStruct* Struct = GetBaseStructure();
 
 	TArray<FProperty*> StructMembers;
 
@@ -27,7 +27,7 @@ void FStructurePropertyNode::InitChildNodes()
 
 	for (FProperty* StructMember : StructMembers)
 	{
-		TSharedPtr<FItemPropertyNode> NewItemNode(new FItemPropertyNode);//;//CreatePropertyItem(StructMember,INDEX_NONE,this);
+		TSharedPtr<FItemPropertyNode> NewItemNode(new FItemPropertyNode);
 
 		FPropertyNodeInitParams InitParams;
 		InitParams.ParentNode = SharedThis(this);
@@ -43,3 +43,50 @@ void FStructurePropertyNode::InitChildNodes()
 		AddChildNode(NewItemNode);
 	}
 }
+
+uint8* FStructurePropertyNode::GetValueBaseAddress(uint8* StartAddress, bool bIsSparseData, bool bIsStruct) const
+{
+	// If called with struct data, we expect that it is compatible with the first structure node down in the property node chain, return the address as is.
+	// This gets called usually when the calling code is dealing with a parent complex node.
+	if (bIsStruct)
+	{
+		return StartAddress;
+	}
+
+	if (StructProvider)
+	{
+		// Assume that this code gets called with an object or object sparse data.
+		//
+		// The passed object might not be the one that contains the values provided by struct provider.
+		// For example this function might get called on an edited object's template object.
+		// In that case the data structure is expected to match between the data edited by this node and the foreign object.
+
+		// If the structure provider is set up as indirection, then it knows how to translate parent node's value address to
+		// new value address even on data that is not the same as in the structure provider.
+		if (StructProvider->IsPropertyIndirection())
+		{
+			const TSharedPtr<FPropertyNode> ParentNode = ParentNodeWeakPtr.Pin();
+			if (!ensureMsgf(ParentNode, TEXT("Expecting valid parent node when indirection structure provider is called with Object data.")))
+			{
+				return nullptr;				
+			}
+			// Resolve from parent nodes data.
+			uint8* ParentValueAddress = ParentNode->GetValueAddress(StartAddress, bIsSparseData);
+			uint8* ValueAddress = StructProvider->GetValueBaseAddress(ParentValueAddress, GetBaseStructure());
+			return ValueAddress;
+		}
+
+		// The struct is really standalone, in which case we always return the standalone struct data.
+		// In that case we can only support one instance, since we cannot discern them.
+		// Note: Multiple standalone structure instances are supported when bIsStruct is true (e.g. when the structure property is root node).
+		TArray<TSharedPtr<FStructOnScope>> Instances;
+		StructProvider->GetInstances(Instances);
+		ensureMsgf(Instances.Num() <= 1, TEXT("Expecting max one instance on standalone structure provider."));
+		if (Instances.Num() == 1 && Instances[0].IsValid())
+		{
+			return Instances[0]->GetStructMemory();
+		}
+	}
+	
+	return nullptr;
+}
diff --git a/Engine/Source/Editor/PropertyEditor/Private/StructurePropertyNode.h b/Engine/Source/Editor/PropertyEditor/Private/StructurePropertyNode.h
--- a/Engine/Source/Editor/PropertyEditor/Private/StructurePropertyNode.h
+++ b/Engine/Source/Editor/PropertyEditor/Private/StructurePropertyNode.h
@@ -5,6 +5,7 @@
 #include "CoreMinimal.h"
 #include "PropertyNode.h"
 #include "UObject/StructOnScope.h"
+#include "IStructureDataProvider.h"
 
 //-----------------------------------------------------------------------------
 //	FStructPropertyNode - Used for the root and various sub-nodes
@@ -14,26 +15,65 @@ class FStructurePropertyNode : public FComplexPropertyNode
 {
 public:
 	FStructurePropertyNode() : FComplexPropertyNode() {}
-	virtual ~FStructurePropertyNode() {}
+	virtual ~FStructurePropertyNode() override {}
 
 	virtual FStructurePropertyNode* AsStructureNode() override { return this; }
 	virtual const FStructurePropertyNode* AsStructureNode() const override { return this; }
 
-	void SetStructure(TSharedPtr<FStructOnScope> InStructData)
+	void RemoveStructure()
 	{
 		ClearCachedReadAddresses(true);
 		DestroyTree();
-		StructData = InStructData;
+		StructProvider = nullptr;
+	}
+
+	void SetStructure(TSharedPtr<FStructOnScope> InStructData)
+	{
+		RemoveStructure();
+		if (InStructData)
+		{
+			StructProvider = MakeShared<FStructOnScopeStructureDataProvider>(InStructData);
+		}
+	}
+
+	void SetStructure(TSharedPtr<IStructureDataProvider> InStructProvider)
+	{
+		RemoveStructure();
+		StructProvider = InStructProvider;
 	}
 
 	bool HasValidStructData() const
 	{
-		return StructData.IsValid() && StructData->IsValid();
+		return StructProvider.IsValid() && StructProvider->IsValid();
 	}
 
+	// Returns just the first structure. Please use GetStructProvider() or GetAllStructureData() when dealing with multiple struct instances.
 	TSharedPtr<FStructOnScope> GetStructData() const
 	{
-		return StructData;
+		if (StructProvider)
+		{
+			TArray<TSharedPtr<FStructOnScope>> Instances;
+			StructProvider->GetInstances(Instances);
+			
+			if (Instances.Num() > 0)
+			{
+				return Instances[0];
+			}
+		}
+		return nullptr;
+	}
+
+	void GetAllStructureData(TArray<TSharedPtr<FStructOnScope>>& OutStructs) const
+	{
+		if (StructProvider)
+		{
+			StructProvider->GetInstances(OutStructs);
+		}
+	}
+
+	TSharedPtr<IStructureDataProvider> GetStructProvider() const
+	{
+		return StructProvider;
 	}
 
 	bool GetReadAddressUncached(const FPropertyNode& InPropertyNode, FReadAddressListData& OutAddresses) const override
@@ -42,6 +82,7 @@ public:
 		{
 			return false;
 		}
+		check(StructProvider.IsValid());
 
 		const FProperty* InItemProperty = InPropertyNode.GetProperty();
 		if (!InItemProperty)
@@ -56,10 +97,20 @@ public:
 			return false;
 		}
 
-		uint8* ReadAddress = StructData->GetStructMemory();
-		check(ReadAddress);
-		OutAddresses.Add(nullptr, InPropertyNode.GetValueBaseAddress(ReadAddress, InPropertyNode.HasNodeFlags(EPropertyNodeFlags::IsSparseClassData) != 0), true);
-		return true;
+		TArray<TSharedPtr<FStructOnScope>> Instances;
+		StructProvider->GetInstances(Instances);
+		bool bHasData = false;
+
+		for (TSharedPtr<FStructOnScope>& Instance : Instances)
+		{
+			uint8* ReadAddress = Instance.IsValid() ? Instance->GetStructMemory() : nullptr;
+			if (ReadAddress)
+			{
+				OutAddresses.Add(nullptr, InPropertyNode.GetValueBaseAddress(ReadAddress, InPropertyNode.HasNodeFlags(EPropertyNodeFlags::IsSparseClassData) != 0, /*bIsStruct=*/true), /*bIsStruct=*/true);
+				bHasData = true;
+			}
+		}
+		return bHasData;
 	}
 
 	bool GetReadAddressUncached(const FPropertyNode& InPropertyNode,
@@ -69,58 +120,197 @@ public:
 		bool bObjectForceCompare,
 		bool bArrayPropertiesCanDifferInSize) const override
 	{
-		if(OutAddresses)
+		if (!HasValidStructData())
+		{
+			return false;
+		}
+		check(StructProvider.IsValid());
+
+		const FProperty* InItemProperty = InPropertyNode.GetProperty();
+		if (!InItemProperty)
+		{
+			return false;
+		}
+
+		const UStruct* OwnerStruct = InItemProperty->GetOwnerStruct();
+		if (!OwnerStruct || OwnerStruct->IsStructTrashed())
+		{
+			// Verify that the property is not part of an invalid trash class
+			return false;
+		}
+
+		bool bAllTheSame = true;
+
+		TArray<TSharedPtr<FStructOnScope>> Instances;
+		StructProvider->GetInstances(Instances);
+		
+		if (Instances.IsEmpty())
+		{
+			return false;
+		}
+
+		if (bComparePropertyContents || bObjectForceCompare)
 		{
-			return GetReadAddressUncached(InPropertyNode, *OutAddresses);
+			const bool bIsSparse = InPropertyNode.HasNodeFlags(EPropertyNodeFlags::IsSparseClassData) != 0;
+			const uint8* BaseAddress = nullptr;
+			const UStruct* BaseStruct = nullptr;
+
+			for (TSharedPtr<FStructOnScope>& Instance : Instances)
+			{
+				if (Instance.IsValid())
+				{
+					if (const UStruct* Struct = Instance->GetStruct())
+					{
+						if (const uint8* ReadAddress = InPropertyNode.GetValueBaseAddress(Instance->GetStructMemory(), bIsSparse, /*bIsStruct=*/true))
+						{
+							if (!BaseAddress)
+							{
+								BaseAddress = ReadAddress;
+								BaseStruct = Struct;
+							}
+							else
+							{
+								if (BaseStruct != Struct)
+								{
+									bAllTheSame = false;
+									break;
+								}
+								if (!InItemProperty->Identical(BaseAddress, ReadAddress))
+								{
+									bAllTheSame = false;
+									break;
+								}
+							}
+						}
+					}
+				}
+			}
 		}
 		else
 		{
-			FReadAddressListData Unused;
-			return GetReadAddressUncached(InPropertyNode, Unused);
+			// Check that all are valid or invalid.
+			const UStruct* BaseStruct = Instances[0].IsValid() ? Instances[0]->GetStruct() : nullptr;
+			for (int32 Index = 1; Index < Instances.Num(); Index++)
+			{
+				const UStruct* Struct = Instances[Index].IsValid() ? Instances[Index]->GetStruct() : nullptr;
+				if (BaseStruct != Struct)
+				{
+					bAllTheSame = false;
+					break;
+				}
+			}
+		}
+
+		if (bAllTheSame && OutAddresses)
+		{
+			for (TSharedPtr<FStructOnScope>& Instance : Instances)
+			{
+				uint8* ReadAddress = Instance.IsValid() ? Instance->GetStructMemory() : nullptr;
+				if (ReadAddress)
+				{
+					OutAddresses->Add(nullptr, InPropertyNode.GetValueBaseAddress(ReadAddress, InPropertyNode.HasNodeFlags(EPropertyNodeFlags::IsSparseClassData) != 0, /*bIsStruct=*/true), /*bIsStruct=*/true);
+				}
+			}
 		}
+
+		return bAllTheSame;
 	}
 
-	UPackage* GetOwnerPackage() const
+	void GetOwnerPackages(TArray<UPackage*>& OutPackages) const
 	{
-		return HasValidStructData() ? StructData->GetPackage() : nullptr;
+		if (StructProvider)
+		{
+			TArray<TSharedPtr<FStructOnScope>> Instances;
+			StructProvider->GetInstances(Instances);
+
+			for (TSharedPtr<FStructOnScope>& Instance : Instances)
+			{
+				// Returning null for invalid instances, to match instance count.
+				OutPackages.Add(Instance.IsValid() ? Instance->GetPackage() : nullptr);
+			}
+		}
 	}
 
 	/** FComplexPropertyNode Interface */
 	virtual const UStruct* GetBaseStructure() const override
 	{ 
-		return HasValidStructData() ? StructData->GetStruct() : nullptr;
+		if (StructProvider)
+		{
+			return StructProvider->GetBaseStructure();
+		}
+		return nullptr; 
 	}
 	virtual UStruct* GetBaseStructure() override
 	{
-		const UStruct* Struct = HasValidStructData() ? StructData->GetStruct() : nullptr;
-		return const_cast<UStruct*>(Struct);
+		if (StructProvider)
+		{
+			return const_cast<UStruct*>(StructProvider->GetBaseStructure());
+		}
+		return nullptr; 
 	}
 	virtual TArray<UStruct*> GetAllStructures() override
 	{
 		TArray<UStruct*> RetVal;
-		if (UStruct* BaseStruct = GetBaseStructure())
+		if (StructProvider)
 		{
-			RetVal.Add(BaseStruct);
+			TArray<TSharedPtr<FStructOnScope>> Instances;
+			StructProvider->GetInstances(Instances);
+			
+			for (TSharedPtr<FStructOnScope>& Instance : Instances)
+			{
+				const UStruct* Struct = Instance.IsValid() ? Instance->GetStruct() : nullptr;
+				if (Struct)
+				{
+					RetVal.AddUnique(const_cast<UStruct*>(Struct));
+				}
+			}
 		}
+
 		return RetVal;
 	}
 	virtual TArray<const UStruct*> GetAllStructures() const override
 	{
 		TArray<const UStruct*> RetVal;
-		if (const UStruct* BaseStruct = GetBaseStructure())
+		if (StructProvider)
 		{
-			RetVal.Add(BaseStruct);
+			TArray<TSharedPtr<FStructOnScope>> Instances;
+			StructProvider->GetInstances(Instances);
+			
+			for (TSharedPtr<FStructOnScope>& Instance : Instances)
+			{
+				const UStruct* Struct = Instance.IsValid() ? Instance->GetStruct() : nullptr;
+				if (Struct)
+				{
+					RetVal.AddUnique(Struct);
+				}
+			}
 		}
 		return RetVal;
 	}
 	virtual int32 GetInstancesNum() const override
-	{ 
-		return HasValidStructData() ? 1 : 0;
+	{
+		if (StructProvider)
+		{
+			TArray<TSharedPtr<FStructOnScope>> Instances;
+			StructProvider->GetInstances(Instances);
+			
+			return Instances.Num();
+		}
+		return 0;
+		
 	}
 	virtual uint8* GetMemoryOfInstance(int32 Index) const override
-	{ 
-		check(0 == Index);
-		return HasValidStructData() ? StructData->GetStructMemory() : NULL;
+	{
+		if (StructProvider)
+		{
+			TArray<TSharedPtr<FStructOnScope>> Instances;
+			StructProvider->GetInstances(Instances);
+			if (Instances.IsValidIndex(Index) && Instances[Index].IsValid())
+			{
+				return Instances[Index]->GetStructMemory();
+			}
+		}
+		return nullptr;
 	}
 	virtual uint8* GetValuePtrOfInstance(int32 Index, const FProperty* InProperty, const FPropertyNode* InParentNode) const override
 	{ 
@@ -135,7 +325,7 @@ public:
 			return nullptr;
 		}
 
-		uint8* ParentBaseAddress = InParentNode->GetValueAddress(StructBaseAddress, false);
+		uint8* ParentBaseAddress = InParentNode->GetValueAddress(StructBaseAddress, false, /*bIsStruct=*/true);
 		if (ParentBaseAddress == nullptr)
 		{
 			return nullptr;
@@ -146,8 +336,7 @@ public:
 
 	virtual TWeakObjectPtr<UObject> GetInstanceAsUObject(int32 Index) const override
 	{
-		check(0 == Index);
-		return NULL;
+		return nullptr;
 	}
 	virtual EPropertyType GetPropertyType() const override
 	{
@@ -156,7 +345,9 @@ public:
 
 	virtual void Disconnect() override
 	{
-		SetStructure(NULL);
+		ClearCachedReadAddresses(true);
+		DestroyTree();
+		StructProvider = nullptr;
 	}
 
 protected:
@@ -164,19 +355,28 @@ protected:
 	/** FPropertyNode interface */
 	virtual void InitChildNodes() override;
 
-	virtual uint8* GetValueBaseAddress(uint8* Base, bool bIsSparseData) const override
-	{
-		check(bIsSparseData == false);
-		return HasValidStructData() ? StructData->GetStructMemory() : nullptr;
-	}
+	virtual uint8* GetValueBaseAddress(uint8* Base, bool bIsSparseData, bool bIsStruct) const override;
 
 	virtual bool GetQualifiedName(FString& PathPlusIndex, const bool bWithArrayIndex, const FPropertyNode* StopParent = nullptr, bool bIgnoreCategories = false) const override
 	{
+		bool bAddedAnything = false;
+		const TSharedPtr<FPropertyNode> ParentNode = ParentNodeWeakPtr.Pin();
+		if (ParentNode && StopParent != ParentNode.Get())
+		{
+			bAddedAnything = ParentNode->GetQualifiedName(PathPlusIndex, bWithArrayIndex, StopParent, bIgnoreCategories);
+		}
+
+		if (bAddedAnything)
+		{
+			PathPlusIndex += TEXT(".");
+		}
+
 		PathPlusIndex += TEXT("Struct");
-		return true;
+		bAddedAnything = true;
+
+		return bAddedAnything;
 	}
 
 private:
-	TSharedPtr<FStructOnScope> StructData;
+	TSharedPtr<IStructureDataProvider> StructProvider;
 };
-
diff --git a/Engine/Source/Editor/PropertyEditor/Public/IPropertyRowGenerator.h b/Engine/Source/Editor/PropertyEditor/Public/IPropertyRowGenerator.h
--- a/Engine/Source/Editor/PropertyEditor/Public/IPropertyRowGenerator.h
+++ b/Engine/Source/Editor/PropertyEditor/Public/IPropertyRowGenerator.h
@@ -8,7 +8,7 @@
 class IDetailTreeNode;
 class FComplexPropertyNode;
 class FStructOnScope;
-
+class IStructureDataProvider;
 
 
 DECLARE_DELEGATE_RetVal_OneParam(bool, FOnValidatePropertyRowGeneratorNodes, const FRootPropertyNodeList&)
@@ -57,6 +57,7 @@ public:
 	 * @param InStruct The structure used to generate rows from.
 	 */
 	virtual void SetStructure(const TSharedPtr<FStructOnScope>& InStruct) = 0;
+	virtual void SetStructure(const TSharedPtr<IStructureDataProvider>& InStruct) = 0;
 
 	/**
 	 * Get the list of objects that were used to generate detail tree nodes
diff --git a/Engine/Source/Editor/PropertyEditor/Public/IStructureDataProvider.h b/Engine/Source/Editor/PropertyEditor/Public/IStructureDataProvider.h
new file mode 100644
--- /dev/null
+++ b/Engine/Source/Editor/PropertyEditor/Public/IStructureDataProvider.h
@@ -0,0 +1,112 @@
+﻿// Copyright Epic Games, Inc. All Rights Reserved.
+#pragma once
+
+class FStructOnScope;
+
+//-----------------------------------------------------------------------------
+//	IStructureDataProvider - Used to provide struct data for a property node.
+//
+// The IsPropertyIndirection() and GetValueBaseAddress() allow a struct provider to expose data to be edited that
+// is not supported by the property system. When IsPropertyIndirection() returns true, GetValueBaseAddress() is called
+// relative to the parent property node to retrieve pointer to the data to be edited.
+// 
+// In the following example, a struct provider is added via ParentProperty property handle. In this case "ParentProperty"
+// is the parent property node, and the value passed into GetValueBaseAddress() of the provider is the value of the parent property.
+//
+//		TSharedPtr<IPropertyHandle> ParentProperty = ...; // E.g. handle to structure being customized
+//		...
+//		TSharedRef<FMyStructProvider> NewStructProvider = MakeShared<FMyStructProvider>(ParentProperty);
+//		TArray<TSharedPtr<IPropertyHandle>> ChildProperties = ParentProperty->AddChildStructure(NewStructProvider);
+// 
+// The GetValueBaseAddress() may get called on foreign objects too (e.g. object templates), and if property indirection is used
+// it should not rely on other values cached in the provider.
+// 
+//		class FMyStructProvider : public IStructureDataProvider
+//		{
+//		public:
+//			...
+//
+//			virtual bool IsPropertyIndirection() const override { return true; }
+//
+//			virtual uint8* GetValueBaseAddress(uint8* ParentValueAddress, const UStruct* ExpectedType) const override
+//			{
+//				if (ParentValueAddress)
+//				{
+//					// "ParentValueAddress" is pointer to a value of the type defined in "ParentProperty".
+//					FMyStruct& Value = *reinterpret_cast<FMyStruct*>(ParentValueAddress);
+//					// The return value should match the ExpectedType. The expected type should be the same as the GetBaseStructure() returned by this provider.
+//					return (Value.GetScriptStruct() == ExpectedType) ? Value.GetMemory() : nullptr;
+//				}
+//				return nullptr;
+//			}
+//		};
+// 
+// If IsPropertyIndirection() is false, there structures are expected to be standalone, and GetInstances() is used instead to retried the instance values.
+// When indirection is not used and the struct provider is part of an another instance, the provider should return only one instance, since there is no to associate the provided value
+// with an foreign object.
+// 
+// Multiple instances are supported when the provider is the root node (e.g. on Structure details view).
+// 
+//-----------------------------------------------------------------------------
+class IStructureDataProvider
+{
+public:
+	virtual ~IStructureDataProvider() {}
+
+	/** @return true if we have any data instances. */
+	virtual bool IsValid() const = 0;
+	
+	/** @return the most common struct of the instance data. This struct will be used to build the UI for the instances. */
+	virtual const UStruct* GetBaseStructure() const = 0;
+
+	/** @return instances to edit. Each provided struct should be compatible with base struct. */
+	virtual void GetInstances(TArray<TSharedPtr<FStructOnScope>>& OutInstances) const = 0;
+
+	/** @return true, if the provider supports handling indirections using GetValueBaseAddress(). */
+	virtual bool IsPropertyIndirection() const { return false; }
+
+	/**
+	 * Returns base address of provided struct based on parent property nodes value.
+	 * @param ParentValueAddress Value of the parent 
+	 * @param ExpectedType the type of struct that is expected to be returned (generally the same type as returned by GetBaseStructure())
+	 * @return value base address based on parent base address.
+	 */
+	virtual uint8* GetValueBaseAddress(uint8* ParentValueAddress, const UStruct* ExpectedType) const { return ParentValueAddress; }
+};
+
+
+//-----------------------------------------------------------------------------
+//	FStructOnScopeStructureDataProvider - Implementation of standalone struct that holds one value.
+//-----------------------------------------------------------------------------
+class FStructOnScopeStructureDataProvider : public IStructureDataProvider
+{
+public:
+	FStructOnScopeStructureDataProvider() = default;
+	FStructOnScopeStructureDataProvider(const TSharedPtr<FStructOnScope>& InStructData)
+		: StructData(InStructData)
+	{
+	}
+	
+	void SetStructData(const TSharedPtr<FStructOnScope>& InStructData)
+	{
+		StructData = InStructData;
+	}
+	
+	virtual bool IsValid() const override
+	{
+		return StructData.IsValid() && StructData->IsValid();
+	};
+	
+	virtual const UStruct* GetBaseStructure() const override
+	{
+		return StructData.IsValid() ? StructData->GetStruct() : nullptr;
+	}
+	
+	virtual void GetInstances(TArray<TSharedPtr<FStructOnScope>>& OutInstances) const override
+	{
+		OutInstances.Add(StructData);
+	}
+
+protected:
+	TSharedPtr<FStructOnScope> StructData;
+};
\ No newline at end of file
diff --git a/Engine/Source/Editor/PropertyEditor/Public/IStructureDetailsView.h b/Engine/Source/Editor/PropertyEditor/Public/IStructureDetailsView.h
--- a/Engine/Source/Editor/PropertyEditor/Public/IStructureDetailsView.h
+++ b/Engine/Source/Editor/PropertyEditor/Public/IStructureDetailsView.h
@@ -37,6 +37,13 @@ public:
 	 */
 	virtual void SetStructureData(TSharedPtr<class FStructOnScope> StructData) = 0;
 
+	/**
+	 * Set the structure data to be displayed in this view.
+	 *
+	 * @param StructProvider Provider of the the structure data to display.
+	 */
+	virtual void SetStructureProvider(TSharedPtr<class IStructureDataProvider> StructProvider) = 0;
+
 	/**
 	 * Set the CustomName used for display purposes.
 	 *
diff --git a/Engine/Source/Editor/PropertyEditor/Public/PropertyHandle.h b/Engine/Source/Editor/PropertyEditor/Public/PropertyHandle.h
--- a/Engine/Source/Editor/PropertyEditor/Public/PropertyHandle.h
+++ b/Engine/Source/Editor/PropertyEditor/Public/PropertyHandle.h
@@ -14,6 +14,7 @@ class FResetToDefaultOverride;
 class IPropertyHandleArray;
 class IPropertyHandleMap;
 class IPropertyHandleSet;
+class IStructureDataProvider;
 
 namespace EPropertyValueSetFlags
 {
@@ -676,6 +677,7 @@ public:
 	 * @return An array of interfaces to the properties that were added
 	 */
 	virtual TArray<TSharedPtr<IPropertyHandle>> AddChildStructure( TSharedRef<FStructOnScope> ChildStructure ) = 0;
+	virtual TArray<TSharedPtr<IPropertyHandle>> AddChildStructure( TSharedRef<IStructureDataProvider> ChildStructure ) = 0;
 
 	/**
 	 * Returns whether or not the property can be set to default
