[BugFix] No clear way to unhiglight a property
CL 22279933 introduced a regression to SDetailsViewBase::HighlightProperty that made it impossible to unhighlight properties.
udn: https://epicgames.lightning.force.com/lightning/r/Case/5004z00001nziOiAAI/view

diff --git a/Engine/Source/Editor/PropertyEditor/Private/SDetailsViewBase.cpp b/Engine/Source/Editor/PropertyEditor/Private/SDetailsViewBase.cpp
--- a/Engine/Source/Editor/PropertyEditor/Private/SDetailsViewBase.cpp
+++ b/Engine/Source/Editor/PropertyEditor/Private/SDetailsViewBase.cpp
@@ -297,16 +297,17 @@ static TSharedPtr<FDetailTreeNode> FindBestFitTreeNodeFromProperty(const TArray<
 
 void SDetailsViewBase::HighlightProperty(const FPropertyPath& Property)
 {
-	if (!Property.IsValid())
-	{
-		return;
-	}
-	
 	TSharedPtr<FDetailTreeNode> PrevHighlightedNodePtr = CurrentlyHighlightedNode.Pin();
 	if (PrevHighlightedNodePtr.IsValid())
 	{
 		PrevHighlightedNodePtr->SetIsHighlighted(false);
 	}
+	
+	if (!Property.IsValid())
+	{
+		CurrentlyHighlightedNode = nullptr;
+		return;
+	}
 
 	const TSharedPtr< FDetailTreeNode > TreeNode = FindBestFitTreeNodeFromProperty(RootTreeNodes, Property);
 	if (TreeNode.IsValid())
