Fixed a bug where the details view would not properly show rows with changed visibility after a tree refresh

diff --git a/Engine/Source/Editor/PropertyEditor/Private/DetailCategoryBuilderImpl.cpp b/Engine/Source/Editor/PropertyEditor/Private/DetailCategoryBuilderImpl.cpp
--- a/Engine/Source/Editor/PropertyEditor/Private/DetailCategoryBuilderImpl.cpp
+++ b/Engine/Source/Editor/PropertyEditor/Private/DetailCategoryBuilderImpl.cpp
@@ -10,15 +10,12 @@
 #include "DetailItemNode.h"
 #include "DetailPropertyRow.h"
 #include "IPropertyGenerationUtilities.h"
-#include "ItemPropertyNode.h"
 #include "Misc/ConfigCacheIni.h"
 #include "Modules/ModuleManager.h"
 #include "ObjectPropertyNode.h"
-#include "PropertyCustomizationHelpers.h"
 #include "PropertyEditorModule.h"
 #include "PropertyPermissionList.h"
 #include "SDetailCategoryTableRow.h"
-#include "SDetailSingleItemRow.h"
 #include "StructurePropertyNode.h"
 #include "Styling/StyleColors.h"
  
@@ -646,6 +643,7 @@ void FDetailCategoryImpl::RefreshTree(bool bRefilterCategory)
 		TSharedPtr<FDetailLayoutBuilderImpl> ParentLayout = GetParentLayoutImpl();
 		if (ParentLayout.IsValid())
 		{
+			ParentLayout->RefreshNodeVisbility();
 			FilterNode(ParentLayout->GetCurrentFilter());
 			ParentLayout->GetPropertyGenerationUtilities().RebuildTreeNodes();
 		}
diff --git a/Engine/Source/Editor/PropertyEditor/Private/DetailItemNode.cpp b/Engine/Source/Editor/PropertyEditor/Private/DetailItemNode.cpp
--- a/Engine/Source/Editor/PropertyEditor/Private/DetailItemNode.cpp
+++ b/Engine/Source/Editor/PropertyEditor/Private/DetailItemNode.cpp
@@ -8,7 +8,6 @@
 #include "IDetailKeyframeHandler.h"
 #include "ObjectPropertyNode.h"
 #include "SConstrainedBox.h"
-#include "PropertyCustomizationHelpers.h"
 #include "PropertyPermissionList.h"
 #include "SDetailCategoryTableRow.h"
 #include "SDetailSingleItemRow.h"
@@ -60,7 +59,7 @@ void FDetailItemNode::Initialize()
 		SetExpansionState(bShouldExpand, bSaveState);
 	}
 
-	CachedItemVisibility = ComputeItemVisibility();
+	RefreshCachedVisibility();
 
 	const bool bUpdateFilteredNodes = false;
 	GenerateChildren( bUpdateFilteredNodes );
@@ -717,16 +716,7 @@ void FDetailItemNode::Tick( float DeltaTime )
 			Customization.CustomBuilderRow->Tick( DeltaTime );
 		}
 
-		// Recache visibility
-		EVisibility NewVisibility = ComputeItemVisibility();
-	
-		if( CachedItemVisibility != NewVisibility )
-		{
-			// The visibility of a node in the tree has changed.  We must refresh the tree to remove the widget
-			CachedItemVisibility = NewVisibility;
-			const bool bRefilterCategory = true;
-			ParentCategory.Pin()->RefreshTree( bRefilterCategory );
-		}
+		RefreshCachedVisibility(true);
 	}
 }
 
@@ -793,6 +783,28 @@ EVisibility FDetailItemNode::ComputeItemVisibility() const
 	return NewVisibility;
 }
 
+void FDetailItemNode::RefreshVisibility()
+{
+	RefreshCachedVisibility();
+}
+
+void FDetailItemNode::RefreshCachedVisibility(bool bCallChangeDelegate)
+{
+	// Recache visibility
+	EVisibility NewVisibility = ComputeItemVisibility();
+	
+	if( CachedItemVisibility != NewVisibility )
+	{
+		// The visibility of a node in the tree has changed.  We must refresh the tree to remove the widget
+		CachedItemVisibility = NewVisibility;
+		if (bCallChangeDelegate)
+		{
+			const bool bRefilterCategory = true;
+			ParentCategory.Pin()->RefreshTree( bRefilterCategory );
+		}
+	}
+}
+
 bool FDetailItemNode::ShouldShowOnlyChildren() const
 {
 	// Show only children of this node if there is no content for custom details or the property node requests that only children be shown
diff --git a/Engine/Source/Editor/PropertyEditor/Private/DetailItemNode.h b/Engine/Source/Editor/PropertyEditor/Private/DetailItemNode.h
--- a/Engine/Source/Editor/PropertyEditor/Private/DetailItemNode.h
+++ b/Engine/Source/Editor/PropertyEditor/Private/DetailItemNode.h
@@ -26,6 +26,7 @@ public:
 	virtual TSharedPtr<IPropertyHandle> CreatePropertyHandle() const override;
 	virtual void GetFilterStrings(TArray<FString>& OutFilterStrings) const override;
 	virtual bool GetInitiallyCollapsed() const override;
+	virtual void RefreshVisibility() override;
 
 	/**
 	 * Initializes this node                                                              
@@ -58,6 +59,11 @@ public:
 	 */
 	EVisibility ComputeItemVisibility() const;
 
+	/**
+	 * updates the cached node visibility and optionally calls a tree refresh if it changed
+	 */
+	void RefreshCachedVisibility(bool bCallChangeDelegate = false);
+
 	/** FDetailTreeNode interface */
 	virtual IDetailsView* GetNodeDetailsView() const override { TSharedPtr<FDetailCategoryImpl> PC = GetParentCategory(); return PC.IsValid() ? PC->GetNodeDetailsView() : nullptr; }
 	virtual IDetailsViewPrivate* GetDetailsView() const override{ TSharedPtr<FDetailCategoryImpl> PC = GetParentCategory(); return PC.IsValid() ? PC->GetDetailsView() : nullptr; }
diff --git a/Engine/Source/Editor/PropertyEditor/Private/DetailLayoutBuilderImpl.cpp b/Engine/Source/Editor/PropertyEditor/Private/DetailLayoutBuilderImpl.cpp
--- a/Engine/Source/Editor/PropertyEditor/Private/DetailLayoutBuilderImpl.cpp
+++ b/Engine/Source/Editor/PropertyEditor/Private/DetailLayoutBuilderImpl.cpp
@@ -728,6 +728,14 @@ FCustomPropertyTypeLayoutMap FDetailLayoutBuilderImpl::GetInstancedPropertyTypeL
 	return TypeLayoutMap;
 }
 
+void FDetailLayoutBuilderImpl::RefreshNodeVisbility()
+{
+	for(TSet<FDetailTreeNode*>::TIterator It = TickableNodes.CreateIterator(); It; ++It)
+	{
+		(*It)->RefreshVisibility();
+	}
+}
+
 TSharedPtr<FAssetThumbnailPool> FDetailLayoutBuilderImpl::GetThumbnailPool() const
 {
 	return PropertyDetailsUtilities.Pin()->GetThumbnailPool();
diff --git a/Engine/Source/Editor/PropertyEditor/Private/DetailLayoutBuilderImpl.h b/Engine/Source/Editor/PropertyEditor/Private/DetailLayoutBuilderImpl.h
--- a/Engine/Source/Editor/PropertyEditor/Private/DetailLayoutBuilderImpl.h
+++ b/Engine/Source/Editor/PropertyEditor/Private/DetailLayoutBuilderImpl.h
@@ -25,7 +25,7 @@ public:
 		const TSharedPtr<IDetailsViewPrivate>& InDetailsView,
 		bool bIsExternal);
 
-	~FDetailLayoutBuilderImpl();
+	virtual ~FDetailLayoutBuilderImpl() override;
 
 	/** IDetailLayoutBuilder Interface */
 	virtual const IDetailsView* GetDetailsView() const override;
@@ -218,6 +218,8 @@ public:
 	/** Combines the type layout map from our PropertyGenerationUtilities object, with the local InstancePropertyTypeExtensions map - which provides instance based customizations/overrides. */
 	FCustomPropertyTypeLayoutMap GetInstancedPropertyTypeLayoutMap() const;
 
+	void RefreshNodeVisbility();
+	
 private:
 	/**
 	 * Finds a property node for the current property by searching in a fast lookup map or a path search if required
diff --git a/Engine/Source/Editor/PropertyEditor/Public/DetailTreeNode.h b/Engine/Source/Editor/PropertyEditor/Public/DetailTreeNode.h
--- a/Engine/Source/Editor/PropertyEditor/Public/DetailTreeNode.h
+++ b/Engine/Source/Editor/PropertyEditor/Public/DetailTreeNode.h
@@ -111,6 +111,11 @@ public:
 	 */
 	virtual ENodeVisibility GetVisibility() const = 0;
 
+	/**
+	 * updates the visibility of this node if necessary
+	 */
+	virtual void RefreshVisibility() {}
+
 	/**
 	 * Called each frame if the node requests that it should be ticked
 	 */
