Disable the Copy Display Name menu option in the Details panel when the property has no display name that can be copied.

diff --git a/Engine/Source/Editor/PropertyEditor/Private/SDetailSingleItemRow.cpp b/Engine/Source/Editor/PropertyEditor/Private/SDetailSingleItemRow.cpp
--- a/Engine/Source/Editor/PropertyEditor/Private/SDetailSingleItemRow.cpp
+++ b/Engine/Source/Editor/PropertyEditor/Private/SDetailSingleItemRow.cpp
@@ -780,6 +780,7 @@ FSlateColor SDetailSingleItemRow::GetInnerBackgroundColor() const
 bool SDetailSingleItemRow::OnContextMenuOpening(FMenuBuilder& MenuBuilder)
 {
 	FUIAction CopyDisplayNameAction = FExecuteAction::CreateSP(this, &SDetailSingleItemRow::OnCopyPropertyDisplayName);
+	CopyDisplayNameAction.CanExecuteAction = FCanExecuteAction::CreateSP(this, &SDetailSingleItemRow::CanCopyPropertyDisplayName);
 	FUIAction CopyInternalNameAction = FExecuteAction::CreateSP(this, &SDetailSingleItemRow::OnCopyPropertyInternalName);
 	CopyInternalNameAction.CanExecuteAction = FCanExecuteAction::CreateSP(this, &SDetailSingleItemRow::CanCopyPropertyInternalName);
 
@@ -933,6 +934,23 @@ void SDetailSingleItemRow::OnCopyPropertyDisplayName()
 	}
 }
 
+bool SDetailSingleItemRow::CanCopyPropertyDisplayName()
+{
+	if (OwnerTreeNode.IsValid())
+	{
+		TSharedPtr<FPropertyNode> PropertyNode = GetPropertyNode();
+		if (PropertyNode.IsValid())
+		{
+			if (!PropertyNode->GetDisplayName().IsEmpty())
+			{
+				return true;
+			}
+				
+		}
+	}
+	return false;
+}
+
 void SDetailSingleItemRow::OnCopyPropertyInternalName()
 {
 	if (OwnerTreeNode.IsValid())
diff --git a/Engine/Source/Editor/PropertyEditor/Private/SDetailSingleItemRow.h b/Engine/Source/Editor/PropertyEditor/Private/SDetailSingleItemRow.h
--- a/Engine/Source/Editor/PropertyEditor/Private/SDetailSingleItemRow.h
+++ b/Engine/Source/Editor/PropertyEditor/Private/SDetailSingleItemRow.h
@@ -72,6 +72,7 @@ protected:
 private:
 	void OnCopyProperty();
 	void OnCopyPropertyDisplayName();
+	bool CanCopyPropertyDisplayName();
 	void OnCopyPropertyInternalName();
 	bool CanCopyPropertyInternalName();
 	void OnPasteProperty();
