Fix for instances of verse classes inserted into an array through the editor not getting correctly initialized. Very similar to the 'add element' fix made in CL21322034

diff --git a/Engine/Source/Editor/PropertyEditor/Private/PropertyHandleImpl.cpp b/Engine/Source/Editor/PropertyEditor/Private/PropertyHandleImpl.cpp
--- a/Engine/Source/Editor/PropertyEditor/Private/PropertyHandleImpl.cpp
+++ b/Engine/Source/Editor/PropertyEditor/Private/PropertyHandleImpl.cpp
@@ -1464,6 +1464,14 @@ void FPropertyValueImpl::InsertChild( TSharedPtr<FPropertyNode> ChildNodeToInser
 
 		ArrayHelper.InsertValues(Index, 1 );
 
+		// check whether the inner type is flagged as a non-nullable. if so, create it.
+		FObjectProperty* InnerObjectProperty = CastField<FObjectProperty>(ArrayProperty->Inner);
+		if (InnerObjectProperty && InnerObjectProperty->HasAnyPropertyFlags(CPF_NonNullable))
+		{
+			UObject* NewItem = NewObject<UObject>(Obj, InnerObjectProperty->PropertyClass);
+			InnerObjectProperty->SetObjectPropertyValue(ArrayHelper.GetRawPtr(Index), NewItem);
+		}
+
 		//set up indices for the coming events
 		TArray< TMap<FString,int32> > ArrayIndicesPerObject;
 		for (int32 ObjectIndex = 0; ObjectIndex < ReadAddresses.Num(); ++ObjectIndex)
