Added OpenComboBox function to SObjectPropertyEntryBox widget

diff --git a/Engine/Source/Editor/PropertyEditor/Private/PropertyCustomizationHelpers.cpp b/Engine/Source/Editor/PropertyEditor/Private/PropertyCustomizationHelpers.cpp
--- a/Engine/Source/Editor/PropertyEditor/Private/PropertyCustomizationHelpers.cpp
+++ b/Engine/Source/Editor/PropertyEditor/Private/PropertyCustomizationHelpers.cpp
@@ -547,6 +547,14 @@ void SObjectPropertyEntryBox::GetDesiredWidth(float& OutMinDesiredWidth, float &
 	PropertyEditorAsset->GetDesiredWidth(OutMinDesiredWidth, OutMaxDesiredWidth);
 }
 
+void SObjectPropertyEntryBox::OpenEntryBox()
+{
+	if (PropertyEditorAsset.IsValid())
+	{
+		PropertyEditorAsset->OpenComboButton();
+	}
+}
+
 FString SObjectPropertyEntryBox::OnGetObjectPath() const
 {
 	FString StringReference;
diff --git a/Engine/Source/Editor/PropertyEditor/Private/UserInterface/PropertyEditor/SPropertyEditorAsset.cpp b/Engine/Source/Editor/PropertyEditor/Private/UserInterface/PropertyEditor/SPropertyEditorAsset.cpp
--- a/Engine/Source/Editor/PropertyEditor/Private/UserInterface/PropertyEditor/SPropertyEditorAsset.cpp
+++ b/Engine/Source/Editor/PropertyEditor/Private/UserInterface/PropertyEditor/SPropertyEditorAsset.cpp
@@ -762,6 +762,14 @@ bool SPropertyEditorAsset::IsFilteredActor( const AActor* const Actor ) const
 	return IsAllowed;
 }
 
+void SPropertyEditorAsset::OpenComboButton()
+{
+	if (AssetComboButton.IsValid())
+	{
+		AssetComboButton->SetIsOpen(true);
+	}
+}
+
 void SPropertyEditorAsset::CloseComboButton()
 {
 	AssetComboButton->SetIsOpen(false);
diff --git a/Engine/Source/Editor/PropertyEditor/Private/UserInterface/PropertyEditor/SPropertyEditorAsset.h b/Engine/Source/Editor/PropertyEditor/Private/UserInterface/PropertyEditor/SPropertyEditorAsset.h
--- a/Engine/Source/Editor/PropertyEditor/Private/UserInterface/PropertyEditor/SPropertyEditorAsset.h
+++ b/Engine/Source/Editor/PropertyEditor/Private/UserInterface/PropertyEditor/SPropertyEditorAsset.h
@@ -100,6 +100,8 @@ public:
 
 	void GetDesiredWidth( float& OutMinDesiredWidth, float &OutMaxDesiredWidth );
 
+	void OpenComboButton();
+
 	/** Override the tick method so we can check for thumbnail differences & update if necessary */
 	virtual void Tick( const FGeometry& AllottedGeometry, const double InCurrentTime, const float InDeltaTime ) override;
 
diff --git a/Engine/Source/Editor/PropertyEditor/Public/PropertyCustomizationHelpers.h b/Engine/Source/Editor/PropertyEditor/Public/PropertyCustomizationHelpers.h
--- a/Engine/Source/Editor/PropertyEditor/Public/PropertyCustomizationHelpers.h
+++ b/Engine/Source/Editor/PropertyEditor/Public/PropertyCustomizationHelpers.h
@@ -229,6 +229,8 @@ public:
 
 	PROPERTYEDITOR_API void GetDesiredWidth(float& OutMinDesiredWidth, float &OutMaxDesiredWidth);
 
+	PROPERTYEDITOR_API void OpenEntryBox();
+
 private:
 	/**
 	 * Delegate function called when an object is changed
