Added support for overriding enum value display names on a per-property basis

This can be specified by the EnumValueDisplayNameOverrides meta-data, and is a semi-colon separated list of Name=DisplayName pairs where the DisplayName can be a raw string or a NS/LOCTEXT style localized value.

diff --git a/Engine/Source/Editor/PropertyEditor/Private/PropertyEditorHelpers.cpp b/Engine/Source/Editor/PropertyEditor/Private/PropertyEditorHelpers.cpp
--- a/Engine/Source/Editor/PropertyEditor/Private/PropertyEditorHelpers.cpp
+++ b/Engine/Source/Editor/PropertyEditor/Private/PropertyEditorHelpers.cpp
@@ -1121,6 +1121,32 @@ namespace PropertyEditorHelpers
 		return InvalidEnumValues;
 	}
 
+	TMap<FName, FText> GetEnumValueDisplayNamesFromPropertyOverride(const FProperty* Property, const UEnum* InEnum)
+	{
+		TMap<FName, FText> DisplayNameOverrides;
+
+		static const FName NAME_EnumValueDisplayNameOverrides = "EnumValueDisplayNameOverrides";
+
+		const FString& DisplayNameOverridesStr = Property->GetMetaData(NAME_EnumValueDisplayNameOverrides);
+		if (DisplayNameOverridesStr.Len() > 0)
+		{
+			TArray<FString> DisplayNameOverridePairs;
+			DisplayNameOverridesStr.ParseIntoArray(DisplayNameOverridePairs, TEXT(";"));
+
+			for (const FString& DisplayNameOverridePair : DisplayNameOverridePairs)
+			{
+				FString DisplayNameKey;
+				FString DisplayNameValue;
+				if (DisplayNameOverridePair.Split(TEXT("="), &DisplayNameKey, &DisplayNameValue))
+				{
+					DisplayNameOverrides.Add(*InEnum->GenerateFullEnumName(*DisplayNameKey), FTextStringHelper::CreateFromBuffer(*DisplayNameValue));
+				}
+			}
+		}
+
+		return DisplayNameOverrides;
+	}
+
 	bool IsCategoryHiddenByClass(const TSharedPtr<FComplexPropertyNode>& InRootNode, FName CategoryName)
 	{
 		return InRootNode->AsObjectNode() && InRootNode->AsObjectNode()->GetHiddenCategories().Contains(CategoryName);
diff --git a/Engine/Source/Editor/PropertyEditor/Private/PropertyEditorHelpers.h b/Engine/Source/Editor/PropertyEditor/Private/PropertyEditorHelpers.h
--- a/Engine/Source/Editor/PropertyEditor/Private/PropertyEditorHelpers.h
+++ b/Engine/Source/Editor/PropertyEditor/Private/PropertyEditorHelpers.h
@@ -261,6 +261,15 @@ namespace PropertyEditorHelpers
 	 * @return The array of disallowed enums.
 	 */
 	TArray<FName> GetInvalidEnumsFromPropertyOverride(const FProperty* Property, const UEnum* InEnum);
+
+	/**
+	 * Returns any enums that are have an overridden display name from the "EnumValueDisplayNameOverrides" metadata on FProperty using the specified enum.
+	 *
+	 * @param Property	The property which may contain the "EnumValueDisplayNameOverrides" metadata
+	 * @param InEnum	The enum to search
+	 * @return The map of display name overrides.
+	 */
+	TMap<FName, FText> GetEnumValueDisplayNamesFromPropertyOverride(const FProperty* Property, const UEnum* InEnum);
 	
 	/**
 	 * Whether or not a category is hidden by a given root object
diff --git a/Engine/Source/Editor/PropertyEditor/Private/PropertyHandleImpl.cpp b/Engine/Source/Editor/PropertyEditor/Private/PropertyHandleImpl.cpp
--- a/Engine/Source/Editor/PropertyEditor/Private/PropertyHandleImpl.cpp
+++ b/Engine/Source/Editor/PropertyEditor/Private/PropertyHandleImpl.cpp
@@ -3029,6 +3029,7 @@ bool FPropertyHandleBase::GeneratePossibleValues(TArray< TSharedPtr<FString> >&
 	{
 		const TArray<FName> ValidEnumValues = PropertyEditorHelpers::GetValidEnumsFromPropertyOverride(Property, Enum);
 		const TArray<FName> InvalidEnumValues = PropertyEditorHelpers::GetInvalidEnumsFromPropertyOverride(Property, Enum);
+		const TMap<FName, FText> EnumValueDisplayNameOverrides = PropertyEditorHelpers::GetEnumValueDisplayNamesFromPropertyOverride(Property, Enum);
 
 		//NumEnums() - 1, because the last item in an enum is the _MAX item
 		for( int32 EnumIndex = 0; EnumIndex < Enum->NumEnums() - 1; ++EnumIndex )
@@ -3057,7 +3058,11 @@ bool FPropertyHandleBase::GeneratePossibleValues(TArray< TSharedPtr<FString> >&
 			{
 				// See if we specified an alternate name for this value using metadata
 				FString EnumName = Enum->GetNameStringByIndex(EnumIndex);
-				FString EnumDisplayName = Enum->GetDisplayNameTextByIndex(EnumIndex).ToString();
+				FString EnumDisplayName = EnumValueDisplayNameOverrides.FindRef(Enum->GetNameByIndex(EnumIndex)).ToString();
+				if (EnumDisplayName.IsEmpty())
+				{
+					EnumDisplayName = Enum->GetDisplayNameTextByIndex(EnumIndex).ToString();
+				}
 
 				FText RestrictionTooltip;
 				const bool bIsRestricted = GenerateRestrictionToolTip(EnumName, RestrictionTooltip);
diff --git a/Engine/Source/Editor/PropertyEditor/Private/UserInterface/PropertyEditor/SPropertyEditorCombo.cpp b/Engine/Source/Editor/PropertyEditor/Private/UserInterface/PropertyEditor/SPropertyEditorCombo.cpp
--- a/Engine/Source/Editor/PropertyEditor/Private/UserInterface/PropertyEditor/SPropertyEditorCombo.cpp
+++ b/Engine/Source/Editor/PropertyEditor/Private/UserInterface/PropertyEditor/SPropertyEditorCombo.cpp
@@ -9,32 +9,6 @@
 
 #define LOCTEXT_NAMESPACE "PropertyEditor"
 
-static int32 FindEnumValueIndex(UEnum* Enum, FString const& ValueString)
-{
-	int32 Index = INDEX_NONE;
-	for(int32 ValIndex = 0; ValIndex < Enum->NumEnums(); ++ValIndex)
-	{
-		FString const EnumName    = Enum->GetNameStringByIndex(ValIndex);
-		FString const DisplayName = Enum->GetDisplayNameTextByIndex(ValIndex).ToString();
-
-		if (DisplayName.Len() > 0)
-		{
-			if (DisplayName == ValueString)
-			{
-				Index = ValIndex;
-				break;
-			}
-		}
-		
-		if (EnumName == ValueString)
-		{
-			Index = ValIndex;
-			break;
-		}
-	}
-	return Index;
-}
-
 void SPropertyEditorCombo::GetDesiredWidth( float& OutMinDesiredWidth, float& OutMaxDesiredWidth )
 {
 	OutMinDesiredWidth = 125.0f;
@@ -65,12 +39,9 @@ void SPropertyEditorCombo::Construct( const FArguments& InArgs, const TSharedPtr
 	PropertyEditor = InPropertyEditor;
 	ComboArgs = InArgs._ComboArgs;
 
-	TAttribute<FText> TooltipAttribute;
-
 	if (PropertyEditor.IsValid())
 	{
 		ComboArgs.PropertyHandle = PropertyEditor->GetPropertyHandle();
-		TooltipAttribute = TAttribute<FText>::Create(TAttribute<FText>::FGetter::CreateSP(PropertyEditor.ToSharedRef(), &FPropertyEditor::GetValueAsText));
 	}
 
 	ensureMsgf(ComboArgs.PropertyHandle.IsValid() || (ComboArgs.OnGetStrings.IsBound() && ComboArgs.OnGetValue.IsBound() && ComboArgs.OnValueSelected.IsBound()), TEXT("Either PropertyEditor or ComboArgs.PropertyHandle must be set!"));
@@ -94,7 +65,7 @@ void SPropertyEditorCombo::Construct( const FArguments& InArgs, const TSharedPtr
 		.OnSelectionChanged( this, &SPropertyEditorCombo::OnComboSelectionChanged )
 		.OnComboBoxOpening( this, &SPropertyEditorCombo::OnComboOpening )
 		.VisibleText( this, &SPropertyEditorCombo::GetDisplayValueAsString )
-		.ToolTipText( TooltipAttribute )
+		.ToolTipText( this, &SPropertyEditorCombo::GetValueToolTip )
 		.ShowSearchForItemCount( ComboArgs.ShowSearchForItemCount );
 
 	ChildSlot
@@ -103,6 +74,7 @@ void SPropertyEditorCombo::Construct( const FArguments& InArgs, const TSharedPtr
 	];
 
 	SetEnabled( TAttribute<bool>( this, &SPropertyEditorCombo::CanEdit ) );
+	SetToolTipText( TAttribute<FText>( this, &SPropertyEditorCombo::GetValueToolTip) );
 }
 
 FString SPropertyEditorCombo::GetDisplayValueAsString() const
@@ -111,24 +83,59 @@ FString SPropertyEditorCombo::GetDisplayValueAsString() const
 	{
 		return ComboArgs.OnGetValue.Execute();
 	}
-	else if (PropertyEditor.IsValid())
+	else if (bUsesAlternateDisplayValues)
 	{
-		return (bUsesAlternateDisplayValues) ? PropertyEditor->GetValueAsDisplayString() : PropertyEditor->GetValueAsString();
+		{
+			FString RawValueString;
+			ComboArgs.PropertyHandle->GetValueAsFormattedString(RawValueString, PPF_None);
+
+			if (const FString* AlternateDisplayValuePtr = InternalValueToAlternateDisplayValue.Find(RawValueString))
+			{
+				return *AlternateDisplayValuePtr;
+			}
+		}
+
+		if (PropertyEditor.IsValid())
+		{
+			return PropertyEditor->GetValueAsDisplayString();
+		}
+
+		FString ValueString;
+		ComboArgs.PropertyHandle->GetValueAsDisplayString(ValueString);
+		return ValueString;
 	}
 	else
 	{
-		FString ValueString;
-
-		if (bUsesAlternateDisplayValues)
+		if (PropertyEditor.IsValid())
 		{
-			ComboArgs.PropertyHandle->GetValueAsDisplayString(ValueString);
+			return PropertyEditor->GetValueAsString();
 		}
-		else
+
+		FString ValueString;
+		ComboArgs.PropertyHandle->GetValueAsFormattedString(ValueString);
+		return ValueString;
+	}
+}
+
+FText SPropertyEditorCombo::GetValueToolTip() const
+{
+	if (bUsesAlternateDisplayValues)
+	{
+		FString RawValueString;
+		ComboArgs.PropertyHandle->GetValueAsFormattedString(RawValueString, PPF_None);
+
+		if (const FString* AlternateDisplayValuePtr = InternalValueToAlternateDisplayValue.Find(RawValueString))
 		{
-			ComboArgs.PropertyHandle->GetValueAsFormattedString(ValueString);
+			return FText::AsCultureInvariant(*AlternateDisplayValuePtr);
 		}
-		return ValueString;
 	}
+
+	if (PropertyEditor.IsValid())
+	{
+		return PropertyEditor->GetValueAsText();
+	}
+
+	return FText();
 }
 
 void SPropertyEditorCombo::GenerateComboBoxStrings( TArray< TSharedPtr<FString> >& OutComboBoxStrings, TArray<TSharedPtr<SToolTip>>& RichToolTips, TArray<bool>& OutRestrictedItems )
@@ -164,6 +171,73 @@ void SPropertyEditorCombo::GenerateComboBoxStrings( TArray< TSharedPtr<FString>
 		}
 	}
 
+	// Build the reverse LUT for alternate display values
+	AlternateDisplayValueToInternalValue.Reset();
+	InternalValueToAlternateDisplayValue.Reset();
+	if (const FProperty* Property = ComboArgs.PropertyHandle->GetProperty();
+		bUsesAlternateDisplayValues && !Property->IsA(FStrProperty::StaticClass()))
+	{
+		// currently only enum properties can use alternate display values; this 
+		// might change, so assert here so that if support is expanded to other 
+		// property types without updating this block of code, we'll catch it quickly
+		const UEnum* Enum = nullptr;
+		if (const FByteProperty* ByteProperty = CastField<FByteProperty>(Property))
+		{
+			Enum = ByteProperty->Enum;
+		}
+		else if (const FEnumProperty* EnumProperty = CastField<FEnumProperty>(Property))
+		{
+			Enum = EnumProperty->GetEnum();
+		}
+		check(Enum != nullptr);
+
+		const TMap<FName, FText> EnumValueDisplayNameOverrides = PropertyEditorHelpers::GetEnumValueDisplayNamesFromPropertyOverride(Property, Enum);
+		auto FindEnumValueIndex = [&EnumValueDisplayNameOverrides, &Enum](const FString& ValueString) -> int32
+		{
+			for (const TTuple<FName, FText>& EnumValueDisplayNameOverridePair : EnumValueDisplayNameOverrides)
+			{
+				if (EnumValueDisplayNameOverridePair.Value.ToString() == ValueString)
+				{
+					return Enum->GetIndexByName(EnumValueDisplayNameOverridePair.Key);
+				}
+			}
+
+			for (int32 ValIndex = 0; ValIndex < Enum->NumEnums() - 1; ++ValIndex)
+			{
+				const FString EnumName = Enum->GetNameStringByIndex(ValIndex);
+				const FString DisplayName = Enum->GetDisplayNameTextByIndex(ValIndex).ToString();
+
+				if (DisplayName.Len() > 0)
+				{
+					if (DisplayName == ValueString)
+					{
+						return ValIndex;
+					}
+				}
+
+				if (EnumName == ValueString)
+				{
+					return ValIndex;
+				}
+			}
+
+			return INDEX_NONE;
+		};
+
+		for (const TSharedPtr<FString>& ValueStringPtr : OutComboBoxStrings)
+		{
+			const int32 EnumIndex = FindEnumValueIndex(*ValueStringPtr);
+			check(EnumIndex != INDEX_NONE);
+
+			const FString EnumValue = Enum->GetNameStringByIndex(EnumIndex);
+			if (EnumValue != *ValueStringPtr)
+			{
+				AlternateDisplayValueToInternalValue.Add(*ValueStringPtr, EnumValue);
+				InternalValueToAlternateDisplayValue.Add(EnumValue, *ValueStringPtr);
+			}
+		}
+	}
+
 	// For enums, look for rich tooltip information
 	if(ComboArgs.PropertyHandle.IsValid())
 	{
@@ -256,26 +330,12 @@ void SPropertyEditorCombo::SendToObjects( const FString& NewValue )
 	{
 		FProperty* Property = ComboArgs.PropertyHandle->GetProperty();
 
-		if (bUsesAlternateDisplayValues && !Property->IsA(FStrProperty::StaticClass()))
+		if (bUsesAlternateDisplayValues)
 		{
-			// currently only enum properties can use alternate display values; this 
-			// might change, so assert here so that if support is expanded to other 
-			// property types without updating this block of code, we'll catch it quickly
-			UEnum* Enum = nullptr;
-			if (FByteProperty* ByteProperty = CastField<FByteProperty>(Property))
-			{
-				Enum = ByteProperty->Enum;
-			}
-			else if (FEnumProperty* EnumProperty = CastField<FEnumProperty>(Property))
+			if (const FString* InternalValuePtr = AlternateDisplayValueToInternalValue.Find(Value))
 			{
-				Enum = EnumProperty->GetEnum();
+				Value = *InternalValuePtr;
 			}
-			check(Enum != nullptr);
-
-			const int32 Index = FindEnumValueIndex(Enum, NewValue);
-			check(Index != INDEX_NONE);
-
-			Value = Enum->GetNameStringByIndex(Index);
 		}
 
 		ComboArgs.PropertyHandle->SetValueFromFormattedString(Value);
diff --git a/Engine/Source/Editor/PropertyEditor/Private/UserInterface/PropertyEditor/SPropertyEditorCombo.h b/Engine/Source/Editor/PropertyEditor/Private/UserInterface/PropertyEditor/SPropertyEditorCombo.h
--- a/Engine/Source/Editor/PropertyEditor/Private/UserInterface/PropertyEditor/SPropertyEditorCombo.h
+++ b/Engine/Source/Editor/PropertyEditor/Private/UserInterface/PropertyEditor/SPropertyEditorCombo.h
@@ -39,6 +39,11 @@ private:
 	 */
 	FString GetDisplayValueAsString() const;
 
+	/**
+	 * Get the active display value tooltip.
+	 */
+	FText GetValueToolTip() const;
+
 	/** @return True if the property can be edited */
 	bool CanEdit() const;
 private:
@@ -56,4 +61,6 @@ private:
 	 * Indicates that this combo box's values are friendly names for the real values; currently only used for enum drop-downs.
 	 */
 	bool bUsesAlternateDisplayValues;
+	TMap<FString, FString> AlternateDisplayValueToInternalValue;
+	TMap<FString, FString> InternalValueToAlternateDisplayValue;
 };
diff --git a/Engine/Source/Editor/PropertyEditor/Private/UserInterface/PropertyEditor/SPropertyEditorNumeric.h b/Engine/Source/Editor/PropertyEditor/Private/UserInterface/PropertyEditor/SPropertyEditorNumeric.h
--- a/Engine/Source/Editor/PropertyEditor/Private/UserInterface/PropertyEditor/SPropertyEditorNumeric.h
+++ b/Engine/Source/Editor/PropertyEditor/Private/UserInterface/PropertyEditor/SPropertyEditorNumeric.h
@@ -76,14 +76,20 @@ public:
 
 				if (BitmaskEnum)
 				{
+					const TMap<FName, FText> EnumValueDisplayNameOverrides = PropertyEditorHelpers::GetEnumValueDisplayNamesFromPropertyOverride(PropertyHandle->GetProperty(), BitmaskEnum);
+
 					const bool bUseEnumValuessAsMaskValues = BitmaskEnum->GetBoolMetaData(PropertyEditorConstants::MD_UseEnumValuesAsMaskValuesInEditor);
-					auto AddNewBitmaskFlagLambda = [BitmaskEnum, &Result](int32 InEnumIndex, int64 InFlagValue)
+					auto AddNewBitmaskFlagLambda = [BitmaskEnum, &Result, &EnumValueDisplayNameOverrides](int32 InEnumIndex, int64 InFlagValue)
 					{
 						Result.Emplace();
 						FBitmaskFlagInfo* BitmaskFlag = &Result.Last();
 
 						BitmaskFlag->Value = InFlagValue;
-						BitmaskFlag->DisplayName = BitmaskEnum->GetDisplayNameTextByIndex(InEnumIndex);
+						BitmaskFlag->DisplayName = EnumValueDisplayNameOverrides.FindRef(BitmaskEnum->GetNameByIndex(InEnumIndex));
+						if (BitmaskFlag->DisplayName.IsEmpty())
+						{
+							BitmaskFlag->DisplayName = BitmaskEnum->GetDisplayNameTextByIndex(InEnumIndex);
+						}
 						BitmaskFlag->ToolTipText = BitmaskEnum->GetToolTipTextByIndex(InEnumIndex);
 						if (BitmaskFlag->ToolTipText.IsEmpty())
 						{
