Added EnsureDataIsValid() for structure property node
- Extended  the EnsureDataIsValid() to

diff --git a/Engine/Source/Editor/PropertyEditor/Private/PropertyNode.h b/Engine/Source/Editor/PropertyEditor/Private/PropertyNode.h
--- a/Engine/Source/Editor/PropertyEditor/Private/PropertyNode.h
+++ b/Engine/Source/Editor/PropertyEditor/Private/PropertyNode.h
@@ -442,7 +442,7 @@ public:
 	/**
 	 * Used to see if any data has been destroyed from under the property tree.  Should only be called during Tick
 	 */
-	EPropertyDataValidationResult EnsureDataIsValid();
+	virtual EPropertyDataValidationResult EnsureDataIsValid();
 
 	//////////////////////////////////////////////////////////////////////////
 	// Text
diff --git a/Engine/Source/Editor/PropertyEditor/Private/StructurePropertyNode.cpp b/Engine/Source/Editor/PropertyEditor/Private/StructurePropertyNode.cpp
--- a/Engine/Source/Editor/PropertyEditor/Private/StructurePropertyNode.cpp
+++ b/Engine/Source/Editor/PropertyEditor/Private/StructurePropertyNode.cpp
@@ -90,3 +90,32 @@ uint8* FStructurePropertyNode::GetValueBaseAddress(uint8* StartAddress, bool bIs
 	
 	return nullptr;
 }
+
+EPropertyDataValidationResult FStructurePropertyNode::EnsureDataIsValid()
+{
+	CachedReadAddresses.Reset();
+
+	const UStruct* BaseStruct = GetBaseStructure();
+	
+	// Check that struct node's children still belong to the current base struct.
+	for (const TSharedPtr<FPropertyNode>& ChildNode : ChildNodes)
+	{
+		if (ChildNode.IsValid())
+		{
+			if (const FProperty* ChildProperty = ChildNode->GetProperty())
+			{
+				const UStruct* OwnerStruct = ChildProperty->GetOwnerStruct();
+				if (!OwnerStruct
+					|| OwnerStruct->IsStructTrashed()
+					|| !BaseStruct
+					|| !BaseStruct->IsChildOf(OwnerStruct)) // OwnerStruct can be BaseStruct or one of structs BaseStruct is derived from.
+				{
+					RebuildChildren();
+					return EPropertyDataValidationResult::ChildrenRebuilt;
+				}
+			}
+		}
+	}
+	
+	return FPropertyNode::EnsureDataIsValid();
+}
diff --git a/Engine/Source/Editor/PropertyEditor/Private/StructurePropertyNode.h b/Engine/Source/Editor/PropertyEditor/Private/StructurePropertyNode.h
--- a/Engine/Source/Editor/PropertyEditor/Private/StructurePropertyNode.h
+++ b/Engine/Source/Editor/PropertyEditor/Private/StructurePropertyNode.h
@@ -352,6 +352,8 @@ public:
 
 protected:
 
+	virtual EPropertyDataValidationResult EnsureDataIsValid() override;
+
 	/** FPropertyNode interface */
 	virtual void InitChildNodes() override;
 
