Update tooltips for property editor combo boxes when the properties are reset to default.

diff --git a/Engine/Source/Editor/PropertyEditor/Private/UserInterface/PropertyEditor/SPropertyEditorCombo.cpp b/Engine/Source/Editor/PropertyEditor/Private/UserInterface/PropertyEditor/SPropertyEditorCombo.cpp
--- a/Engine/Source/Editor/PropertyEditor/Private/UserInterface/PropertyEditor/SPropertyEditorCombo.cpp
+++ b/Engine/Source/Editor/PropertyEditor/Private/UserInterface/PropertyEditor/SPropertyEditorCombo.cpp
@@ -42,6 +42,10 @@ void SPropertyEditorCombo::Construct( const FArguments& InArgs, const TSharedPtr
 	if (PropertyEditor.IsValid())
 	{
 		ComboArgs.PropertyHandle = PropertyEditor->GetPropertyHandle();
+		if (ComboArgs.PropertyHandle.IsValid())
+		{
+			ComboArgs.PropertyHandle->SetOnPropertyResetToDefault(FSimpleDelegate::CreateSP(this, &SPropertyEditorCombo::OnResetToDefault));
+		}
 	}
 
 	ensureMsgf(ComboArgs.PropertyHandle.IsValid() || (ComboArgs.OnGetStrings.IsBound() && ComboArgs.OnGetValue.IsBound() && ComboArgs.OnValueSelected.IsBound()), TEXT("Either PropertyEditor or ComboArgs.PropertyHandle must be set!"));
@@ -304,6 +308,12 @@ void SPropertyEditorCombo::OnComboSelectionChanged( TSharedPtr<FString> NewValue
 	}
 }
 
+void SPropertyEditorCombo::OnResetToDefault()
+{
+	FString CurrentDisplayValue = GetDisplayValueAsString();
+	ComboBox->SetSelectedItem(CurrentDisplayValue);
+}
+
 void SPropertyEditorCombo::OnComboOpening()
 {
 	TArray<TSharedPtr<FString>> ComboItems;
diff --git a/Engine/Source/Editor/PropertyEditor/Private/UserInterface/PropertyEditor/SPropertyEditorCombo.h b/Engine/Source/Editor/PropertyEditor/Private/UserInterface/PropertyEditor/SPropertyEditorCombo.h
--- a/Engine/Source/Editor/PropertyEditor/Private/UserInterface/PropertyEditor/SPropertyEditorCombo.h
+++ b/Engine/Source/Editor/PropertyEditor/Private/UserInterface/PropertyEditor/SPropertyEditorCombo.h
@@ -30,6 +30,7 @@ public:
 private:
 	void GenerateComboBoxStrings( TArray< TSharedPtr<FString> >& OutComboBoxStrings, TArray<TSharedPtr<class SToolTip>>& OutToolTips, TArray<bool>& OutRestrictedItems );
 	void OnComboSelectionChanged( TSharedPtr<FString> NewValue, ESelectInfo::Type SelectInfo );
+	void OnResetToDefault();
 	void OnComboOpening();
 
 	virtual void SendToObjects( const FString& NewValue );
