Make the property matrix respect the property editor permission list

diff --git a/Engine/Source/Editor/PropertyEditor/Private/Presentation/PropertyTable/PropertyTable.cpp b/Engine/Source/Editor/PropertyEditor/Private/Presentation/PropertyTable/PropertyTable.cpp
--- a/Engine/Source/Editor/PropertyEditor/Private/Presentation/PropertyTable/PropertyTable.cpp
+++ b/Engine/Source/Editor/PropertyEditor/Private/Presentation/PropertyTable/PropertyTable.cpp
@@ -10,6 +10,7 @@
 #include "Presentation/PropertyTable/PropertyTableRowHeaderColumn.h"
 #include "Presentation/PropertyTable/PropertyTableObjectNameColumn.h"
 #include "Presentation/PropertyTable/PropertyTablePropertyNameColumn.h"
+#include "PropertyPermissionList.h"
 
 #include "EditConditionParser.h"
 
@@ -1369,7 +1370,7 @@ void FPropertyTable::UpdateColumns()
 				TWeakFieldPtr< FProperty > Property = *PropertyIter;
 
 				// Don't expose CPF_NativeAccessSpecifierProtected and CPF_NativeAccessSpecifierPrivate properties
-				if ( PropertyIter->HasAllPropertyFlags(CPF_NativeAccessSpecifierPublic | CPF_AssetRegistrySearchable) )
+				if ( PropertyIter->HasAllPropertyFlags(CPF_NativeAccessSpecifierPublic | CPF_AssetRegistrySearchable) && FPropertyEditorPermissionList::Get().DoesPropertyPassFilter(PrimaryType.Get(), Property->GetFName()))
 				{
 					TArray< TSharedRef< IPropertyTableColumn > > MapResults;
 
diff --git a/Engine/Source/Editor/PropertyEditor/Private/SPropertyTreeViewImpl.cpp b/Engine/Source/Editor/PropertyEditor/Private/SPropertyTreeViewImpl.cpp
--- a/Engine/Source/Editor/PropertyEditor/Private/SPropertyTreeViewImpl.cpp
+++ b/Engine/Source/Editor/PropertyEditor/Private/SPropertyTreeViewImpl.cpp
@@ -20,6 +20,7 @@
 #include "UserInterface/PropertyTree/SPropertyTreeCategoryRow.h"
 #include "ScopedTransaction.h"
 #include "PropertyEditorHelpers.h"
+#include "PropertyPermissionList.h"
 #include "Misc/ConfigCacheIni.h"
 #include "Widgets/Colors/SColorPicker.h"
 #include "Widgets/Input/SSearchBox.h"
@@ -158,6 +159,14 @@ void SPropertyTreeViewImpl::Construct(const FArguments& InArgs)
 	PropertySettings = MakeShareable( new FPropertyUtilitiesTreeView(*this) );
 
 	ConstructPropertyTree();
+	
+	FPropertyEditorPermissionList::Get().PermissionListUpdatedDelegate.AddSP(this, &SPropertyTreeViewImpl::OnPermissionListUpdated);
+	FPropertyEditorPermissionList::Get().PermissionListEnabledDelegate.AddSP(this, &SPropertyTreeViewImpl::RequestRefresh);
+}
+
+void SPropertyTreeViewImpl::OnPermissionListUpdated(TSoftObjectPtr<UStruct>, FName)
+{
+	RequestRefresh();
 }
 
 /** Reconstructs the entire property tree widgets */
@@ -652,6 +661,59 @@ void SPropertyTreeViewImpl::MarkFavoritesInternal( TSharedPtr<FPropertyNode> InP
 	}
 }
 
+bool SPropertyTreeViewImpl::IsPropertyNodeVisible(TSharedPtr<FPropertyNode> InPropertyNode)
+{
+	bool bPropertyVisible = true;
+
+	// Object nodes always mark themselves as visible (but are never actually shown)
+	if(InPropertyNode->AsObjectNode())
+	{
+		return true;
+	}
+
+	// Category nodes are visible if they have any visible children
+	if(InPropertyNode->AsCategoryNode())
+	{
+		bool bHasVisibleChild = false;
+
+		// If we have a permission list in use, make sure the category has at least one child visible otherwise hide it
+		if(FPropertyEditorPermissionList::Get().IsEnabled())
+		{
+			for( int32 GrandChildIndex = 0; GrandChildIndex < InPropertyNode->GetNumChildNodes(); ++GrandChildIndex )
+			{
+				const TSharedPtr<FPropertyNode> GrandChildNode = InPropertyNode->GetChildNode( GrandChildIndex );
+
+				if(IsPropertyNodeVisible(GrandChildNode) && GrandChildNode->IsVisible())
+				{
+					bHasVisibleChild = true;
+					break;
+				}
+			}
+		}
+		// If we don't have a permission list, we can trivially check the child node count to see if there are any children
+		else
+		{
+			bHasVisibleChild = InPropertyNode->GetNumChildNodes() > 0;
+		}
+
+		return bHasVisibleChild;
+	}
+	
+	const FProperty* Property = InPropertyNode->GetProperty();
+	if(Property != NULL)
+	{
+		bPropertyVisible = FPropertyEditorPermissionList::Get().DoesPropertyPassFilter(FDetailTreeNode::GetPropertyNodeBaseStructure(InPropertyNode->GetParentNode()), Property->GetFName());
+
+		if(IsPropertyVisible.IsBound())
+		{
+			const FPropertyAndParent PropertyAndParent(InPropertyNode.ToSharedRef());
+			bPropertyVisible &= IsPropertyVisible.Execute(PropertyAndParent);
+		}
+	}
+
+	return bPropertyVisible;
+}
+
 void SPropertyTreeViewImpl::OnGetChildrenForPropertyNode( TSharedPtr<FPropertyNode> InPropertyNode, TArray< TSharedPtr<FPropertyNode> >& OutChildren )
 {
 	if( CurrentFilterText.Len() > 0 ) 
@@ -687,18 +749,13 @@ void SPropertyTreeViewImpl::OnGetChildrenForPropertyNode( TSharedPtr<FPropertyNo
 	for( int32 ChildIndex = 0; ChildIndex < InPropertyNode->GetNumChildNodes(); ++ChildIndex )
 	{
 		TSharedPtr<FPropertyNode> ChildNode = InPropertyNode->GetChildNode( ChildIndex );
-		FObjectPropertyNode* ObjNode = ChildNode->AsObjectNode();
-
-		bool bPropertyVisible = true;
-		FProperty* Property = ChildNode->GetProperty();
-		if(Property != NULL && IsPropertyVisible.IsBound())
-		{
-			FPropertyAndParent PropertyAndParent(ChildNode.ToSharedRef());
-			bPropertyVisible = IsPropertyVisible.Execute(PropertyAndParent);
-		}
 
+		const bool bPropertyVisible = IsPropertyNodeVisible(ChildNode);
+		
 		if(bPropertyVisible)
 		{
+			const FObjectPropertyNode* ObjNode = ChildNode->AsObjectNode();
+
 			if( ObjNode )
 			{
 				// Currently object property nodes do not provide any useful information other than being a container for its children.  We do not draw anything for them.
@@ -707,14 +764,7 @@ void SPropertyTreeViewImpl::OnGetChildrenForPropertyNode( TSharedPtr<FPropertyNo
 			}
 			else
 			{
-				if( ChildNode->IsVisible() )
-				{				
-					// Don't add empty category nodes
-					if( ChildNode->AsCategoryNode() == NULL || ChildNode->GetNumChildNodes() > 0 )
-					{
-						OutChildren.Add( ChildNode );
-					}
-				}
+				OutChildren.Add(ChildNode);
 			}
 		}
 	}
diff --git a/Engine/Source/Editor/PropertyEditor/Private/SPropertyTreeViewImpl.h b/Engine/Source/Editor/PropertyEditor/Private/SPropertyTreeViewImpl.h
--- a/Engine/Source/Editor/PropertyEditor/Private/SPropertyTreeViewImpl.h
+++ b/Engine/Source/Editor/PropertyEditor/Private/SPropertyTreeViewImpl.h
@@ -182,6 +182,9 @@ public:
 
 private:
 
+	/** Delegate called when the property editor permission list is updated */
+	void OnPermissionListUpdated(TSoftObjectPtr<UStruct>, FName);
+
 	/** Restores the expansion state in the tree for a specific object set */
 	void RestoreExpandedItems();
 
@@ -209,6 +212,9 @@ private:
 	 */
 	void OnGetChildrenForPropertyNode( TSharedPtr<FPropertyNode> InPropertyNode, TArray< TSharedPtr<FPropertyNode> >& OutChildren );
 
+	/** Check if a Property Node should be shown in the property tree if we have any external IsPropertyVisible delegates or permission lists */
+	bool IsPropertyNodeVisible(TSharedPtr<FPropertyNode> InPropertyNode);
+
 	/** 
 	 * Function called through a delegate on the favorites TreeView to request children of a property node  
 	 * 
diff --git a/Engine/Source/Editor/PropertyEditor/Private/UserInterface/PropertyTable/SPropertyTable.h b/Engine/Source/Editor/PropertyEditor/Private/UserInterface/PropertyTable/SPropertyTable.h
--- a/Engine/Source/Editor/PropertyEditor/Private/UserInterface/PropertyTable/SPropertyTable.h
+++ b/Engine/Source/Editor/PropertyEditor/Private/UserInterface/PropertyTable/SPropertyTable.h
@@ -11,6 +11,7 @@
 #include "Styling/AppStyle.h"
 #include "Widgets/Layout/SBorder.h"
 #include "PropertyPath.h"
+#include "PropertyPermissionList.h"
 #include "Framework/Commands/UIAction.h"
 #include "Widgets/Views/STableViewBase.h"
 #include "Widgets/Views/STableRow.h"
@@ -715,7 +716,10 @@ private:
 		for( auto ExtensionIter = PathExtensions.CreateIterator(); ExtensionIter; ++ExtensionIter )
 		{
 			const FPropertyInfo& Extension = *ExtensionIter;
-			TypeToProperties.Add( Extension.Property->GetOwnerStruct(), Extension );
+			if(FPropertyEditorPermissionList::Get().DoesPropertyPassFilter(Extension.Property->GetOwnerStruct(), Extension.Property->GetFName()))
+			{
+				TypeToProperties.Add( Extension.Property->GetOwnerStruct(), Extension );
+			}
 		}
 
 		FMenuBuilder MenuBuilder( true, NULL );
