Fix EditInlineNew property editor class picker not closing when class is picked by reverting the change 24927105

diff --git a/Engine/Source/Editor/PropertyEditor/Private/UserInterface/PropertyEditor/SPropertyEditorEditInline.cpp b/Engine/Source/Editor/PropertyEditor/Private/UserInterface/PropertyEditor/SPropertyEditorEditInline.cpp
--- a/Engine/Source/Editor/PropertyEditor/Private/UserInterface/PropertyEditor/SPropertyEditorEditInline.cpp
+++ b/Engine/Source/Editor/PropertyEditor/Private/UserInterface/PropertyEditor/SPropertyEditorEditInline.cpp
@@ -348,5 +348,7 @@ void SPropertyEditorEditInline::OnClassPicked(UClass* InClass)
 
 		// Force a rebuild of the children when this node changes
 		PropertyNode->RequestRebuildChildren();
+
+		ComboButton->SetIsOpen(false);
 	}
 }
