Allow the "browse to" button in SPropertyEditorAsset to be enabled even if the rest of the widget is disabled

#rnx

diff --git a/Engine/Source/Editor/PropertyEditor/Private/UserInterface/PropertyEditor/SPropertyEditorAsset.cpp b/Engine/Source/Editor/PropertyEditor/Private/UserInterface/PropertyEditor/SPropertyEditorAsset.cpp
--- a/Engine/Source/Editor/PropertyEditor/Private/UserInterface/PropertyEditor/SPropertyEditorAsset.cpp
+++ b/Engine/Source/Editor/PropertyEditor/Private/UserInterface/PropertyEditor/SPropertyEditorAsset.cpp
@@ -241,11 +241,15 @@ void SPropertyEditorAsset::Construct(const FArguments& InArgs, const TSharedPtr<
 	PropertyEditor = InPropertyEditor;
 	PropertyHandle = InArgs._PropertyHandle;
 	OwnerAssetDataArray = InArgs._OwnerAssetDataArray;
+	OnIsEnabled = InArgs._IsEnabled;
 	OnSetObject = InArgs._OnSetObject;
 	OnShouldFilterAsset = InArgs._OnShouldFilterAsset;
 	OnShouldFilterActor = InArgs._OnShouldFilterActor;
 	ObjectPath = InArgs._ObjectPath;
 
+	// Override this as we stole the value to use as OnIsEnabled for the inner widgets
+	SetEnabled(true);
+
 	if(InArgs._InWidgetRow.IsSet() && InArgs._InWidgetRow.GetValue() != nullptr)
 	{
 		if (!InArgs._InWidgetRow.GetValue()->CopyMenuAction.IsBound())
@@ -1357,7 +1361,11 @@ FReply SPropertyEditorAsset::OnAssetThumbnailDoubleClick( const FGeometry& InMyG
 
 bool SPropertyEditorAsset::CanEdit() const
 {
-	return PropertyEditor.IsValid() ? !PropertyEditor->IsEditConst() : true;
+	if (PropertyEditor.IsValid() && PropertyEditor->IsEditConst())
+	{
+		return false;
+	}
+	return OnIsEnabled.Get(true);
 }
 
 bool SPropertyEditorAsset::CanSetBasedOnCustomClasses( const FAssetData& InAssetData ) const
diff --git a/Engine/Source/Editor/PropertyEditor/Private/UserInterface/PropertyEditor/SPropertyEditorAsset.h b/Engine/Source/Editor/PropertyEditor/Private/UserInterface/PropertyEditor/SPropertyEditorAsset.h
--- a/Engine/Source/Editor/PropertyEditor/Private/UserInterface/PropertyEditor/SPropertyEditorAsset.h
+++ b/Engine/Source/Editor/PropertyEditor/Private/UserInterface/PropertyEditor/SPropertyEditorAsset.h
@@ -394,6 +394,9 @@ private:
 	/** The number of additional buttons this picker has. */
 	int32 NumButtons;
 
+	/** Base attribute passed to specify whether this editor is enabled */ 
+	TAttribute<bool> OnIsEnabled;
+
 	/** Delegate to call when our object value is set */
 	FOnSetObject OnSetObject;
 
