Remove (the never should have been exposed) direct access to the private PropertyNode via the PropertyHandle

diff --git a/Engine/Plugins/Animation/DeformerGraph/Source/OptimusEditor/OptimusEditor.Build.cs b/Engine/Plugins/Animation/DeformerGraph/Source/OptimusEditor/OptimusEditor.Build.cs
--- a/Engine/Plugins/Animation/DeformerGraph/Source/OptimusEditor/OptimusEditor.Build.cs
+++ b/Engine/Plugins/Animation/DeformerGraph/Source/OptimusEditor/OptimusEditor.Build.cs
@@ -8,7 +8,6 @@ namespace UnrealBuildTool.Rules
         {
 			PrivateIncludePaths.AddRange(new string[] {
 				System.IO.Path.Combine(GetModuleDirectory("OptimusCore"), "Private"),
-				System.IO.Path.Combine(GetModuleDirectory("PropertyEditor"), "Private"),
 			});
 
             PrivateDependencyModuleNames.AddRange(
diff --git a/Engine/Plugins/Animation/DeformerGraph/Source/OptimusEditor/Private/OptimusDetailsCustomization.cpp b/Engine/Plugins/Animation/DeformerGraph/Source/OptimusEditor/Private/OptimusDetailsCustomization.cpp
--- a/Engine/Plugins/Animation/DeformerGraph/Source/OptimusEditor/Private/OptimusDetailsCustomization.cpp
+++ b/Engine/Plugins/Animation/DeformerGraph/Source/OptimusEditor/Private/OptimusDetailsCustomization.cpp
@@ -27,7 +27,6 @@
 #include "OptimusSource.h"
 #include "OptimusValidatedName.h"
 #include "OptimusValueContainer.h"
-#include "PropertyNode.h"
 #include "ScopedTransaction.h"
 #include "Styling/AppStyle.h"
 #include "Styling/SlateIconFinder.h"
@@ -992,11 +991,10 @@ private:
 	void OnDeleteItem() const
 	{
 		TSharedPtr<IPropertyHandleArray> ArrayHandle = BindingPropertyHandle->GetParentHandle()->AsArray();
-		TSharedPtr<FPropertyNode> PropertyNode = BindingPropertyHandle->GetPropertyNode(); 
 
 		check(ArrayHandle.IsValid());
 
-		int32 Index = PropertyNode->GetArrayIndex();
+		int32 Index = BindingPropertyHandle->GetArrayIndex();
 
 		if (ArrayHandle.IsValid())
 		{
@@ -1004,7 +1002,7 @@ private:
 		}
 
 		//In case the property is show in the favorite category refresh the whole tree
-		if (PropertyNode->IsFavorite() || (PropertyNode->GetParentNode() != nullptr && PropertyNode->GetParentNode()->IsFavorite()))
+		if (BindingPropertyHandle->IsFavorite() || (BindingPropertyHandle->GetParentHandle() != nullptr && BindingPropertyHandle->GetParentHandle()->IsFavorite()))
 		{
 			CustomizationUtils->GetPropertyUtilities()->ForceRefresh();
 		}
diff --git a/Engine/Plugins/Experimental/AnimNext/Source/AnimNextEditor/Private/Graph/PropertyTypeCustomization.cpp b/Engine/Plugins/Experimental/AnimNext/Source/AnimNextEditor/Private/Graph/PropertyTypeCustomization.cpp
--- a/Engine/Plugins/Experimental/AnimNext/Source/AnimNextEditor/Private/Graph/PropertyTypeCustomization.cpp
+++ b/Engine/Plugins/Experimental/AnimNext/Source/AnimNextEditor/Private/Graph/PropertyTypeCustomization.cpp
@@ -4,7 +4,6 @@
 
 #include "Graph/WidgetFactories.h"
 #include "DetailWidgetRow.h"
-#include "Editor/PropertyEditor/Private/PropertyNode.h"
 #include "IDetailChildrenBuilder.h"
 #include "PropertyHandle.h"
 #include "SClassViewer.h"
@@ -47,7 +46,7 @@ void FPropertyTypeCustomization::CustomizeHeader(TSharedRef<IPropertyHandle> Pro
 						PropertyHandle->NotifyPreChange();
 						PropertyHandle->SetValue(NewValue);
 						PropertyHandle->NotifyPostChange(EPropertyChangeType::ValueSet);
-						PropertyHandle->GetPropertyNode()->GetParentNode()->RequestRebuildChildren();
+						PropertyHandle->GetParentHandle()->RequestRebuildChildren();
 					}
 				}),
 				nullptr
diff --git a/Engine/Plugins/FX/Niagara/Source/NiagaraEditor/NiagaraEditor.Build.cs b/Engine/Plugins/FX/Niagara/Source/NiagaraEditor/NiagaraEditor.Build.cs
--- a/Engine/Plugins/FX/Niagara/Source/NiagaraEditor/NiagaraEditor.Build.cs
+++ b/Engine/Plugins/FX/Niagara/Source/NiagaraEditor/NiagaraEditor.Build.cs
@@ -7,10 +7,6 @@ public class NiagaraEditor : ModuleRules
 {
 	public NiagaraEditor(ReadOnlyTargetRules Target) : base(Target)
 	{
-		PrivateIncludePaths.AddRange(new string[] {
-			Path.Combine(GetModuleDirectory("PropertyEditor"), "Private"),
-		});
-
 		PrivateDependencyModuleNames.AddRange(
 			new string[] {
                 "RHI",
diff --git a/Engine/Plugins/FX/Niagara/Source/NiagaraEditor/Private/Customizations/NiagaraPlatformSetCustomization.cpp b/Engine/Plugins/FX/Niagara/Source/NiagaraEditor/Private/Customizations/NiagaraPlatformSetCustomization.cpp
--- a/Engine/Plugins/FX/Niagara/Source/NiagaraEditor/Private/Customizations/NiagaraPlatformSetCustomization.cpp
+++ b/Engine/Plugins/FX/Niagara/Source/NiagaraEditor/Private/Customizations/NiagaraPlatformSetCustomization.cpp
@@ -18,8 +18,6 @@
 #include "NiagaraTypes.h"
 #include "PlatformInfo.h"
 #include "PropertyHandle.h"
-#include "PropertyNode.h"
-#include "StructurePropertyNode.h"
 #include "Widgets/Images/SImage.h"
 #include "Widgets/Input/SButton.h"
 #include "Widgets/Input/SCheckBox.h"
@@ -39,26 +37,26 @@ uint8* GetBaseAddress(TSharedRef<IPropertyHandle> PropertyHandle)
 		PropertyHandle->GetOuterObjects(Objects);
 		return reinterpret_cast<uint8*>(Objects[0]);
 	}
-	
+
 	// walk the struct hierarchy to find the parent StructOnScope
-	FStructurePropertyNode* StructurePropertyNode = nullptr;
 	TSharedPtr<IPropertyHandle> LastParent = PropertyHandle;
-	while (StructurePropertyNode == nullptr)
+	while (true)
 	{
 		TSharedPtr<IPropertyHandle> NewParent = LastParent->GetParentHandle();
 		if (!ensureMsgf(NewParent.IsValid() && NewParent != LastParent, TEXT("Unable to walk property chain")))
 		{
-			return nullptr;
+			break;
 		}
 		
-		if (FComplexPropertyNode* ComplexPropertyNode = NewParent->GetPropertyNode()->AsComplexNode())
+		TSharedPtr<IPropertyHandleStruct> StructHandle = NewParent->AsStruct();
+		if (StructHandle.IsValid())
 		{
-			StructurePropertyNode = ComplexPropertyNode->AsStructureNode();
+			return StructHandle->GetStructData()->GetStructMemory();
 		}
 		LastParent = NewParent;
 	}
-	TSharedPtr<FStructOnScope> StructOnScope = StructurePropertyNode->GetStructData();
-	return StructOnScope->GetStructMemory();
+
+	return nullptr;
 }
 
 void FNiagaraPlatformSetCustomization::CustomizeHeader(TSharedRef<IPropertyHandle> InPropertyHandle, FDetailWidgetRow& HeaderRow, IPropertyTypeCustomizationUtils& CustomizationUtils)
diff --git a/Engine/Plugins/MovieScene/MovieRenderPipeline/Source/MovieRenderPipelineEditor/Private/Customizations/ConsoleVariableSettingCustomization.h b/Engine/Plugins/MovieScene/MovieRenderPipeline/Source/MovieRenderPipelineEditor/Private/Customizations/ConsoleVariableSettingCustomization.h
--- a/Engine/Plugins/MovieScene/MovieRenderPipeline/Source/MovieRenderPipelineEditor/Private/Customizations/ConsoleVariableSettingCustomization.h
+++ b/Engine/Plugins/MovieScene/MovieRenderPipeline/Source/MovieRenderPipelineEditor/Private/Customizations/ConsoleVariableSettingCustomization.h
@@ -12,8 +12,6 @@
 #include "PropertyHandle.h"
 #include "Sections/MovieSceneConsoleVariableTrackInterface.h"
 
-#include "Editor/PropertyEditor/Private/PropertyNode.h"
-
 #define LOCTEXT_NAMESPACE "SMoviePipelineEditor"
 
 /** Customize how properties in UMoviePipelineConsoleVariableSetting appear in the details panel. */
@@ -71,8 +69,8 @@ protected:
 		static const FText ConsoleVariableDisabledText = LOCTEXT("DisabledConsoleVariable", "This console variable is disabled.");
 
 		// Add the preset asset chooser as the group header
-		const FString& PropName = ElementProperty->GetPropertyNode()->GetPropertyPath();
-		IDetailGroup& Group = ChildrenBuilder.AddGroup(FName(PropName), FText::FromString(PropName));
+		FStringView PropName = ElementProperty->GetPropertyPath();
+		IDetailGroup& Group = ChildrenBuilder.AddGroup(FName(PropName), FText::FromStringView(PropName));
 		Group.HeaderRow()
 		[
 			ElementProperty->CreatePropertyValueWidget()
diff --git a/Engine/Plugins/Runtime/WorldConditions/Source/WorldConditionsEditor/Private/Customizations/WorldConditionContextDataRefDetails.cpp b/Engine/Plugins/Runtime/WorldConditions/Source/WorldConditionsEditor/Private/Customizations/WorldConditionContextDataRefDetails.cpp
--- a/Engine/Plugins/Runtime/WorldConditions/Source/WorldConditionsEditor/Private/Customizations/WorldConditionContextDataRefDetails.cpp
+++ b/Engine/Plugins/Runtime/WorldConditions/Source/WorldConditionsEditor/Private/Customizations/WorldConditionContextDataRefDetails.cpp
@@ -9,7 +9,6 @@
 #include "WorldConditionQuery.h"
 #include "WorldConditionSchema.h"
 #include "Editor.h"
-#include "Editor/PropertyEditor/Private/PropertyNode.h"
 
 #define LOCTEXT_NAMESPACE "WorldCondition"
 
@@ -73,9 +72,9 @@ void FWorldConditionContextDataRefDetails::CacheContextData()
 
 	if (!BaseStruct)
 	{
-		check(StructProperty.IsValid() && StructProperty->GetPropertyNode().IsValid())
+		check(StructProperty.IsValid() && StructProperty->IsValidHandle())
 		UE_LOG(LogWorldCondition, Error, TEXT("%s: Could not find BaseStruct '%s' nor BaseClass '%s' based on the property metadata, expecting full struct name."),
-			*StructProperty->GetPropertyNode()->GetPropertyPath(), *BaseStructName, *BaseClassName);
+			*FString(StructProperty->GetPropertyPath()), *BaseStructName, *BaseClassName);
 	}
 
 	// Find schema from outer FWorldConditionQueryDefinition.
diff --git a/Engine/Plugins/Runtime/nDisplay/Source/DisplayClusterColorGrading/Private/DetailView/SDisplayClusterColorGradingDetailTreeRow.cpp b/Engine/Plugins/Runtime/nDisplay/Source/DisplayClusterColorGrading/Private/DetailView/SDisplayClusterColorGradingDetailTreeRow.cpp
--- a/Engine/Plugins/Runtime/nDisplay/Source/DisplayClusterColorGrading/Private/DetailView/SDisplayClusterColorGradingDetailTreeRow.cpp
+++ b/Engine/Plugins/Runtime/nDisplay/Source/DisplayClusterColorGrading/Private/DetailView/SDisplayClusterColorGradingDetailTreeRow.cpp
@@ -939,14 +939,14 @@ bool SDisplayClusterColorGradingDetailTreeRow::IsValidDrop(const FDragDropEvent&
 		{
 			TSharedPtr<IPropertyHandle> SwappingPropertyHandle = ArrayRowDropOp->RowBeingDragged.Pin()->GetPropertyHandle();
 			TSharedPtr<IPropertyHandle> ThisPropertyHandle = GetPropertyHandle();
-			if (ThisPropertyHandle.IsValid() && SwappingPropertyHandle.IsValid() && SwappingPropertyHandle->GetPropertyNode() != ThisPropertyHandle->GetPropertyNode())
+			if (ThisPropertyHandle.IsValid() && SwappingPropertyHandle.IsValid() && !SwappingPropertyHandle->IsSamePropertyNode(ThisPropertyHandle))
 			{
 				const int32 OriginalIndex = SwappingPropertyHandle->GetIndexInArray();
 				const int32 NewIndex = GetDropNewIndex(OriginalIndex, ThisPropertyHandle->GetIndexInArray(), DropZone);
 
 				if (OriginalIndex != NewIndex)
 				{
-					if (SwappingPropertyHandle->GetParentHandle()->AsArray() && SwappingPropertyHandle->GetParentHandle()->GetPropertyNode() == ThisPropertyHandle->GetParentHandle()->GetPropertyNode())
+					if (SwappingPropertyHandle->GetParentHandle()->AsArray() && SwappingPropertyHandle->GetParentHandle()->IsSamePropertyNode(ThisPropertyHandle->GetParentHandle()))
 					{
 						return true;
 					}
diff --git a/Engine/Plugins/Runtime/nDisplay/Source/DisplayClusterTests/Private/DisplayClusterTestUtils.h b/Engine/Plugins/Runtime/nDisplay/Source/DisplayClusterTests/Private/DisplayClusterTestUtils.h
--- a/Engine/Plugins/Runtime/nDisplay/Source/DisplayClusterTests/Private/DisplayClusterTestUtils.h
+++ b/Engine/Plugins/Runtime/nDisplay/Source/DisplayClusterTests/Private/DisplayClusterTestUtils.h
@@ -2,7 +2,6 @@
 
 #pragma once
 #include "DisplayClusterConfigurationTypes_Viewport.h"
-#include "Editor/PropertyEditor/Private/PropertyNode.h"
 
 #if WITH_DEV_AUTOMATION_TESTS
 
diff --git a/Engine/Plugins/VirtualProduction/RemoteControl/Source/RemoteControlUI/Private/UI/SRCPanelExposedField.cpp b/Engine/Plugins/VirtualProduction/RemoteControl/Source/RemoteControlUI/Private/UI/SRCPanelExposedField.cpp
--- a/Engine/Plugins/VirtualProduction/RemoteControl/Source/RemoteControlUI/Private/UI/SRCPanelExposedField.cpp
+++ b/Engine/Plugins/VirtualProduction/RemoteControl/Source/RemoteControlUI/Private/UI/SRCPanelExposedField.cpp
@@ -14,7 +14,6 @@
 #include "IRemoteControlUIModule.h"
 #include "Layout/Visibility.h"
 #include "Materials/MaterialInstanceDynamic.h"
-#include "PropertyNode.h"
 #include "RCPanelWidgetRegistry.h"
 #include "RemoteControlField.h"
 #include "RemoteControlPanelStyle.h"
@@ -585,10 +584,10 @@ TSharedRef<SWidget> SRCPanelExposedField::ConstructResetToDefaultWidget(UObject*
 				FRCResetToDefaultArgs Args;
 				Args.Property = InPropertyHandle->GetProperty();
 				
-				if (TSharedPtr<FPropertyNode> PropertyNode = InPropertyHandle->GetPropertyNode())
+				if (InPropertyHandle->IsValidHandle())
 				{
-					Args.Path = PropertyNode->GetPropertyPath();
-					Args.ArrayIndex = PropertyNode->GetArrayIndex();
+					Args.Path = InPropertyHandle->GetPropertyPath();
+					Args.ArrayIndex = InPropertyHandle->GetArrayIndex();
 				}
 
 				Args.bCreateTransaction = false;
@@ -608,10 +607,10 @@ TSharedRef<SWidget> SRCPanelExposedField::ConstructResetToDefaultWidget(UObject*
 				FRCResetToDefaultArgs Args;
 				Args.Property = InPropertyHandle->GetProperty();
 				
-				if (TSharedPtr<FPropertyNode> PropertyNode = InPropertyHandle->GetPropertyNode())
+				if (InPropertyHandle->IsValidHandle())
 				{
-					Args.Path = PropertyNode->GetPropertyPath();
-					Args.ArrayIndex = PropertyNode->GetArrayIndex();
+					Args.Path = InPropertyHandle->GetPropertyPath();
+					Args.ArrayIndex = InPropertyHandle->GetArrayIndex();
 				}
 				
 				Args.bCreateTransaction = true;
@@ -629,18 +628,14 @@ TSharedRef<SWidget> SRCPanelExposedField::ConstructResetToDefaultWidget(UObject*
 	if (UMeshComponent* MeshComponent = Cast<UMeshComponent>(InObject))
 	{
 		if (InPropertyHandle->GetProperty()->GetFName() == GET_MEMBER_NAME_CHECKED(UMeshComponent, OverrideMaterials) &&
-			MeshComponent->OverrideMaterials[InPropertyHandle->GetPropertyNode()->GetArrayIndex()] == NULL)
+			MeshComponent->OverrideMaterials[InPropertyHandle->GetArrayIndex()] == NULL)
 		{
 			IRemoteControlModule& RemoteControlModule = IRemoteControlModule::Get();
 
 			FRCResetToDefaultArgs Args;
 			Args.Property = InPropertyHandle->GetProperty();
-			
-			if (TSharedPtr<FPropertyNode> PropertyNode = InPropertyHandle->GetPropertyNode())
-			{
-				Args.Path = PropertyNode->GetPropertyPath();
-				Args.ArrayIndex = PropertyNode->GetArrayIndex();
-			}
+			Args.Path = InPropertyHandle->GetPropertyPath();
+			Args.ArrayIndex = InPropertyHandle->GetArrayIndex();
 
 			Args.bCreateTransaction = false;
 			RemoteControlModule.ResetToDefaultValue(InObject, Args);
diff --git a/Engine/Plugins/VirtualProduction/RemoteControl/Source/RemoteControlUI/RemoteControlUI.Build.cs b/Engine/Plugins/VirtualProduction/RemoteControl/Source/RemoteControlUI/RemoteControlUI.Build.cs
--- a/Engine/Plugins/VirtualProduction/RemoteControl/Source/RemoteControlUI/RemoteControlUI.Build.cs
+++ b/Engine/Plugins/VirtualProduction/RemoteControl/Source/RemoteControlUI/RemoteControlUI.Build.cs
@@ -7,14 +7,6 @@ public class RemoteControlUI : ModuleRules
 {
 	public RemoteControlUI(ReadOnlyTargetRules Target) : base(Target)
 	{
-		PrivateIncludePaths.AddRange(new string[] {
-			Path.Combine(GetModuleDirectory("PropertyEditor"), "Private"),
-		});
-
-		PublicDependencyModuleNames.AddRange(
-			new string[] {}
-		);
-
         PrivateDependencyModuleNames.AddRange(
 			new string[] {
 				"ApplicationCore",
diff --git a/Engine/Source/Editor/PropertyEditor/Private/DetailLayoutBuilderImpl.cpp b/Engine/Source/Editor/PropertyEditor/Private/DetailLayoutBuilderImpl.cpp
--- a/Engine/Source/Editor/PropertyEditor/Private/DetailLayoutBuilderImpl.cpp
+++ b/Engine/Source/Editor/PropertyEditor/Private/DetailLayoutBuilderImpl.cpp
@@ -745,7 +745,7 @@ bool FDetailLayoutBuilderImpl::IsPropertyVisible( TSharedRef<IPropertyHandle> Pr
 {
 	if (PropertyHandle->IsValidHandle() && DetailsView != nullptr)
 	{
-		TSharedPtr<FPropertyNode> PropertyNode = PropertyHandle->GetPropertyNode();
+		TSharedPtr<FPropertyNode> PropertyNode = StaticCastSharedRef<FPropertyHandleBase>(PropertyHandle)->GetPropertyNode();
 		const FCategoryPropertyNode* CategoryNode = PropertyNode.IsValid() ? PropertyNode->AsCategoryNode() : nullptr;
 		if (CategoryNode != nullptr)
 		{
diff --git a/Engine/Source/Editor/PropertyEditor/Private/PropertyEditorDelegates.cpp b/Engine/Source/Editor/PropertyEditor/Private/PropertyEditorDelegates.cpp
--- a/Engine/Source/Editor/PropertyEditor/Private/PropertyEditorDelegates.cpp
+++ b/Engine/Source/Editor/PropertyEditor/Private/PropertyEditorDelegates.cpp
@@ -5,7 +5,7 @@
 #include "HAL/PlatformCrt.h"
 #include "Misc/AssertionMacros.h"
 #include "ObjectPropertyNode.h"
-#include "PropertyHandle.h"
+#include "PropertyHandleImpl.h"
 #include "PropertyNode.h"
 
 class FProperty;
@@ -13,7 +13,7 @@ class FProperty;
 FPropertyAndParent::FPropertyAndParent(const TSharedRef<IPropertyHandle>& InPropertyHandle) :
 	Property(*InPropertyHandle->GetProperty())
 {
-	Initialize(InPropertyHandle->GetPropertyNode().ToSharedRef());
+	Initialize(StaticCastSharedRef<FPropertyHandleBase>(InPropertyHandle)->GetPropertyNode().ToSharedRef());
 }
 
 FPropertyAndParent::FPropertyAndParent(const TSharedRef<FPropertyNode>& InPropertyNode) :
diff --git a/Engine/Source/Editor/PropertyEditor/Private/PropertyEditorHelpers.cpp b/Engine/Source/Editor/PropertyEditor/Private/PropertyEditorHelpers.cpp
--- a/Engine/Source/Editor/PropertyEditor/Private/PropertyEditorHelpers.cpp
+++ b/Engine/Source/Editor/PropertyEditor/Private/PropertyEditorHelpers.cpp
@@ -626,9 +626,14 @@ namespace PropertyEditorHelpers
 		{
 			PropertyHandle = MakeShareable(new FPropertyHandleFieldPath(PropertyNode, NotifyHook, PropertyUtilities));
 		}
+		// struct should be checked last as there are several specializations of it above
+		else if (FPropertyHandleStruct::Supports(PropertyNode))
+		{
+			PropertyHandle = MakeShareable(new FPropertyHandleStruct(PropertyNode, NotifyHook, PropertyUtilities));
+		}
 		else
 		{
-			// Untyped or doesn't support getting the property directly but the property is still valid(probably struct property)
+			// Untyped or doesn't support getting the property directly but the property is still valid
 			PropertyHandle = MakeShareable( new FPropertyHandleBase( PropertyNode, NotifyHook, PropertyUtilities ) ); 
 		}
 
diff --git a/Engine/Source/Editor/PropertyEditor/Private/PropertyHandleImpl.cpp b/Engine/Source/Editor/PropertyEditor/Private/PropertyHandleImpl.cpp
--- a/Engine/Source/Editor/PropertyEditor/Private/PropertyHandleImpl.cpp
+++ b/Engine/Source/Editor/PropertyEditor/Private/PropertyHandleImpl.cpp
@@ -2220,6 +2220,16 @@ bool FPropertyHandleBase::IsValidHandle() const
 	return Implementation->HasValidPropertyNode();
 }
 
+bool FPropertyHandleBase::IsSamePropertyNode(TSharedPtr<IPropertyHandle> OtherHandle) const
+{
+	if (OtherHandle.IsValid() && OtherHandle->IsValidHandle())
+	{
+		return GetPropertyNode() == StaticCastSharedPtr<FPropertyHandleBase>(OtherHandle)->GetPropertyNode();
+	}
+
+	return !IsValidHandle();
+}
+
 FText FPropertyHandleBase::GetPropertyDisplayName() const
 {
 	return Implementation->GetDisplayName();
@@ -2272,6 +2282,42 @@ void FPropertyHandleBase::ClearResetToDefaultCustomized()
 	}
 }
 
+FStringView FPropertyHandleBase::GetPropertyPath() const
+{
+	if (Implementation->GetPropertyNode().IsValid())
+	{
+		return Implementation->GetPropertyNode()->GetPropertyPath();
+	}
+
+	return FStringView();
+}
+
+int32 FPropertyHandleBase::GetArrayIndex() const
+{
+	if (Implementation->GetPropertyNode().IsValid())
+	{
+		return Implementation->GetPropertyNode()->GetArrayIndex();
+	}
+	return INDEX_NONE;
+}
+
+void FPropertyHandleBase::RequestRebuildChildren()
+{
+	if (Implementation->GetPropertyNode().IsValid())
+	{
+		Implementation->GetPropertyNode()->RequestRebuildChildren();
+	}
+}
+
+bool FPropertyHandleBase::IsFavorite() const
+{
+	if (Implementation->GetPropertyNode().IsValid())
+	{
+		return Implementation->GetPropertyNode()->IsFavorite();
+	}
+	return false;
+}
+
 bool FPropertyHandleBase::IsCustomized() const
 {
 	return Implementation->GetPropertyNode()->HasNodeFlags( EPropertyNodeFlags::IsCustomized ) != 0;
@@ -3472,6 +3518,7 @@ IMPLEMENT_PROPERTY_VALUE( FPropertyHandleText )
 IMPLEMENT_PROPERTY_VALUE( FPropertyHandleSet )
 IMPLEMENT_PROPERTY_VALUE( FPropertyHandleMap )
 IMPLEMENT_PROPERTY_VALUE( FPropertyHandleFieldPath )
+IMPLEMENT_PROPERTY_VALUE( FPropertyHandleStruct )
 
 // int32 
 bool FPropertyHandleInt::Supports( TSharedRef<FPropertyNode> PropertyNode )
@@ -4384,6 +4431,34 @@ FPropertyAccess::Result FPropertyHandleMixed::SetValue(const float& NewValue, EP
 	return SetValue((double)NewValue);
 }
 
+// Struct
+bool FPropertyHandleStruct::Supports( TSharedRef<FPropertyNode> PropertyNode )
+{
+	if (FComplexPropertyNode* ComplexNode = PropertyNode->AsComplexNode())
+	{
+		return (ComplexNode->AsStructureNode() != nullptr);
+	}
+
+	return false;
+}
+
+TSharedPtr<IPropertyHandleStruct> FPropertyHandleStruct::AsStruct()
+{
+	return SharedThis(this);
+}
+
+TSharedPtr<FStructOnScope> FPropertyHandleStruct::GetStructData() const
+{
+	TSharedPtr<FStructurePropertyNode> StructPropertyNode = StaticCastSharedPtr<FStructurePropertyNode>(Implementation->GetPropertyNode());
+
+	if (StructPropertyNode.IsValid())
+	{
+		return StructPropertyNode->GetStructData();
+	}
+
+	return nullptr;
+}
+
 // Vector
 bool FPropertyHandleVector::Supports( TSharedRef<FPropertyNode> PropertyNode )
 {
@@ -4411,7 +4486,7 @@ bool FPropertyHandleVector::Supports( TSharedRef<FPropertyNode> PropertyNode )
 }
 
 FPropertyHandleVector::FPropertyHandleVector( TSharedRef<class FPropertyNode> PropertyNode, class FNotifyHook* NotifyHook, TSharedPtr<IPropertyUtilities> PropertyUtilities )
-	: FPropertyHandleBase( PropertyNode, NotifyHook, PropertyUtilities ) 
+	: FPropertyHandleStruct( PropertyNode, NotifyHook, PropertyUtilities ) 
 {
 	if( Implementation->GetNumChildren() > 0 )
 	{
@@ -4661,7 +4736,7 @@ bool FPropertyHandleRotator::Supports( TSharedRef<FPropertyNode> PropertyNode )
 }
 
 FPropertyHandleRotator::FPropertyHandleRotator( TSharedRef<class FPropertyNode> PropertyNode, FNotifyHook* NotifyHook, TSharedPtr<IPropertyUtilities> PropertyUtilities )
-	: FPropertyHandleBase( PropertyNode, NotifyHook, PropertyUtilities ) 
+	: FPropertyHandleStruct( PropertyNode, NotifyHook, PropertyUtilities ) 
 {
 	const bool bRecurse = false;
 	// A vector is a struct property that has 3 children.  We get/set the values from the children
@@ -4729,8 +4804,8 @@ FPropertyAccess::Result FPropertyHandleRotator::SetYaw( double InYaw, EPropertyV
 	return Res;
 }
 
-FPropertyHandleColor::FPropertyHandleColor(TSharedRef<FPropertyNode> PropertyNode, FNotifyHook* NotifyHook, TSharedPtr<IPropertyUtilities> PropertyUtilities)\
-	: FPropertyHandleBase(PropertyNode, NotifyHook, PropertyUtilities)
+FPropertyHandleColor::FPropertyHandleColor(TSharedRef<FPropertyNode> PropertyNode, FNotifyHook* NotifyHook, TSharedPtr<IPropertyUtilities> PropertyUtilities)
+	: FPropertyHandleStruct(PropertyNode, NotifyHook, PropertyUtilities)
 {
 }
 
diff --git a/Engine/Source/Editor/PropertyEditor/Private/PropertyHandleImpl.h b/Engine/Source/Editor/PropertyEditor/Private/PropertyHandleImpl.h
--- a/Engine/Source/Editor/PropertyEditor/Private/PropertyHandleImpl.h
+++ b/Engine/Source/Editor/PropertyEditor/Private/PropertyHandleImpl.h
@@ -415,6 +415,7 @@ public:
 
 	/** IPropertyHandle interface */
 	virtual bool IsValidHandle() const override;
+	virtual bool IsSamePropertyNode(TSharedPtr<IPropertyHandle> OtherHandle) const override;
 	virtual FText GetPropertyDisplayName() const override;
 	virtual void SetPropertyDisplayName(FText InDisplayName) override;
 	virtual void ResetToDefault() override;
@@ -464,6 +465,7 @@ public:
 	virtual TSharedPtr<IPropertyHandleArray> AsArray() override { return nullptr; }
 	virtual TSharedPtr<IPropertyHandleSet> AsSet() override { return nullptr; }
 	virtual TSharedPtr<IPropertyHandleMap> AsMap() override { return nullptr; }
+	virtual TSharedPtr<IPropertyHandleStruct> AsStruct() override { return nullptr; }
 	virtual const FFieldClass* GetPropertyClass() const override;
 	virtual FProperty* GetProperty() const override;
 	virtual FProperty* GetMetaDataProperty() const override;
@@ -508,8 +510,12 @@ public:
 	virtual void ExecuteCustomResetToDefault(const FResetToDefaultOverride& InOnCustomResetToDefault) override;
 	virtual FName GetDefaultCategoryName() const override;
 	virtual FText GetDefaultCategoryText() const override;
+	virtual FStringView GetPropertyPath() const override;
+	virtual int32 GetArrayIndex() const override;
+	virtual void RequestRebuildChildren() override;
+	virtual bool IsFavorite() const override;
 
-	TSharedPtr<FPropertyNode> GetPropertyNode() const;
+	PROPERTYEDITOR_API TSharedPtr<FPropertyNode> GetPropertyNode() const;
 	void OnCustomResetToDefault(const FResetToDefaultOverride& OnCustomResetToDefault);
 
 private:
@@ -625,7 +631,16 @@ public:
 	virtual FPropertyAccess::Result SetValue(const double& InValue, EPropertyValueSetFlags::Type Flags = EPropertyValueSetFlags::DefaultFlags) override;
 };
 
-class FPropertyHandleVector : public FPropertyHandleBase
+class FPropertyHandleStruct : public FPropertyHandleBase, public IPropertyHandleStruct
+{
+public:
+	FPropertyHandleStruct(TSharedRef<FPropertyNode> PropertyNode, FNotifyHook* NotifyHook, TSharedPtr<IPropertyUtilities> PropertyUtilities);
+	static bool Supports(TSharedRef<FPropertyNode> PropertyNode);
+	virtual TSharedPtr<IPropertyHandleStruct> AsStruct() override;
+	virtual TSharedPtr<FStructOnScope> GetStructData() const override;
+};
+
+class FPropertyHandleVector : public FPropertyHandleStruct
 {
 public:
 	FPropertyHandleVector( TSharedRef<FPropertyNode> PropertyNode, FNotifyHook* NotifyHook, TSharedPtr<IPropertyUtilities> PropertyUtilities );
@@ -650,7 +665,7 @@ private:
 	TArray< TSharedPtr<FPropertyHandleMixed> > VectorComponents;	
 };
 
-class FPropertyHandleRotator : public FPropertyHandleBase
+class FPropertyHandleRotator : public FPropertyHandleStruct
 {
 public:
 	FPropertyHandleRotator( TSharedRef<FPropertyNode> PropertyNode, FNotifyHook* NotifyHook, TSharedPtr<IPropertyUtilities> PropertyUtilities  );
@@ -667,7 +682,7 @@ private:
 	TSharedPtr<FPropertyHandleMixed> YawValue;
 };
 
-class FPropertyHandleColor : public FPropertyHandleBase
+class FPropertyHandleColor : public FPropertyHandleStruct
 {
 public:
 	FPropertyHandleColor( TSharedRef<FPropertyNode> PropertyNode, FNotifyHook* NotifyHook, TSharedPtr<IPropertyUtilities> PropertyUtilities );
@@ -762,4 +777,4 @@ public:
 	virtual FPropertyAccess::Result SetValue(FProperty* const& InValue, EPropertyValueSetFlags::Type Flags = EPropertyValueSetFlags::DefaultFlags) override;
 	virtual FPropertyAccess::Result GetValue(const FProperty*& OutValue) const override;
 	virtual FPropertyAccess::Result SetValue(const FProperty* const& InValue, EPropertyValueSetFlags::Type Flags = EPropertyValueSetFlags::DefaultFlags) override;
-};
+};
\ No newline at end of file
diff --git a/Engine/Source/Editor/PropertyEditor/Private/PropertyNode.h b/Engine/Source/Editor/PropertyEditor/Private/PropertyNode.h
--- a/Engine/Source/Editor/PropertyEditor/Private/PropertyNode.h
+++ b/Engine/Source/Editor/PropertyEditor/Private/PropertyNode.h
@@ -495,8 +495,8 @@ public:
 	/**
 	 * Accessor functions for internals
 	 */
-	const int32 GetArrayOffset() const { return ArrayOffset; }
-	const int32 GetArrayIndex() const { return ArrayIndex; }
+	int32 GetArrayOffset() const { return ArrayOffset; }
+	int32 GetArrayIndex() const { return ArrayIndex; }
 
 	/**
 	 * Return number of children that survived being filtered
diff --git a/Engine/Source/Editor/PropertyEditor/Private/SDetailSingleItemRow.cpp b/Engine/Source/Editor/PropertyEditor/Private/SDetailSingleItemRow.cpp
--- a/Engine/Source/Editor/PropertyEditor/Private/SDetailSingleItemRow.cpp
+++ b/Engine/Source/Editor/PropertyEditor/Private/SDetailSingleItemRow.cpp
@@ -1200,7 +1200,7 @@ bool SDetailSingleItemRow::IsFavorite() const
 			TSharedPtr<IPropertyHandle> PropertyHandle = Customization->CustomBuilderRow->GetPropertyHandle();
 			if (PropertyHandle.IsValid())
 			{
-				return PropertyHandle->GetPropertyNode()->IsFavorite();
+				return PropertyHandle->IsFavorite();
 			}
 
 			const FString& OriginalPath = Customization->CustomBuilderRow->GetOriginalPath();
@@ -1242,7 +1242,7 @@ void SDetailSingleItemRow::OnFavoriteMenuToggle()
 
 		if (PropertyHandle.IsValid())
 		{
-			PropertyHandle->GetPropertyNode()->SetFavorite(bNewValue); 
+			StaticCastSharedPtr<FPropertyHandleBase>(PropertyHandle)->GetPropertyNode()->SetFavorite(bNewValue); 
 		}
 		else if (!OriginalPath.IsEmpty())
 		{		
diff --git a/Engine/Source/Editor/PropertyEditor/Private/UserInterface/PropertyEditor/SPropertyEditorAsset.cpp b/Engine/Source/Editor/PropertyEditor/Private/UserInterface/PropertyEditor/SPropertyEditorAsset.cpp
--- a/Engine/Source/Editor/PropertyEditor/Private/UserInterface/PropertyEditor/SPropertyEditorAsset.cpp
+++ b/Engine/Source/Editor/PropertyEditor/Private/UserInterface/PropertyEditor/SPropertyEditorAsset.cpp
@@ -1156,7 +1156,7 @@ void SPropertyEditorAsset::OnBrowse()
 			if (Value.Object)
 			{
 				// This code only works on loaded objects
-				if (TSharedPtr<FPropertyNode> PropertyNodeToSync = PropertyHandleToUse->GetPropertyNode())
+				if (TSharedPtr<FPropertyNode> PropertyNodeToSync = StaticCastSharedPtr<FPropertyHandleBase>(PropertyHandleToUse)->GetPropertyNode())
 				{
 					FPropertyEditor::SyncToObjectsInNode(PropertyNodeToSync);
 				}
diff --git a/Engine/Source/Editor/PropertyEditor/Public/PropertyHandle.h b/Engine/Source/Editor/PropertyEditor/Public/PropertyHandle.h
--- a/Engine/Source/Editor/PropertyEditor/Public/PropertyHandle.h
+++ b/Engine/Source/Editor/PropertyEditor/Public/PropertyHandle.h
@@ -14,6 +14,7 @@ class FResetToDefaultOverride;
 class IPropertyHandleArray;
 class IPropertyHandleMap;
 class IPropertyHandleSet;
+class IPropertyHandleStruct;
 class IStructureDataProvider;
 
 namespace EPropertyValueSetFlags
@@ -44,6 +45,11 @@ public:
 	 * @return Whether or not the handle points to a valid property node. This can be true but GetProperty may still return null
 	 */
 	virtual bool IsValidHandle() const = 0;
+
+	/**
+	 * @return Whether or not the handle points to the same property node. Will return true if both point to an invalid node.
+	 */
+	virtual bool IsSamePropertyNode(TSharedPtr<IPropertyHandle> OtherHandle) const = 0;
 	
 	/**
 	 * @return Whether or not the property is edit const (can't be changed)
@@ -66,9 +72,19 @@ public:
 	virtual FProperty* GetProperty() const = 0;
 
 	/**
-	 * Gets the property node being edited.
+	 * Helper to fetch a PropertyPath 
+	 */
+	virtual FStringView GetPropertyPath() const = 0;
+
+	/**
+	 * Helper to fetch the ArrayIndex
 	 */
-	virtual TSharedPtr<FPropertyNode> GetPropertyNode() const = 0;
+	virtual int32 GetArrayIndex() const = 0;
+
+	/**
+	 * Indicates that children of this node should be rebuilt next tick.  Some topology changes will require this
+	 */
+	virtual void RequestRebuildChildren() = 0;
 
 	/**
 	 * Gets the property we should use to read meta-data
@@ -505,6 +521,11 @@ public:
 	 */
 	virtual TSharedPtr<IPropertyHandleMap> AsMap() = 0;
 
+	/**
+	 * @return This handle as struct if possible
+	 */
+	virtual TSharedPtr<IPropertyHandleStruct> AsStruct() = 0;
+
 	/**
 	 * @return The display name of the property
 	 */
@@ -550,6 +571,11 @@ public:
 	 */
 	virtual void ClearResetToDefaultCustomized() = 0;
 
+	/**
+	* @return true if the property is mark as a favorite
+	*/
+	virtual bool IsFavorite() const = 0;
+
 	/**
 	 * @return True if this property's UI is customized                                                              
 	 */
@@ -875,3 +901,14 @@ public:
 	 */
 	virtual void UnregisterOnNumElementsChanged(FDelegateHandle Handle) = 0;
 };
+
+/**
+ * A handle to a property which allows you to access a Struct's Data
+ */
+class IPropertyHandleStruct
+{
+public:
+	virtual ~IPropertyHandleStruct() {}
+
+	virtual TSharedPtr<FStructOnScope> GetStructData() const = 0;
+};
\ No newline at end of file
