Added templated overloads to DetailLayoutBuilder

diff --git a/Engine/Source/Editor/PropertyEditor/Public/DetailLayoutBuilder.h b/Engine/Source/Editor/PropertyEditor/Public/DetailLayoutBuilder.h
--- a/Engine/Source/Editor/PropertyEditor/Public/DetailLayoutBuilder.h
+++ b/Engine/Source/Editor/PropertyEditor/Public/DetailLayoutBuilder.h
@@ -76,6 +76,13 @@ public:
 	 */
 	virtual const TArray< TWeakObjectPtr<UObject> >& GetSelectedObjects() const = 0;
 
+	/**
+	 * Get the root objects (of ObjectType) observed by this layout.  
+	 * This is not guaranteed to be the same as the objects customized by this builder.  See GetObjectsBeingCustomized for that.
+	 */
+	template <typename ObjectType = UObject>
+	TArray< TWeakObjectPtr<ObjectType> > GetSelectedObjectsOfType() const;
+
 	/**
 	 * Gets the current object(s) being customized by this builder
 	 *
@@ -83,6 +90,15 @@ public:
 	 */
 	virtual void GetObjectsBeingCustomized( TArray< TWeakObjectPtr<UObject> >& OutObjects ) const = 0;
 
+	/**
+	 * Gets the current object(s) being customized by this builder of ObjectType 
+	 *
+	 * If this is a sub-object customization it will return those sub objects.  Otherwise the root objects will be returned.
+	 * @return true if one or more objects of ObjectType were found.
+	 */
+	template <typename ObjectType = UObject>
+	TArray< TWeakObjectPtr<ObjectType> > GetObjectsOfTypeBeingCustomized() const;
+
 	/**
 	 * Gets the current struct(s) being customized by this builder
 	 *
@@ -270,3 +286,46 @@ public:
 	 */
 	virtual bool IsPropertyPathAllowed(const FString& InPath) const = 0;
 };
+
+template <typename ObjectType>
+TArray<TWeakObjectPtr<ObjectType>> IDetailLayoutBuilder::GetSelectedObjectsOfType() const
+{
+	TArray<TWeakObjectPtr<UObject>> SelectedObjects = GetSelectedObjects();
+	TArray<TWeakObjectPtr<ObjectType>> SelectedObjectsOfType;
+	Algo::TransformIf(
+		SelectedObjects,
+		SelectedObjectsOfType,
+		[](const TWeakObjectPtr<UObject>& InObj)
+		{
+			return InObj.IsValid() && InObj->IsA(ObjectType::StaticClass());			
+		},
+		[](const TWeakObjectPtr<UObject>& InObj)
+		{
+			return Cast<ObjectType>(InObj);			
+		});
+	
+	return SelectedObjectsOfType;
+}
+
+template <typename ObjectType>
+TArray<TWeakObjectPtr<ObjectType>> IDetailLayoutBuilder::GetObjectsOfTypeBeingCustomized() const
+{
+	TArray<TWeakObjectPtr<UObject>> ObjectsBeingCustomized;
+	GetObjectsBeingCustomized(ObjectsBeingCustomized);
+
+	TArray<TWeakObjectPtr<ObjectType>> ObjectsOfTypeBeingCustomized;
+
+	Algo::TransformIf(
+		ObjectsBeingCustomized,
+		ObjectsOfTypeBeingCustomized,
+		[](const TWeakObjectPtr<UObject>& InObj)
+		{
+			return InObj.IsValid() && InObj->IsA(ObjectType::StaticClass());			
+		},
+		[](const TWeakObjectPtr<UObject>& InObj)
+		{
+			return Cast<ObjectType>(InObj);			
+		});
+
+	return ObjectsOfTypeBeingCustomized;
+}
