Require property nodes directly on a structure node to validate their data.
This fixes the case that the address pointed by a structure provider changes outside the usual editor functionality.

diff --git a/Engine/Source/Editor/PropertyEditor/Private/PropertyNode.cpp b/Engine/Source/Editor/PropertyEditor/Private/PropertyNode.cpp
--- a/Engine/Source/Editor/PropertyEditor/Private/PropertyNode.cpp
+++ b/Engine/Source/Editor/PropertyEditor/Private/PropertyNode.cpp
@@ -248,6 +248,10 @@ void FPropertyNode::InitNode(const FPropertyNodeInitParams& InitParams)
 		// We require validation if our parent also needs validation (if an array parent was resized all the addresses of children are invalid)
 		bRequiresValidation |= (GetParentNode() && GetParentNode()->HasNodeFlags(EPropertyNodeFlags::RequiresValidation));
 
+		// We require validation is we are on a structure node (the value of the structure may change externally, which invalidates the addresses).
+		const FComplexPropertyNode* ComplexParent = GetParentNode() ? GetParentNode()->AsComplexNode() : nullptr;
+		bRequiresValidation |= ComplexParent && ComplexParent->GetPropertyType() == FComplexPropertyNode::EPT_StandaloneStructure; 
+		
 		SetNodeFlags( EPropertyNodeFlags::RequiresValidation, bRequiresValidation );
 	}
 
