Property Matrix Reskin

diff --git a/Engine/Plugins/Editor/ContentBrowser/ContentBrowserAssetDataSource/Source/ContentBrowserAssetDataSource/Private/AssetFileContextMenu.cpp b/Engine/Plugins/Editor/ContentBrowser/ContentBrowserAssetDataSource/Source/ContentBrowserAssetDataSource/Private/AssetFileContextMenu.cpp
--- a/Engine/Plugins/Editor/ContentBrowser/ContentBrowserAssetDataSource/Source/ContentBrowserAssetDataSource/Private/AssetFileContextMenu.cpp
+++ b/Engine/Plugins/Editor/ContentBrowser/ContentBrowserAssetDataSource/Source/ContentBrowserAssetDataSource/Private/AssetFileContextMenu.cpp
@@ -473,7 +473,7 @@ void FAssetFileContextMenu::MakeAssetActionsSubMenu(UToolMenu* Menu)
 
 			Section.AddMenuEntry(
 				"PropertyMatrix",
-				LOCTEXT("PropertyMatrix", "Bulk Edit via Property Matrix..."),
+				LOCTEXT("PropertyMatrix", "Open Selection in Property Matrix"),
 				DynamicTooltipAttribute,
 				FSlateIcon(),
 				FUIAction(
@@ -1934,7 +1934,7 @@ FText FAssetFileContextMenu::GetExecutePropertyMatrixTooltip() const
 	FText ResultTooltip;
 	if (CanExecutePropertyMatrix(ResultTooltip))
 	{
-		ResultTooltip = LOCTEXT("PropertyMatrixTooltip", "Opens the property matrix editor for the selected assets.");
+		ResultTooltip = LOCTEXT("PropertyMatrixTooltip", "Bulk edit the selected assets in the Property Matrix");
 	}
 	return ResultTooltip;
 }
diff --git a/Engine/Source/Editor/EditorStyle/Private/StarshipStyle.cpp b/Engine/Source/Editor/EditorStyle/Private/StarshipStyle.cpp
--- a/Engine/Source/Editor/EditorStyle/Private/StarshipStyle.cpp
+++ b/Engine/Source/Editor/EditorStyle/Private/StarshipStyle.cpp
@@ -2842,53 +2842,8 @@ void FStarshipEditorStyle::FStyle::SetupPropertyEditorStyles()
 		Set( "PropertyEditor.VerticalDottedLine",		new IMAGE_BRUSH( "Common/VerticalDottedLine_1x16px", FVector2D(1.0f, 16.0f), FLinearColor::White, ESlateBrushTileType::Vertical ) );
 		Set( "PropertyEditor.SlateBrushPreview",		new BOX_BRUSH( "PropertyView/SlateBrushPreview_32px", Icon32x32, FMargin(3.f/32.f, 3.f/32.f, 15.f/32.f, 13.f/32.f) ) );
 
-		Set( "PropertyTable.TableRow", FTableRowStyle()
-			.SetEvenRowBackgroundBrush( FSlateColorBrush( FLinearColor( 0.70f, 0.70f, 0.70f, 255 ) ) )
-			.SetEvenRowBackgroundHoveredBrush( IMAGE_BRUSH( "Common/Selection", Icon8x8, SelectionColor_Inactive ) )
-			.SetOddRowBackgroundBrush( FSlateColorBrush( FLinearColor( 0.80f, 0.80f, 0.80f, 255 ) ) )
-			.SetOddRowBackgroundHoveredBrush( IMAGE_BRUSH( "Common/Selection", Icon8x8, SelectionColor_Inactive ) )
-			.SetSelectorFocusedBrush( BORDER_BRUSH( "Common/Selector", FMargin(4.f/16.f), SelectorColor ) )
-			.SetActiveBrush( IMAGE_BRUSH( "Common/Selection", Icon8x8, SelectionColor ) )
-			.SetActiveHoveredBrush( IMAGE_BRUSH( "Common/Selection", Icon8x8, SelectionColor ) )
-			.SetInactiveBrush( IMAGE_BRUSH( "Common/Selection", Icon8x8, SelectionColor_Inactive ) )
-			.SetInactiveHoveredBrush( IMAGE_BRUSH( "Common/Selection", Icon8x8, SelectionColor_Inactive ) )
-			.SetTextColor( DefaultForeground )
-			.SetSelectedTextColor( InvertedForeground )
-			);
-
-		const FTableColumnHeaderStyle PropertyTableColumnHeaderStyle = FTableColumnHeaderStyle()
-			.SetSortPrimaryAscendingImage(IMAGE_BRUSH("Common/SortUpArrow", Icon8x4))
-			.SetSortPrimaryDescendingImage(IMAGE_BRUSH("Common/SortDownArrow", Icon8x4))
-			.SetSortSecondaryAscendingImage(IMAGE_BRUSH("Common/SortUpArrows", Icon16x4))
-			.SetSortSecondaryDescendingImage(IMAGE_BRUSH("Common/SortDownArrows", Icon16x4))
-			.SetNormalBrush( BOX_BRUSH( "Common/ColumnHeader", 4.f/32.f ) )
-			.SetHoveredBrush( BOX_BRUSH( "Common/ColumnHeader_Hovered", 4.f/32.f ) )
-			.SetMenuDropdownImage( IMAGE_BRUSH( "Common/ColumnHeader_Arrow", Icon8x8 ) )
-			.SetMenuDropdownNormalBorderBrush( BOX_BRUSH( "Common/ColumnHeaderMenuButton_Normal", 4.f/32.f ) )
-			.SetMenuDropdownHoveredBorderBrush( BOX_BRUSH( "Common/ColumnHeaderMenuButton_Hovered", 4.f/32.f ) );
-
-		const FTableColumnHeaderStyle PropertyTableLastColumnHeaderStyle = FTableColumnHeaderStyle()
-			.SetSortPrimaryAscendingImage(IMAGE_BRUSH("Common/SortUpArrow", Icon8x4))
-			.SetSortPrimaryDescendingImage(IMAGE_BRUSH("Common/SortDownArrow", Icon8x4))
-			.SetSortSecondaryAscendingImage(IMAGE_BRUSH("Common/SortUpArrows", Icon16x4))
-			.SetSortSecondaryDescendingImage(IMAGE_BRUSH("Common/SortDownArrows", Icon16x4))
-			.SetNormalBrush( FSlateNoResource() )
-			.SetHoveredBrush( BOX_BRUSH( "Common/LastColumnHeader_Hovered", 4.f/32.f ) )
-			.SetMenuDropdownImage( IMAGE_BRUSH( "Common/ColumnHeader_Arrow", Icon8x8 ) )
-			.SetMenuDropdownNormalBorderBrush( BOX_BRUSH( "Common/ColumnHeaderMenuButton_Normal", 4.f/32.f ) )
-			.SetMenuDropdownHoveredBorderBrush( BOX_BRUSH( "Common/ColumnHeaderMenuButton_Hovered", 4.f/32.f ) );
-
-		const FSplitterStyle PropertyTableHeaderSplitterStyle = FSplitterStyle()
-			.SetHandleNormalBrush( FSlateNoResource() )
-			.SetHandleHighlightBrush( IMAGE_BRUSH( "Common/HeaderSplitterGrip", Icon8x8 ) );
-
-		Set( "PropertyTable.HeaderRow", FHeaderRowStyle()
-			.SetColumnStyle( PropertyTableColumnHeaderStyle )
-			.SetLastColumnStyle( PropertyTableLastColumnHeaderStyle )
-			.SetColumnSplitterStyle( PropertyTableHeaderSplitterStyle )
-			.SetBackgroundBrush( BOX_BRUSH( "Common/TableViewHeader", 4.f/32.f ) )
-			.SetForegroundColor( DefaultForeground )
-			);
+		Set( "PropertyTable.TableRow", GetWidgetStyle<FTableRowStyle>("TableView.AlternatingRow"));
+		Set( "PropertyTable.HeaderRow", GetWidgetStyle<FHeaderRowStyle>("TableView.Header"));
 
 		FWindowStyle InViewportDecoratorWindow = FStarshipCoreStyle::GetCoreStyle().GetWidgetStyle<FWindowStyle>("Window");
 		InViewportDecoratorWindow.SetCornerRadius(4);
@@ -2930,17 +2885,20 @@ void FStarshipEditorStyle::FStyle::SetupPropertyEditorStyles()
 			.SetHandleHighlightBrush(FSlateNoResource());
 		Set("PropertyTable.InViewport.Splitter", TransparentSplitterStyle);
 
-		Set( "PropertyTable.Selection.Active",						new IMAGE_BRUSH( "Common/Selection", Icon8x8, SelectionColor ) );
+		float BorderPadding = 0.5f;
+		Set( "PropertyTable.CellBorder",	new FSlateRoundedBoxBrush(FStyleColors::Transparent, 0.0f, FStyleColors::Background, BorderPadding) );
+		Set( "PropertyTable.CurrentCellBorder",						new FSlateRoundedBoxBrush(FStyleColors::Primary, 0.0f, FStyleColors::White, BorderPadding) );
+		Set( "PropertyTable.SelectedCellBorder",						new FSlateRoundedBoxBrush(FStyleColors::Primary, 0.0f, FStyleColors::Background, BorderPadding) );
+		Set( "PropertyTable.EditModeCellBorder",						new FSlateRoundedBoxBrush(FStyleColors::Primary, 0.0f, FStyleColors::Background, BorderPadding) );
 
-		Set( "PropertyTable.HeaderRow.Column.PathDelimiter",		new IMAGE_BRUSH( "Common/SmallArrowRight", Icon10x10 ) );
+		Set( "PropertyTable.Selection.Active",						new IMAGE_BRUSH( "Common/Selector", Icon8x8, SelectionColor ) );
 
+		Set( "PropertyTable.HeaderRow.Column.PathDelimiter",		new IMAGE_BRUSH( "Common/SmallArrowRight", Icon10x10 ) );
+		Set( "PropertyTable.ColumnBorder",	new BOX_BRUSH( "Common/CellBorder", FMargin(4.f/16.f), FStyleColors::Background )  );
 		Set( "PropertyTable.RowHeader.Background",					new BOX_BRUSH( "Old/Menu_Background", FMargin(4.f/64.f) ) );
 		Set( "PropertyTable.RowHeader.BackgroundActive",			new BOX_BRUSH( "Old/Menu_Background", FMargin(4.f/64.f), SelectionColor_Inactive ) );
-		Set( "PropertyTable.ColumnBorder",							new BOX_BRUSH( "Common/ColumnBorder", FMargin(4.f/16.f), FLinearColor(0.1f, 0.1f, 0.1f, 0.5f) ) );
-		Set( "PropertyTable.CellBorder",							new BOX_BRUSH( "Common/CellBorder", FMargin(4.f/16.f), FLinearColor(0.1f, 0.1f, 0.1f, 0.5f) ) );
 		Set( "PropertyTable.ReadOnlyEditModeCellBorder",			new BORDER_BRUSH( "Common/ReadOnlyEditModeCellBorder", FMargin(6.f/32.f), SelectionColor ) );
 		Set( "PropertyTable.ReadOnlyCellBorder",					new BOX_BRUSH( "Common/ReadOnlyCellBorder", FMargin(4.f/16.f), FLinearColor(0.1f, 0.1f, 0.1f, 0.5f) ) );
-		Set( "PropertyTable.CurrentCellBorder",						new BOX_BRUSH( "Common/CurrentCellBorder", FMargin(4.f/16.f), FLinearColor(0.0f, 0.0f, 0.0f, 1.0f) ) );
 		Set( "PropertyTable.ReadOnlySelectedCellBorder",			new BOX_BRUSH( "Common/ReadOnlySelectedCellBorder", FMargin(4.f/16.f), FLinearColor(0.0f, 0.0f, 0.0f, 1.0f) ) );
 		Set( "PropertyTable.ReadOnlyCurrentCellBorder",				new BOX_BRUSH( "Common/ReadOnlyCurrentCellBorder", FMargin(4.f/16.f), FLinearColor(0.0f, 0.0f, 0.0f, 1.0f) ) );
 		Set( "PropertyTable.Cell.DropDown.Background",				new BOX_BRUSH( "Common/GroupBorder", FMargin(4.f/16.f) ) );
@@ -2949,6 +2907,12 @@ void FStarshipEditorStyle::FStyle::SetupPropertyEditorStyles()
 		Set( "PropertyTable.BoldFont",								DEFAULT_FONT( "Bold", 9 ) );
 		Set( "PropertyTable.FilterFont",							DEFAULT_FONT( "Regular", 10 ) );
 
+		const FTableRowStyle PropertyEditorPropertyRowStyle = FTableRowStyle(NormalTableRowStyle)
+			.SetEvenRowBackgroundBrush(FSlateColorBrush(FStyleColors::Panel))
+			.SetOddRowBackgroundBrush(FSlateColorBrush(FStyleColors::Panel));
+
+		Set( "PropertyWindow.PropertyRow", PropertyEditorPropertyRowStyle);
+
 		Set( "PropertyWindow.FilterSearch", new IMAGE_BRUSH( "Old/FilterSearch", Icon16x16 ) );
 		Set( "PropertyWindow.FilterCancel", new IMAGE_BRUSH( "Old/FilterCancel", Icon16x16 ) );
 		Set( "PropertyWindow.Favorites_Disabled", new IMAGE_BRUSH( "Icons/EmptyStar_16x", Icon16x16 ) );
diff --git a/Engine/Source/Editor/LevelEditor/Private/LevelEditorActions.cpp b/Engine/Source/Editor/LevelEditor/Private/LevelEditorActions.cpp
--- a/Engine/Source/Editor/LevelEditor/Private/LevelEditorActions.cpp
+++ b/Engine/Source/Editor/LevelEditor/Private/LevelEditorActions.cpp
@@ -1417,6 +1417,20 @@ bool FLevelEditorActionCallbacks::EditAsset_CanExecute()
 	return false;
 }
 
+void FLevelEditorActionCallbacks::OpenSelectionInPropertyMatrix_Clicked()
+{
+	TArray<UObject*> SelectedObjects;
+	GEditor->GetSelectedActors()->GetSelectedObjects<UObject>(SelectedObjects);
+
+	FPropertyEditorModule& PropertyEditorModule = FModuleManager::LoadModuleChecked<FPropertyEditorModule>( "PropertyEditor" );
+	PropertyEditorModule.CreatePropertyEditorToolkit(TSharedPtr<IToolkitHost>(), SelectedObjects );
+}
+
+bool FLevelEditorActionCallbacks::OpenSelectionInPropertyMatrix_IsVisible()
+{
+	return FModuleManager::LoadModuleChecked<FPropertyEditorModule>( "PropertyEditor" ).GetCanUsePropertyMatrix() && GEditor->GetSelectedActorCount() > 1;
+}
+
 void FLevelEditorActionCallbacks::LockActorMovement_Clicked()
 {
 	GEditor->ToggleSelectedActorMovementLock();
@@ -3652,6 +3666,7 @@ void FLevelEditorCommands::RegisterCommands()
 
 	UI_COMMAND( EditAsset, "Edit Asset", "Edits the asset associated with the selected actor", EUserInterfaceActionType::Button, FInputChord( EKeys::E, EModifierKey::Control ) );
 	UI_COMMAND( EditAssetNoConfirmMultiple, "Edit Multiple Assets", "Edits multiple assets associated with the selected actor without a confirmation prompt", EUserInterfaceActionType::Button, FInputChord( EKeys::E, EModifierKey::Control | EModifierKey::Shift ) );
+	UI_COMMAND( OpenSelectionInPropertyMatrix, "Open Selection in Property Matrix", "Bulk edit the selected assets in the Property Matrix", EUserInterfaceActionType::Button, FInputChord() );
 
 	UI_COMMAND( GoHere, "Go Here", "Moves the camera to the current mouse position", EUserInterfaceActionType::Button, FInputChord() );
 
diff --git a/Engine/Source/Editor/LevelEditor/Private/LevelEditorContextMenu.cpp b/Engine/Source/Editor/LevelEditor/Private/LevelEditorContextMenu.cpp
--- a/Engine/Source/Editor/LevelEditor/Private/LevelEditorContextMenu.cpp
+++ b/Engine/Source/Editor/LevelEditor/Private/LevelEditorContextMenu.cpp
@@ -322,6 +322,13 @@ void FLevelEditorContextMenu::RegisterActorContextMenu()
 				}
 			}
 
+			Section.AddMenuEntry(
+				FLevelEditorCommands::Get().OpenSelectionInPropertyMatrix,
+				TAttribute<FText>(), // use command's label
+				TAttribute<FText>(), // use command's tooltip
+				FSlateIcon(FAppStyle::GetAppStyleSetName(), "DetailsView.EditRawProperties")
+				);
+
 			if (LevelEditorContext->ContextType == ELevelEditorMenuContext::MainMenu)
 			{
 				Section.AddSubMenu(
diff --git a/Engine/Source/Editor/LevelEditor/Private/SLevelEditor.cpp b/Engine/Source/Editor/LevelEditor/Private/SLevelEditor.cpp
--- a/Engine/Source/Editor/LevelEditor/Private/SLevelEditor.cpp
+++ b/Engine/Source/Editor/LevelEditor/Private/SLevelEditor.cpp
@@ -111,10 +111,17 @@ void SLevelEditor::BindCommands()
 		FExecuteAction::CreateStatic( &FLevelEditorActionCallbacks::EditAsset_Clicked, EToolkitMode::Standalone, TWeakPtr< SLevelEditor >( SharedThis( this ) ), false ),
 		FCanExecuteAction::CreateStatic( &FLevelEditorActionCallbacks::EditAsset_CanExecute ) );
 
+	LevelEditorCommands->MapAction( 
+		Actions.EditAssetNoConfirmMultiple, 
+		FExecuteAction::CreateStatic( &FLevelEditorActionCallbacks::EditAsset_Clicked, EToolkitMode::Standalone, TWeakPtr< SLevelEditor >( SharedThis( this ) ), false ),
+		FCanExecuteAction::CreateStatic( &FLevelEditorActionCallbacks::EditAsset_CanExecute ) );
+	
 	LevelEditorCommands->MapAction(
-		Actions.EditAsset,
-		FExecuteAction::CreateStatic( &FLevelEditorActionCallbacks::EditAsset_Clicked, EToolkitMode::Standalone, TWeakPtr< SLevelEditor >( SharedThis( this ) ), true ),
-		FCanExecuteAction::CreateStatic(&FLevelEditorActionCallbacks::EditAsset_CanExecute));
+		Actions.OpenSelectionInPropertyMatrix,
+		FExecuteAction::CreateStatic( &FLevelEditorActionCallbacks::OpenSelectionInPropertyMatrix_Clicked ),
+		FCanExecuteAction(),
+		FIsActionChecked(),
+		FIsActionButtonVisible::CreateStatic(&FLevelEditorActionCallbacks::OpenSelectionInPropertyMatrix_IsVisible));
 
 	LevelEditorCommands->MapAction(
 		Actions.CheckOutProjectSettingsConfig,
diff --git a/Engine/Source/Editor/LevelEditor/Public/LevelEditorActions.h b/Engine/Source/Editor/LevelEditor/Public/LevelEditorActions.h
--- a/Engine/Source/Editor/LevelEditor/Public/LevelEditorActions.h
+++ b/Engine/Source/Editor/LevelEditor/Public/LevelEditorActions.h
@@ -127,6 +127,9 @@ public:
 	/** Edits associated asset(s) */
 	TSharedPtr< FUICommandInfo > EditAssetNoConfirmMultiple;
 
+	/** Opens the associated asset(s) in the property matrix */
+	TSharedPtr< FUICommandInfo > OpenSelectionInPropertyMatrix;
+
 	/** Moves the camera to the current mouse position */
 	TSharedPtr< FUICommandInfo > GoHere;
 
@@ -913,6 +916,11 @@ public:
 	static void EditAsset_Clicked( const EToolkitMode::Type ToolkitMode, TWeakPtr< class SLevelEditor > LevelEditor, bool bAskMultiple );
 	static bool EditAsset_CanExecute();
 
+	/** Called to when "Open Selection in Property Matrix" is clicked */
+	static void OpenSelectionInPropertyMatrix_Clicked();
+	static bool OpenSelectionInPropertyMatrix_IsVisible();
+
+
 	/** Called when 'detach' is clicked */
 	static void DetachActor_Clicked();
 	static bool DetachActor_CanExecute();
diff --git a/Engine/Source/Editor/PropertyEditor/Private/Presentation/PropertyTable/PropertyTableRowHeaderColumn.h b/Engine/Source/Editor/PropertyEditor/Private/Presentation/PropertyTable/PropertyTableRowHeaderColumn.h
--- a/Engine/Source/Editor/PropertyEditor/Private/Presentation/PropertyTable/PropertyTableRowHeaderColumn.h
+++ b/Engine/Source/Editor/PropertyEditor/Private/Presentation/PropertyTable/PropertyTableRowHeaderColumn.h
@@ -60,7 +60,7 @@ public:
 
 	virtual void SetSizeMode(EPropertyTableColumnSizeMode::Type InSizeMode) override {}
 
-	virtual float GetWidth() const override { return 20.0f; } 
+	virtual float GetWidth() const override { return 24.0f; } 
 
 	virtual void SetWidth( float InWidth ) override {  }
 
diff --git a/Engine/Source/Editor/PropertyEditor/Private/PropertyEditorToolkit.cpp b/Engine/Source/Editor/PropertyEditor/Private/PropertyEditorToolkit.cpp
--- a/Engine/Source/Editor/PropertyEditor/Private/PropertyEditorToolkit.cpp
+++ b/Engine/Source/Editor/PropertyEditor/Private/PropertyEditorToolkit.cpp
@@ -11,6 +11,7 @@
 #include "IPropertyTableColumn.h"
 #include "IPropertyTreeRow.h"
 #include "IPropertyTableRow.h"
+#include "SSimpleButton.h"
 
 #include "Widgets/Docking/SDockTab.h"
 #include "Subsystems/AssetEditorSubsystem.h"
@@ -35,12 +36,12 @@ void FPropertyEditorToolkit::RegisterTabSpawners(const TSharedRef<class FTabMana
 		.SetIcon(FSlateIcon(FAppStyle::GetAppStyleSetName(), "LevelEditor.Tabs.Viewports"));
 
 	InTabManager->RegisterTabSpawner( TreeTabId, FOnSpawnTab::CreateSP(this, &FPropertyEditorToolkit::SpawnTab_PropertyTree) )
-		.SetDisplayName( LOCTEXT("PropertiesTab", "Display") )
+		.SetDisplayName( LOCTEXT("PropertiesTab", "Pinned Columns") )
 		.SetGroup( WorkspaceMenuCategory.ToSharedRef() )
 		.SetIcon(FSlateIcon(FAppStyle::GetAppStyleSetName(), "PropertyEditor.Grid.TabIcon"));
 
 	InTabManager->RegisterTabSpawner(DetailsTabId, FOnSpawnTab::CreateSP(this, &FPropertyEditorToolkit::SpawnTab_DetailsPanel))
-		.SetDisplayName(LOCTEXT("DetailsTab", "Property Editor"))
+		.SetDisplayName(LOCTEXT("DetailsTab", "Details"))
 		.SetGroup(WorkspaceMenuCategory.ToSharedRef())
 		.SetIcon(FSlateIcon(FAppStyle::GetAppStyleSetName(), "LevelEditor.Tabs.Details"));
 }
@@ -150,20 +151,21 @@ void FPropertyEditorToolkit::Initialize( const EToolkitMode::Type Mode, const TS
 		PropertyTable->SetObjects(AdjustedObjectsToEdit);
 		TableColumnsChanged();
 
-		TSharedRef<FTabManager::FLayout> StandaloneDefaultLayout = FTabManager::NewLayout("Standalone_PropertyEditorToolkit_Layout_v1")
+		TSharedRef<FTabManager::FLayout> StandaloneDefaultLayout = FTabManager::NewLayout("Standalone_PropertyEditorToolkit_Layout_v1.1")
 			->AddArea
 			(
 				FTabManager::NewPrimaryArea()->SetOrientation(Orient_Horizontal)
 				->Split
 				(
 					FTabManager::NewStack()
-					->SetSizeCoefficient(0.8f)
+					->SetSizeCoefficient(0.65f)
+					->SetHideTabWell(true)
 					->AddTab(GridTabId, ETabState::OpenedTab)
 				)
 				->Split
 				(
 					FTabManager::NewStack()
-					->SetSizeCoefficient(0.2f)
+					->SetSizeCoefficient(0.35f)
 					->AddTab(TreeTabId, ETabState::OpenedTab)
 					->AddTab(DetailsTabId, ETabState::OpenedTab)
 					->SetForegroundTab(TreeTabId)
@@ -171,7 +173,7 @@ void FPropertyEditorToolkit::Initialize( const EToolkitMode::Type Mode, const TS
 			);
 
 		const bool bCreateDefaultStandaloneMenu = true;
-		const bool bCreateDefaultToolbar = false;
+		const bool bCreateDefaultToolbar = true;
 		FAssetEditorToolkit::InitAssetEditor(Mode, InitToolkitHost, ApplicationId, StandaloneDefaultLayout, bCreateDefaultStandaloneMenu, bCreateDefaultToolbar, AdjustedObjectsToEdit);
 
 		TArray<UObject*> AdjustedObjectsToEditWeak;
@@ -192,7 +194,7 @@ TSharedRef<SDockTab> FPropertyEditorToolkit::SpawnTab_PropertyTree( const FSpawn
 	check( Args.GetTabId() == TreeTabId );
 
 	TSharedRef<SDockTab> TreeToolkitTab = SNew(SDockTab)
-		.Label( LOCTEXT("GenericDetailsTitle", "Display") )
+		.Label( LOCTEXT("GenericDetailsTitle", "Pinned Columns") )
 		.TabColorScale( GetTabColorScale() )
 		.Content()
 		[
@@ -219,50 +221,7 @@ TSharedRef<SDockTab> FPropertyEditorToolkit::SpawnTab_PropertyTable( const FSpaw
 		.TabColorScale( GetTabColorScale() )
 		.Content()
 		[
-			SNew( SOverlay )
-			+SOverlay::Slot()
-			[
-				PropertyEditorModule.CreatePropertyTableWidget( PropertyTable.ToSharedRef() )
-			]
-			+SOverlay::Slot()
-			.HAlign( HAlign_Right )
-			.VAlign( VAlign_Top )
-			.Padding( FMargin( 0, 3, 0, 0 ) )
-			[
-				SNew( SHorizontalBox )
-				+SHorizontalBox::Slot()
-				.AutoWidth()
-				.HAlign( HAlign_Center )
-				.VAlign( VAlign_Center )
-				[
-					SNew( SImage )
-					.Image( FAppStyle::GetBrush( "PropertyEditor.AddColumnOverlay" ) )
-					.Visibility( this, &FPropertyEditorToolkit::GetAddColumnInstructionsOverlayVisibility )
-				]
-
-				+SHorizontalBox::Slot()
-				.AutoWidth()
-				.HAlign( HAlign_Center )
-				.VAlign( VAlign_Center )
-				[
-					SNew( SImage )
-					.Image( FAppStyle::GetBrush( "Icons.Unpinned" ) )
-					.Visibility( this, &FPropertyEditorToolkit::GetAddColumnInstructionsOverlayVisibility )
-				]
-
-				+SHorizontalBox::Slot()
-				.AutoWidth()
-				.HAlign( HAlign_Center )
-				.VAlign( VAlign_Center )
-				.Padding( FMargin( 0, 0, 3, 0 ) )
-				[
-					SNew( STextBlock )
-					.Font( FAppStyle::GetFontStyle( "PropertyEditor.AddColumnMessage.Font" ) )
-					.Text( LOCTEXT("GenericPropertiesTitle", "Pin Properties to Add Columns") )
-					.Visibility( this, &FPropertyEditorToolkit::GetAddColumnInstructionsOverlayVisibility )
-					.ColorAndOpacity( FAppStyle::GetColor( "PropertyEditor.AddColumnMessage.ColorAndOpacity" ) )
-				]
-			]
+			PropertyEditorModule.CreatePropertyTableWidget( PropertyTable.ToSharedRef() )
 		];	
 
 	return GridToolkitTab;
@@ -273,7 +232,7 @@ TSharedRef<SDockTab> FPropertyEditorToolkit::SpawnTab_DetailsPanel(const FSpawnT
 	check(Args.GetTabId() == DetailsTabId);
 
 	TSharedPtr<SDockTab> DetailsTab = SNew(SDockTab)
-		.Label(LOCTEXT("GenericDetailsPanel", "Property Editor"))
+		.Label(LOCTEXT("GenericDetailsPanel", "Details"))
 		[
 			DetailsView.ToSharedRef()
 		];
@@ -342,35 +301,17 @@ TSharedRef< SWidget > FPropertyEditorToolkit::ConstructTreeCell( const FName& Co
 		const TWeakPtr<IPropertyTreeRow> RowPtr = Row;
 		PinRows.Add( Row );
 
-		return SNew( SBorder )
-			.Padding( 0 )
-			.BorderImage( &FAppStyle::GetWidgetStyle<FHeaderRowStyle>("PropertyTable.HeaderRow").ColumnStyle.NormalBrush )
-			[
-				SNew(SButton)
-				.ToolTipText(NSLOCTEXT("PropertyEditor", "ToggleColumnButtonToolTip", "Toggle Column"))
-				.ButtonStyle( FAppStyle::Get(), "HoverOnlyButton" )
-				.ContentPadding(0) 
-				.OnClicked( this, &FPropertyEditorToolkit::OnToggleColumnClicked, RowPtr )
-				.HAlign( HAlign_Center )
-				.VAlign( VAlign_Center )
-				[
-					SNew(SImage)
-					.Image( this, &FPropertyEditorToolkit::GetToggleColumnButtonImageBrush, RowPtr )
-					.ColorAndOpacity(FSlateColor::UseForeground())
-				]
-			];
+		return
+			SNew(SSimpleButton)
+			.ToolTipText(NSLOCTEXT("PropertyEditor", "ToggleColumnButtonToolTip", "Toggle Column"))
+			.OnClicked( this, &FPropertyEditorToolkit::OnToggleColumnClicked, RowPtr )
+			.Icon(this, &FPropertyEditorToolkit::GetToggleColumnButtonImageBrush, RowPtr);
+
 	}
 
 	return SNullWidget::NullWidget;
 }
 
-
-EVisibility FPropertyEditorToolkit::GetAddColumnInstructionsOverlayVisibility() const
-{
-	return TableHasCustomColumns() ? EVisibility::Collapsed : EVisibility::HitTestInvisible;
-}
-
-
 void FPropertyEditorToolkit::ToggleColumnForProperty( const TSharedPtr< FPropertyPath >& PropertyPath )
 {
 	if ( !PropertyPath.IsValid() )
diff --git a/Engine/Source/Editor/PropertyEditor/Private/PropertyEditorToolkit.h b/Engine/Source/Editor/PropertyEditor/Private/PropertyEditorToolkit.h
--- a/Engine/Source/Editor/PropertyEditor/Private/PropertyEditorToolkit.h
+++ b/Engine/Source/Editor/PropertyEditor/Private/PropertyEditorToolkit.h
@@ -51,6 +51,13 @@ public:
 	static TSharedRef<FPropertyEditorToolkit> CreateEditor( const EToolkitMode::Type Mode, const TSharedPtr< IToolkitHost >& InitToolkitHost, UObject* ObjectToEdit );
 
 	static TSharedRef<FPropertyEditorToolkit> CreateEditor( const EToolkitMode::Type Mode, const TSharedPtr< IToolkitHost >& InitToolkitHost, const TArray<UObject*>& ObjectsToEdit );
+
+
+protected:
+
+	/** "Find in Content Browser" is not visible in the property matrix because it only works on assets (and not on actors) */
+	virtual bool IsFindInContentBrowserButtonVisible() const override { return false; }
+
 private:
 	static TSharedPtr<FPropertyEditorToolkit> FindExistingEditor( UObject* Object );
 
@@ -86,8 +93,6 @@ private:
 
 	void TableColumnsChanged();
 
-	EVisibility GetAddColumnInstructionsOverlayVisibility() const;
-
 private:
 
 	TSharedPtr< class SPropertyTreeViewImpl > PropertyTree;
diff --git a/Engine/Source/Editor/PropertyEditor/Private/SPropertyTreeViewImpl.cpp b/Engine/Source/Editor/PropertyEditor/Private/SPropertyTreeViewImpl.cpp
--- a/Engine/Source/Editor/PropertyEditor/Private/SPropertyTreeViewImpl.cpp
+++ b/Engine/Source/Editor/PropertyEditor/Private/SPropertyTreeViewImpl.cpp
@@ -235,6 +235,7 @@ void SPropertyTreeViewImpl::ConstructPropertyTree()
 				[
 					SNew( SImage )
 					.Image( this, &SPropertyTreeViewImpl::OnGetLockButtonImageResource )
+					.ColorAndOpacity(FSlateColor::UseForeground())
 				]
 			]
 		]
diff --git a/Engine/Source/Editor/PropertyEditor/Private/UserInterface/PropertyEditor/SPropertyEditorColor.cpp b/Engine/Source/Editor/PropertyEditor/Private/UserInterface/PropertyEditor/SPropertyEditorColor.cpp
--- a/Engine/Source/Editor/PropertyEditor/Private/UserInterface/PropertyEditor/SPropertyEditorColor.cpp
+++ b/Engine/Source/Editor/PropertyEditor/Private/UserInterface/PropertyEditor/SPropertyEditorColor.cpp
@@ -49,7 +49,9 @@ TSharedRef< SColorBlock > SPropertyEditorColor::ConstructColorBlock()
 		.ShowBackgroundForAlpha(true)
 		.AlphaDisplayMode(bIgnoreAlpha ? EColorBlockAlphaDisplayMode::Ignore : EColorBlockAlphaDisplayMode::Combined)
 		.OnMouseButtonDown( this, &SPropertyEditorColor::ColorBlock_OnMouseButtonDown )
-		.Size( FVector2D( 10.0f, 10.0f ) );
+		.Size( FVector2D( 120.0f, 20.0f ) )
+		.CornerRadius(FVector4(4.0f,4.0f,4.0f,4.0f));
+
 }
 
 TSharedRef< SColorBlock > SPropertyEditorColor::ConstructAlphaColorBlock()
@@ -61,7 +63,8 @@ TSharedRef< SColorBlock > SPropertyEditorColor::ConstructAlphaColorBlock()
 		.AlphaDisplayMode(EColorBlockAlphaDisplayMode::Ignore)
 		.Visibility( this, &SPropertyEditorColor::GetVisibilityForOpaqueDisplay )
 		.OnMouseButtonDown( this, &SPropertyEditorColor::ColorBlock_OnMouseButtonDown )
-		.Size( FVector2D( 10.0f, 10.0f ) );
+		.Size( FVector2D( 120.0f, 20.0f ) )
+		.CornerRadius(FVector4(4.0f,4.0f,4.0f,4.0f));
 }
 
 bool SPropertyEditorColor::ShouldDisplayIgnoreAlpha() const
diff --git a/Engine/Source/Editor/PropertyEditor/Private/UserInterface/PropertyEditor/SPropertyEditorTableRow.cpp b/Engine/Source/Editor/PropertyEditor/Private/UserInterface/PropertyEditor/SPropertyEditorTableRow.cpp
--- a/Engine/Source/Editor/PropertyEditor/Private/UserInterface/PropertyEditor/SPropertyEditorTableRow.cpp
+++ b/Engine/Source/Editor/PropertyEditor/Private/UserInterface/PropertyEditor/SPropertyEditorTableRow.cpp
@@ -35,7 +35,7 @@ void SPropertyEditorTableRow::Construct( const FArguments& InArgs, const TShared
 
 	this->SetToolTipText(PropertyEditor->GetToolTipText());
 	this->SetVisibility(TAttribute<EVisibility>::CreateSP(this, &SPropertyEditorTableRow::OnGetRowVisibility));
-	SMultiColumnTableRow< TSharedPtr<FPropertyNode*> >::Construct( FSuperRowType::FArguments(), InOwnerTable );
+	SMultiColumnTableRow< TSharedPtr<FPropertyNode*> >::Construct( FSuperRowType::FArguments().Style(&FAppStyle::GetWidgetStyle<FTableRowStyle>("PropertyWindow.PropertyRow")), InOwnerTable );
 }
 
 TSharedRef<SWidget> SPropertyEditorTableRow::GenerateWidgetForColumn( const FName& ColumnName )
@@ -108,9 +108,16 @@ TSharedRef< SWidget > SPropertyEditorTableRow::ConstructValueColumnWidget()
 
 	HorizontalBox->AddSlot()
 		.FillWidth(1) // Fill the entire width if possible
-		.VAlign(VAlign_Center)	
+		.VAlign(VAlign_Center)
+		.HAlign(HAlign_Left)
+		.Padding(20, 0, 0, 0)
 		[
-			ValueEditorWidget.ToSharedRef()
+			SNew(SBox)
+			.MinDesiredWidth( MinWidth )
+			.MaxDesiredWidth( MaxWidth )
+			[
+				ValueEditorWidget.ToSharedRef()
+			]
 		];
 
 	// The favorites star for this property
@@ -236,79 +243,106 @@ TSharedRef<SWidget> SPropertyEditorTableRow::ConstructPropertyEditorWidget()
 	const TSharedRef< FPropertyEditor > PropertyEditorRef = PropertyEditor.ToSharedRef();
 	const TSharedRef< IPropertyUtilities > PropertyUtilitiesRef = PropertyUtilities.ToSharedRef();
 
+	float MinDesiredWidth = 0.0f;
+	float MaxDesiredWidth = 0.0f;
+
 	if( Property )
 	{
 		// ORDER MATTERS: first widget type to support the property node wins!
 		if ( SPropertyEditorNumeric<float>::Supports(PropertyEditorRef) )
 		{
-			PropertyWidget = SNew( SPropertyEditorNumeric<float>, PropertyEditorRef );
+			TSharedRef<SPropertyEditorNumeric<float>> TempWidget = SAssignNew(PropertyWidget, SPropertyEditorNumeric<float>, PropertyEditorRef );
+			TempWidget->GetDesiredWidth(MinDesiredWidth, MaxDesiredWidth);
 		}
 		else if (SPropertyEditorNumeric<int8>::Supports(PropertyEditorRef))
 		{
-			PropertyWidget = SNew(SPropertyEditorNumeric<int8>, PropertyEditorRef);
+			TSharedRef<SPropertyEditorNumeric<int8>> TempWidget = SAssignNew(PropertyWidget, SPropertyEditorNumeric<int8>, PropertyEditorRef );
+			TempWidget->GetDesiredWidth(MinDesiredWidth, MaxDesiredWidth);
 		}
 		else if (SPropertyEditorNumeric<int16>::Supports(PropertyEditorRef))
 		{
-			PropertyWidget = SNew(SPropertyEditorNumeric<int16>, PropertyEditorRef);
+			TSharedRef<SPropertyEditorNumeric<int16>> TempWidget = SAssignNew(PropertyWidget, SPropertyEditorNumeric<int16>, PropertyEditorRef );
+			TempWidget->GetDesiredWidth(MinDesiredWidth, MaxDesiredWidth);
 		}
 		else if (SPropertyEditorNumeric<int32>::Supports(PropertyEditorRef))
 		{
-			PropertyWidget = SNew(SPropertyEditorNumeric<int32>, PropertyEditorRef);
+			TSharedRef<SPropertyEditorNumeric<int32>> TempWidget = SAssignNew(PropertyWidget, SPropertyEditorNumeric<int32>, PropertyEditorRef );
+			TempWidget->GetDesiredWidth(MinDesiredWidth, MaxDesiredWidth);
 		}
 		else if ( SPropertyEditorNumeric<int64>::Supports(PropertyEditorRef) )
 		{
-			PropertyWidget = SNew( SPropertyEditorNumeric<int64>, PropertyEditorRef );
+			TSharedRef<SPropertyEditorNumeric<int64>> TempWidget = SAssignNew(PropertyWidget, SPropertyEditorNumeric<int64>, PropertyEditorRef );
+			TempWidget->GetDesiredWidth(MinDesiredWidth, MaxDesiredWidth);
 		}
 		else if ( SPropertyEditorNumeric<uint8>::Supports(PropertyEditorRef) )
 		{
-			PropertyWidget = SNew( SPropertyEditorNumeric<uint8>, PropertyEditorRef );
+			TSharedRef<SPropertyEditorNumeric<uint8>> TempWidget = SAssignNew(PropertyWidget, SPropertyEditorNumeric<uint8>, PropertyEditorRef );
+			TempWidget->GetDesiredWidth(MinDesiredWidth, MaxDesiredWidth);
 		}
 		else if (SPropertyEditorNumeric<uint16>::Supports(PropertyEditorRef))
 		{
-			PropertyWidget = SNew(SPropertyEditorNumeric<uint16>, PropertyEditorRef);
+			TSharedRef<SPropertyEditorNumeric<uint16>> TempWidget = SAssignNew(PropertyWidget, SPropertyEditorNumeric<uint16>, PropertyEditorRef );
+			TempWidget->GetDesiredWidth(MinDesiredWidth, MaxDesiredWidth);
 		}
 		else if (SPropertyEditorNumeric<uint32>::Supports(PropertyEditorRef))
 		{
-			PropertyWidget = SNew(SPropertyEditorNumeric<uint32>, PropertyEditorRef);
+			TSharedRef<SPropertyEditorNumeric<uint32>> TempWidget = SAssignNew(PropertyWidget, SPropertyEditorNumeric<uint32>, PropertyEditorRef );
+			TempWidget->GetDesiredWidth(MinDesiredWidth, MaxDesiredWidth);
 		}
 		else if (SPropertyEditorNumeric<uint64>::Supports(PropertyEditorRef))
 		{
-			PropertyWidget = SNew(SPropertyEditorNumeric<uint64>, PropertyEditorRef);
+			TSharedRef<SPropertyEditorNumeric<uint64>> TempWidget = SAssignNew(PropertyWidget, SPropertyEditorNumeric<uint64>, PropertyEditorRef );
+			TempWidget->GetDesiredWidth(MinDesiredWidth, MaxDesiredWidth);
 		}
 		else if ( SPropertyEditorArray::Supports(PropertyEditorRef) )
 		{
-			PropertyWidget = SNew( SPropertyEditorArray, PropertyEditorRef );
+			TSharedRef<SPropertyEditorArray> TempWidget = SAssignNew(PropertyWidget, SPropertyEditorArray, PropertyEditorRef );
+			TempWidget->GetDesiredWidth(MinDesiredWidth, MaxDesiredWidth);
 		}
 		else if ( SPropertyEditorCombo::Supports(PropertyEditorRef) )
 		{
-			PropertyWidget = SNew( SPropertyEditorCombo, PropertyEditorRef );
+			TSharedRef<SPropertyEditorCombo> TempWidget = SAssignNew(PropertyWidget, SPropertyEditorCombo, PropertyEditorRef );
+			TempWidget->GetDesiredWidth(MinDesiredWidth, MaxDesiredWidth);
 		}
 		else if ( SPropertyEditorEditInline::Supports(PropertyEditorRef) )
 		{
-			PropertyWidget = SNew( SPropertyEditorEditInline, PropertyEditorRef );
+			TSharedRef<SPropertyEditorEditInline> TempWidget = SAssignNew(PropertyWidget, SPropertyEditorEditInline, PropertyEditorRef );
+			TempWidget->GetDesiredWidth(MinDesiredWidth, MaxDesiredWidth);
 		}
 		else if ( SPropertyEditorText::Supports(PropertyEditorRef) )
 		{
-			PropertyWidget = SNew( SPropertyEditorText, PropertyEditorRef );
+			TSharedRef<SPropertyEditorText> TempWidget = SAssignNew(PropertyWidget, SPropertyEditorText, PropertyEditorRef );
+			TempWidget->GetDesiredWidth(MinDesiredWidth, MaxDesiredWidth);
 		}
 		else if ( SPropertyEditorBool::Supports(PropertyEditorRef) )
 		{
-			PropertyWidget = SNew( SPropertyEditorBool, PropertyEditorRef );
+			TSharedRef<SPropertyEditorBool> TempWidget = SAssignNew(PropertyWidget, SPropertyEditorBool, PropertyEditorRef );
+			TempWidget->GetDesiredWidth(MinDesiredWidth, MaxDesiredWidth);
 		}
 		else if ( SPropertyEditorColor::Supports(PropertyEditorRef) )
 		{
-			PropertyWidget = SNew( SPropertyEditorColor, PropertyEditorRef, PropertyUtilitiesRef );
+			TSharedRef<SPropertyEditorColor> TempWidget = SAssignNew(PropertyWidget, SPropertyEditorColor, PropertyEditorRef, PropertyUtilitiesRef );
 		}
 		else if ( SPropertyEditorArrayItem::Supports(PropertyEditorRef) )
 		{
-			PropertyWidget = SNew( SPropertyEditorArrayItem, PropertyEditorRef );
+			TSharedRef<SPropertyEditorArrayItem> TempWidget = SAssignNew(PropertyWidget, SPropertyEditorArrayItem, PropertyEditorRef );
+			TempWidget->GetDesiredWidth(MinDesiredWidth, MaxDesiredWidth);
 		}
 		else if ( SPropertyEditorDateTime::Supports(PropertyEditorRef) )
 		{
-			PropertyWidget = SNew( SPropertyEditorDateTime, PropertyEditorRef );
+			TSharedRef<SPropertyEditorDateTime> TempWidget = SAssignNew(PropertyWidget, SPropertyEditorDateTime, PropertyEditorRef );
 		}
 	}
 
+	if(MinDesiredWidth != 0.0f)
+	{
+		MinWidth = MinDesiredWidth;
+	}
+	if(MaxDesiredWidth != 0.0f)
+	{
+		MaxWidth = MaxDesiredWidth;
+	}
+
 	if( !PropertyWidget.IsValid() )
 	{
 		PropertyWidget = SNew( SPropertyEditor, PropertyEditorRef );
diff --git a/Engine/Source/Editor/PropertyEditor/Private/UserInterface/PropertyEditor/SPropertyEditorTableRow.h b/Engine/Source/Editor/PropertyEditor/Private/UserInterface/PropertyEditor/SPropertyEditorTableRow.h
--- a/Engine/Source/Editor/PropertyEditor/Private/UserInterface/PropertyEditor/SPropertyEditorTableRow.h
+++ b/Engine/Source/Editor/PropertyEditor/Private/UserInterface/PropertyEditor/SPropertyEditorTableRow.h
@@ -80,4 +80,7 @@ private:
 
 	/** Called to construct any the cell contents for columns created by external code*/
 	FConstructExternalColumnCell ConstructExternalColumnCell;
+
+	FOptionalSize MinWidth;
+	FOptionalSize MaxWidth;
 };
diff --git a/Engine/Source/Editor/PropertyEditor/Private/UserInterface/PropertyTable/ObjectNameTableCellPresenter.h b/Engine/Source/Editor/PropertyEditor/Private/UserInterface/PropertyTable/ObjectNameTableCellPresenter.h
--- a/Engine/Source/Editor/PropertyEditor/Private/UserInterface/PropertyTable/ObjectNameTableCellPresenter.h
+++ b/Engine/Source/Editor/PropertyEditor/Private/UserInterface/PropertyTable/ObjectNameTableCellPresenter.h
@@ -35,7 +35,7 @@ public:
 	{
 		return SNew( SBox )
 			.VAlign( VAlign_Center)
-			.Padding( FMargin( 2, 0, 2, 0 ) )
+			.Padding( FMargin( 4, 0, 2, 0 ) )
 			[
 				ConstructNameWidget( PropertyTableConstants::NormalFontStyle )
 			]
diff --git a/Engine/Source/Editor/PropertyEditor/Private/UserInterface/PropertyTable/SColumnHeader.h b/Engine/Source/Editor/PropertyEditor/Private/UserInterface/PropertyTable/SColumnHeader.h
--- a/Engine/Source/Editor/PropertyEditor/Private/UserInterface/PropertyTable/SColumnHeader.h
+++ b/Engine/Source/Editor/PropertyEditor/Private/UserInterface/PropertyTable/SColumnHeader.h
@@ -18,6 +18,9 @@
 #include "PropertyEditorHelpers.h"
 #include "IDocumentation.h"
 
+#define LOCTEXT_NAMESPACE "SColumnHeader"
+
+
 class SColumnHeader : public SCompoundWidget
 {
 public:
@@ -70,10 +73,10 @@ protected:
 		else if ( PropertyPath.IsValid() && PropertyPath->GetNumProperties() > 0 )
 		{
 			FProperty* Property = PropertyPath->GetLeafMostProperty().Property.Get();
-			const FText ToolTipText = PropertyEditorHelpers::GetToolTipText( Property );
+			const FText& DisplayNameText = Column->GetDisplayName();
 			const FString DocumentationLink = PropertyEditorHelpers::GetDocumentationLink( Property );
 			const FString DocumentationExcerptName = PropertyEditorHelpers::GetDocumentationExcerptName( Property );
-			ColumnHeaderToolTip = IDocumentation::Get()->CreateToolTip( ToolTipText, NULL, DocumentationLink, DocumentationExcerptName );
+			ColumnHeaderToolTip = IDocumentation::Get()->CreateToolTip( DisplayNameText, NULL, DocumentationLink, DocumentationExcerptName );
 		}
 
 		this->SetToolTip( ColumnHeaderToolTip );
@@ -101,29 +104,18 @@ protected:
 		TArray< FString > DisplayNamePieces;
 		DisplayNameText.ToString().ParseIntoArray(DisplayNamePieces, TEXT("->"), true);
 
-		for (int Index = 0; Index < DisplayNamePieces.Num(); Index++)
+		if(!DisplayNamePieces.IsEmpty())
 		{
+			// We only show the property name in the column header, the full path in the tooltip
 			NameBox->AddSlot()
 			.AutoWidth()
 			[
 				SNew( STextBlock )
 				.Font( FAppStyle::GetFontStyle( TextFontStyle ) )
-				.Text( FText::FromString(DisplayNamePieces[ Index ]) )
+				.Text( FText::FromString(DisplayNamePieces[ DisplayNamePieces.Num() - 1 ]) )
 			];
-
-			if ( Index < DisplayNamePieces.Num() - 1 )
-			{
-				NameBox->AddSlot()
-				.AutoWidth()
-				.HAlign( HAlign_Center )
-				.VAlign( VAlign_Center )
-				[
-					SNew( SImage )
-					.Image( FAppStyle::GetBrush( Style, ".HeaderRow.Column.PathDelimiter" ) )
-				];
-			}
 		}
-
+		
 		return NameBox;
 	}
 
@@ -136,3 +128,6 @@ protected:
 
 	TSharedPtr< IPropertyTableCustomColumn > Customization;
 };
+
+#undef LOCTEXT_NAMESPACE
+
diff --git a/Engine/Source/Editor/PropertyEditor/Private/UserInterface/PropertyTable/SPropertyTable.h b/Engine/Source/Editor/PropertyEditor/Private/UserInterface/PropertyTable/SPropertyTable.h
--- a/Engine/Source/Editor/PropertyEditor/Private/UserInterface/PropertyTable/SPropertyTable.h
+++ b/Engine/Source/Editor/PropertyEditor/Private/UserInterface/PropertyTable/SPropertyTable.h
@@ -83,14 +83,15 @@ public:
 			[
 				SNew( SVerticalBox )
 				+SVerticalBox::Slot()
-				.Padding(0.0f, 0.0f, 0.0f, 4.0f)
+				.Padding(0.0f, 0.0f, 0.0f, 1.0f)
 				.AutoHeight()
 				[
 					SNew(SBorder)
-					.BorderImage( FAppStyle::GetBrush( "ToolPanel.GroupBorder" ) )
+					.BorderImage( FAppStyle::GetBrush( "Brushes.Header" ) )
 					[
 						SAssignNew( BreadcrumbTrail, SBreadcrumbTrail< int32 > )
 						.PersistentBreadcrumbs( true )
+						.DelimiterImage(FAppStyle::Get().GetBrush("Icons.ChevronRight"))
 						.OnCrumbClicked( this, &SPropertyTable::OnCrumbClicked )
 						.GetCrumbMenuContent( this, &SPropertyTable::GetCrumbMenuContent )
 					]
@@ -695,6 +696,8 @@ private:
 		{
 			Table->SetRootPath( RootPath->TrimPath( AmountToTrimRoot ) );
 		}
+
+		ClearSelection();
 	}
 
 	TSharedRef< SWidget > GetCrumbMenuContent( const int32& Item )
diff --git a/Engine/Source/Editor/PropertyEditor/Private/UserInterface/PropertyTable/SPropertyTableCell.cpp b/Engine/Source/Editor/PropertyEditor/Private/UserInterface/PropertyTable/SPropertyTableCell.cpp
--- a/Engine/Source/Editor/PropertyEditor/Private/UserInterface/PropertyTable/SPropertyTableCell.cpp
+++ b/Engine/Source/Editor/PropertyEditor/Private/UserInterface/PropertyTable/SPropertyTableCell.cpp
@@ -48,7 +48,8 @@ void SPropertyTableCell::Construct( const FArguments& InArgs, const TSharedRef<
 	Presenter = InArgs._Presenter;
 	Style = InArgs._Style;
 
-	CellBackground = FAppStyle::GetBrush( Style, ".ColumnBorder" );
+	// This draws a box around each column which looks like gridlines
+	CellBackground = FAppStyle::GetBrush( Style, ".CellBorder" );
 
 	SetContent( ConstructCellContents() );
 
@@ -56,10 +57,6 @@ void SPropertyTableCell::Construct( const FArguments& InArgs, const TSharedRef<
 	Cell->OnExitedEditMode().AddSP( this, &SPropertyTableCell::ExitedEditMode );
 
 	FCoreUObjectDelegates::OnObjectPropertyChanged.AddSP(this, &SPropertyTableCell::OnCellValueChanged);
-
-	static const FName InvertedForegroundName("InvertedForeground");
-
-	SetForegroundColor( FAppStyle::GetSlateColor(InvertedForegroundName) );
 }
 
 void SPropertyTableCell::SetContent( const TSharedRef< SWidget >& NewContents )
@@ -105,18 +102,6 @@ TSharedRef< SWidget > SPropertyTableCell::ConstructCellContents()
 	return CellContents;
 }
 
-const FSlateBrush* SPropertyTableCell::GetCurrentCellBorder() const
-{
-	const bool IsReadOnly = !Presenter.IsValid() || Presenter->HasReadOnlyEditMode() || Cell->IsReadOnly();
-
-	if ( IsReadOnly )
-	{
-		return FAppStyle::GetBrush( Style, ".ReadOnlyCurrentCellBorder" );
-	}
-
-	return FAppStyle::GetBrush( Style, ".CurrentCellBorder" );
-}
-
 void SPropertyTableCell::OnAnchorWindowClosed( const TSharedRef< SWindow >& WindowClosing )
 {
 	Cell->ExitEditMode();
@@ -177,12 +162,24 @@ int32 SPropertyTableCell::OnPaint( const FPaintArgs& Args, const FGeometry& Allo
 
 		if ( Cell->GetTable()->GetCurrentCell() == Cell )
 		{
-			Background = GetCurrentCellBorder();
+			// The current cell has the blue selection color and a white border around it
+			Background = FAppStyle::GetBrush( Style, ".CurrentCellBorder" );
 		}
 		else if ( Cell->GetTable()->GetSelectedCells().Contains( Cell.ToSharedRef() ) )
 		{
-			Background = FAppStyle::GetBrush( Style, ".ReadOnlySelectedCellBorder" );
+			// Other selected cells have the blue selection color along with the gridlines
+			Background = FAppStyle::GetBrush( Style, ".SelectedCellBorder" );
+		}
+
+		// If the cell is in edit mode, we give the presenter a chance to specify the background
+		if(Cell->InEditMode())
+		{
+			if(const FSlateBrush* PresenterBackground = Presenter->GetEditModeCellBrush())
+			{
+				Background = PresenterBackground;
+			}
 		}
+		
 
 		FSlateDrawElement::MakeBox(
 			OutDrawElements,
@@ -215,22 +212,14 @@ FReply SPropertyTableCell::OnMouseButtonDoubleClick( const FGeometry& InMyGeomet
 
 TSharedRef< class SWidget > SPropertyTableCell::ConstructEditModeCellWidget()
 {
-	const FSlateBrush* BorderBrush = ( Presenter->HasReadOnlyEditMode() || Cell->IsReadOnly() ) ? FAppStyle::GetBrush( Style, ".ReadOnlyEditModeCellBorder" ) : FAppStyle::GetBrush( Style, ".Selection.Active" );
-
-	return SNew( SBorder )
-		.BorderImage( BorderBrush )
-		.VAlign( VAlign_Center )
-		.Padding( 0 )
-		.Content()
-		[
+	return
 			SAssignNew( DropDownAnchor, SMenuAnchor )
 			.Placement( MenuPlacement_ComboBox )
 			.OnGetMenuContent( this, &SPropertyTableCell::ConstructEditModeDropDownWidget )
 			.Content()
 			[
 				Presenter->ConstructEditModeCellWidget()
-			]
-		];
+			];
 }
 
 TSharedRef< class SWidget > SPropertyTableCell::ConstructEditModeDropDownWidget()
@@ -238,28 +227,21 @@ TSharedRef< class SWidget > SPropertyTableCell::ConstructEditModeDropDownWidget(
 	return Presenter->ConstructEditModeDropDownWidget();
 }
 
-TSharedRef<SBorder> SPropertyTableCell::ConstructInvalidPropertyWidget()
+TSharedRef<SWidget> SPropertyTableCell::ConstructInvalidPropertyWidget()
 {
 	return 
-		SNew(SBorder)
-		.BorderImage(FAppStyle::GetBrush( Style, ".ReadOnlyEditModeCellBorder"))
-		.VAlign(VAlign_Center)
-		.Padding(0)
-		.Content()
+		SNew(SHorizontalBox)
+		+SHorizontalBox::Slot()
+		.AutoWidth()
+		.Padding(FMargin(0.0f, 0.0f, 4.0f, 0.0f))
 		[
-			SNew(SHorizontalBox)
-			+SHorizontalBox::Slot()
-			.AutoWidth()
-			.Padding(FMargin(0.0f, 0.0f, 4.0f, 0.0f))
-			[
-				SNew(SImage)
-				.Image(FAppStyle::GetBrush("Icons.Error"))
-			]
-			+SHorizontalBox::Slot()
-			[
-				SNew(STextBlock)
-				.ColorAndOpacity(FLinearColor::Red)
-				.Text(NSLOCTEXT("PropertyEditor", "InvalidTableCellProperty", "Failed to retrieve value"))
-			]
+			SNew(SImage)
+			.Image(FAppStyle::GetBrush("Icons.Error"))
+		]
+		+SHorizontalBox::Slot()
+		[
+			SNew(STextBlock)
+			.ColorAndOpacity(FLinearColor::Red)
+			.Text(NSLOCTEXT("PropertyEditor", "InvalidTableCellProperty", "Failed to retrieve value"))
 		];
 }
diff --git a/Engine/Source/Editor/PropertyEditor/Private/UserInterface/PropertyTable/SPropertyTableCell.h b/Engine/Source/Editor/PropertyEditor/Private/UserInterface/PropertyTable/SPropertyTableCell.h
--- a/Engine/Source/Editor/PropertyEditor/Private/UserInterface/PropertyTable/SPropertyTableCell.h
+++ b/Engine/Source/Editor/PropertyEditor/Private/UserInterface/PropertyTable/SPropertyTableCell.h
@@ -62,9 +62,7 @@ private:
 
 	TSharedRef< class SWidget > ConstructEditModeDropDownWidget();
 
-	TSharedRef< class SBorder > ConstructInvalidPropertyWidget();
-
-	const FSlateBrush* GetCurrentCellBorder() const;
+	TSharedRef< class SWidget > ConstructInvalidPropertyWidget();
 
 	void OnAnchorWindowClosed( const TSharedRef< SWindow >& WindowClosing );
 
diff --git a/Engine/Source/Editor/PropertyEditor/Private/UserInterface/PropertyTable/SPropertyTableHeaderRow.h b/Engine/Source/Editor/PropertyEditor/Private/UserInterface/PropertyTable/SPropertyTableHeaderRow.h
--- a/Engine/Source/Editor/PropertyEditor/Private/UserInterface/PropertyTable/SPropertyTableHeaderRow.h
+++ b/Engine/Source/Editor/PropertyEditor/Private/UserInterface/PropertyTable/SPropertyTableHeaderRow.h
@@ -83,6 +83,7 @@ public:
 				AddColumn( 
 					SHeaderRow::Column( ColumnId )
 					.SortMode( TableRef, &IPropertyTable::GetColumnSortMode, Column )
+					.HeaderComboVisibility(EHeaderComboVisibility::Never)
 					.SortPriority( TableRef, &IPropertyTable::GetColumnSortPriority, Column)
 					.OnSort( TableRef, &IPropertyTable::SortByColumnWithId )
 					.MenuContent()
@@ -111,6 +112,7 @@ public:
 				AddColumn(
 					SHeaderRow::Column( ColumnId )
 					.SortMode( SortMode )
+					.HeaderComboVisibility(EHeaderComboVisibility::Never)
 					.SortPriority( SortPriority )
 					.OnSort( OnSort )
 					.MenuContent()
diff --git a/Engine/Source/Editor/PropertyEditor/Private/UserInterface/PropertyTable/SPropertyTableRow.h b/Engine/Source/Editor/PropertyEditor/Private/UserInterface/PropertyTable/SPropertyTableRow.h
--- a/Engine/Source/Editor/PropertyEditor/Private/UserInterface/PropertyTable/SPropertyTableRow.h
+++ b/Engine/Source/Editor/PropertyEditor/Private/UserInterface/PropertyTable/SPropertyTableRow.h
@@ -44,7 +44,11 @@ public:
 		Row->OnRefresh()->AddSP( this, &SPropertyTableRow::Refresh );
 		Style = InArgs._Style;
 
-		SMultiColumnTableRow< TSharedRef< class IPropertyTableRow > >::Construct( FSuperRowType::FArguments().Style( FAppStyle::Get(), "PropertyTable.TableRow" ), InOwnerTableView );
+		SMultiColumnTableRow< TSharedRef< class IPropertyTableRow > >::Construct(
+			FSuperRowType::FArguments()
+			.Style( FAppStyle::Get(), "TableView.AlternatingRow" )
+			.ShowSelection(false) // We don't show row selection on the property table - each selected cell is invidually shown
+			, InOwnerTableView );
 	}
 
 	FReply OnMouseButtonDown( const FGeometry& MyGeometry, const FPointerEvent& MouseEvent ) override
diff --git a/Engine/Source/Editor/PropertyEditor/Private/UserInterface/PropertyTable/SRowHeaderCell.h b/Engine/Source/Editor/PropertyEditor/Private/UserInterface/PropertyTable/SRowHeaderCell.h
--- a/Engine/Source/Editor/PropertyEditor/Private/UserInterface/PropertyTable/SRowHeaderCell.h
+++ b/Engine/Source/Editor/PropertyEditor/Private/UserInterface/PropertyTable/SRowHeaderCell.h
@@ -37,6 +37,9 @@ class SRowHeaderCell : public SCompoundWidget
 		Style = InArgs._Style;
 		Cell = InCell;
 
+		// This draws a box around the cell which looks like gridlines on the table
+		CellBackground = FAppStyle::GetBrush( Style, ".CellBorder" );
+		
 		TArray< EPropertyButton::Type > RequiredButtons;
 		if ( PropertyEditor.IsValid() )
 		{
@@ -45,7 +48,9 @@ class SRowHeaderCell : public SCompoundWidget
 
 		TSharedRef< SWidget > Content = 
 			SNew(SImage)
-			.Image( FAppStyle::GetBrush("ContentBrowser.ContentDirty") )
+			.Image( FAppStyle::GetBrush("Icons.DirtyBadge") )
+			.DesiredSizeOverride(FVector2D{ 16.0f })
+			.ColorAndOpacity(FSlateColor::UseForeground())
 			.Visibility( this, &SRowHeaderCell::GetDirtyImageVisibility );
 
 		if ( RequiredButtons.Contains( EPropertyButton::Insert_Delete_Duplicate ) )
@@ -68,15 +73,10 @@ class SRowHeaderCell : public SCompoundWidget
 		[
 			SNew( SBox )
 			.HeightOverride( 20.0f )
+			.HAlign(HAlign_Center)
+			.VAlign(VAlign_Center)
 			[
-				SNew( SBorder )
-				.BorderImage( this, &SRowHeaderCell::GetBorder )
-				.HAlign( HAlign_Center )
-				.VAlign( VAlign_Center )
-				.Content()
-				[
-					Content
-				]
+				Content
 			]
 		];
 	}
@@ -97,6 +97,24 @@ class SRowHeaderCell : public SCompoundWidget
 		return FReply::Unhandled();
 	}
 
+	virtual int32 OnPaint( const FPaintArgs& Args, const FGeometry& AllottedGeometry, const FSlateRect& MyCullingRect, FSlateWindowElementList& OutDrawElements, int32 LayerId, const FWidgetStyle& InWidgetStyle, bool bParentEnabled ) const override
+	{
+		if ( CellBackground && CellBackground->DrawAs != ESlateBrushDrawType::NoDrawType )
+		{
+			FSlateDrawElement::MakeBox(
+				OutDrawElements,
+				LayerId,
+				AllottedGeometry.ToPaintGeometry(),
+				CellBackground,
+				ESlateDrawEffect::None,
+				CellBackground->GetTint( InWidgetStyle ) * InWidgetStyle.GetColorAndOpacityTint() 
+				);
+		}
+
+		return SCompoundWidget::OnPaint( Args, AllottedGeometry, MyCullingRect, OutDrawElements, LayerId, InWidgetStyle, bParentEnabled );
+	}
+
+
 
 private:
 
@@ -124,4 +142,5 @@ private:
 	TSharedPtr< IPropertyTableCell > Cell;
 	TSharedPtr< FPropertyEditor > Editor;
 	FName Style;
+	const FSlateBrush* CellBackground = nullptr;
 };
diff --git a/Engine/Source/Editor/PropertyEditor/Private/UserInterface/PropertyTable/SRowHeaderColumnHeader.h b/Engine/Source/Editor/PropertyEditor/Private/UserInterface/PropertyTable/SRowHeaderColumnHeader.h
--- a/Engine/Source/Editor/PropertyEditor/Private/UserInterface/PropertyTable/SRowHeaderColumnHeader.h
+++ b/Engine/Source/Editor/PropertyEditor/Private/UserInterface/PropertyTable/SRowHeaderColumnHeader.h
@@ -42,7 +42,10 @@ class SRowHeaderColumnHeader : public SColumnHeader
 
 		ChildSlot
 		[
-			SNew( SSpacer )
+			SNew( SImage )
+			.DesiredSizeOverride(FVector2D{16.0f})
+			.ColorAndOpacity(FSlateColor::UseForeground())
+			.Image(FAppStyle::GetBrush("Icons.Save"))
 		];
 	}
 
diff --git a/Engine/Source/Editor/PropertyEditor/Private/UserInterface/PropertyTable/TextPropertyTableCellPresenter.h b/Engine/Source/Editor/PropertyEditor/Private/UserInterface/PropertyTable/TextPropertyTableCellPresenter.h
--- a/Engine/Source/Editor/PropertyEditor/Private/UserInterface/PropertyTable/TextPropertyTableCellPresenter.h
+++ b/Engine/Source/Editor/PropertyEditor/Private/UserInterface/PropertyTable/TextPropertyTableCellPresenter.h
@@ -19,6 +19,8 @@ public:
 	virtual TSharedRef< class SWidget > ConstructDisplayWidget() override;
 
 	virtual bool RequiresDropDown() override;
+	
+	virtual const FSlateBrush* GetEditModeCellBrush() override { return FAppStyle::GetBrush("NoBrush"); }
 
 	virtual TSharedRef< class SWidget > ConstructEditModeCellWidget() override;
 
diff --git a/Engine/Source/Editor/PropertyEditor/Private/UserInterface/PropertyTree/SPropertyTreeCategoryRow.h b/Engine/Source/Editor/PropertyEditor/Private/UserInterface/PropertyTree/SPropertyTreeCategoryRow.h
--- a/Engine/Source/Editor/PropertyEditor/Private/UserInterface/PropertyTree/SPropertyTreeCategoryRow.h
+++ b/Engine/Source/Editor/PropertyEditor/Private/UserInterface/PropertyTree/SPropertyTreeCategoryRow.h
@@ -32,17 +32,96 @@ public:
 			STableRow< TSharedPtr< class FPropertyNode* > >::FArguments()
 				.Padding(3)
 				[
+					SNew( STextBlock )
+					.Text( InArgs._DisplayName )
+					.Font( FAppStyle::GetFontStyle( PropertyEditorConstants::CategoryFontStyle ) )
+				]
+			, InOwnerTable );
+	}
+
+	// We override ConstructChildren because we want a border around the row to add some spacing between headers and have the expander spaced differently
+	virtual void ConstructChildren( ETableViewMode::Type InOwnerTableMode, const TAttribute<FMargin>& InPadding, const TSharedRef<SWidget>& InContent ) override
+	{
+		this->Content = InContent;
+		InnerContentSlot = nullptr;
+
+		SHorizontalBox::FSlot* InnerContentSlotNativePtr = nullptr;
+
+		this->ChildSlot
+		[
+			// This border determines how far the category header will be indented based on its hierarchy
+			SNew(SBorder)
+			.Padding(this, &SPropertyTreeCategoryRow::GetBorderPadding)
+			.BorderImage(FAppStyle::GetBrush("Brushes.Panel"))
+			[
+				// This border draws a grind line after the header for when two headers are next to each other
+				SNew(SBorder)
+				.BorderImage(FAppStyle::Get().GetBrush("DetailsView.GridLine"))
+				.Padding(FMargin(0,0,0,1))
+				[
+					// This border is what contains the actual category row color and contents
 					SNew(SBorder)
-					.VAlign(VAlign_Center)
-					.BorderImage(FAppStyle::GetBrush("PropertyWindow.CategoryBackground"))
+					.BorderImage(this, &SPropertyTreeCategoryRow::GetInnerBackgroundImage)
 					.Padding(FMargin(5,1,5,0))
-					.ForegroundColor(FAppStyle::GetColor("PropertyWindow.CategoryForeground"))
 					[
-						SNew( STextBlock )
-						.Text( InArgs._DisplayName )
-						.Font( FAppStyle::GetFontStyle( PropertyEditorConstants::CategoryFontStyle ) )
+						SNew(SHorizontalBox)
+
+						+ SHorizontalBox::Slot()
+						.AutoWidth()
+						.HAlign(HAlign_Right)
+						.VAlign(VAlign_Fill)
+						.Padding(FMargin(4, 0, 0, 0))
+						[
+							SAssignNew(ExpanderArrowWidget, SExpanderArrow, SharedThis(this) )
+							.StyleSet(ExpanderStyleSet)
+							// We set the indent amount to 0 as we indent the widget ourselves so that we can also indent the border around it
+							.IndentAmount(0) 
+						]
+
+						+ SHorizontalBox::Slot()
+						.FillWidth(1)
+						.Expose( InnerContentSlotNativePtr )
+						.Padding( InPadding )
+						[
+							InContent
+						]
 					]
 				]
-			, InOwnerTable );
+			]
+		];
+
+		InnerContentSlot = InnerContentSlotNativePtr;
+	}
+
+
+	virtual FReply OnMouseButtonDown( const FGeometry& MyGeometry, const FPointerEvent& MouseEvent ) override
+	{
+		if (MouseEvent.GetEffectingButton() == EKeys::LeftMouseButton)
+		{
+			ToggleExpansion();
+		}
+
+		return STableRow< TSharedPtr< class FPropertyNode* > >::OnMouseButtonDown(MyGeometry, MouseEvent);
 	}
+
+protected:
+
+	const FSlateBrush* GetInnerBackgroundImage() const
+	{
+		if(IsHovered())
+		{
+			return FAppStyle::GetBrush("Brushes.Dropdown");
+		}
+
+		return FAppStyle::GetBrush("Brushes.Header");
+	}
+
+	/** @return the margin corresponding to how far this item is indented */
+	FMargin GetBorderPadding() const
+	{
+		const int32 NestingDepth = GetIndentLevel();
+		const float Indent = 20.0f;
+		return FMargin( NestingDepth * Indent, 0,0,0 );
+	}
+
 };
diff --git a/Engine/Source/Editor/PropertyEditor/Public/IPropertyTableCellPresenter.h b/Engine/Source/Editor/PropertyEditor/Public/IPropertyTableCellPresenter.h
--- a/Engine/Source/Editor/PropertyEditor/Public/IPropertyTableCellPresenter.h
+++ b/Engine/Source/Editor/PropertyEditor/Public/IPropertyTableCellPresenter.h
@@ -10,6 +10,8 @@ public:
 
 	virtual TSharedRef< class SWidget > ConstructDisplayWidget() = 0;
 
+	virtual const FSlateBrush* GetEditModeCellBrush() { return nullptr; }
+
 	virtual bool RequiresDropDown() = 0;
 
 	virtual TSharedRef< class SWidget > ConstructEditModeCellWidget() = 0;
diff --git a/Engine/Source/Editor/UnrealEd/Private/Toolkits/AssetEditorToolkit.cpp b/Engine/Source/Editor/UnrealEd/Private/Toolkits/AssetEditorToolkit.cpp
--- a/Engine/Source/Editor/UnrealEd/Private/Toolkits/AssetEditorToolkit.cpp
+++ b/Engine/Source/Editor/UnrealEd/Private/Toolkits/AssetEditorToolkit.cpp
@@ -250,7 +250,9 @@ void FAssetEditorToolkit::InitAssetEditor( const EToolkitMode::Type Mode, const
 	ToolkitCommands->MapAction(
 		FGlobalEditorCommonCommands::Get().FindInContentBrowser,
 		FExecuteAction::CreateSP( this, &FAssetEditorToolkit::FindInContentBrowser_Execute ),
-		FCanExecuteAction::CreateSP( this, &FAssetEditorToolkit::CanFindInContentBrowser ));
+		FCanExecuteAction::CreateSP( this, &FAssetEditorToolkit::CanFindInContentBrowser ),
+		FIsActionChecked(),
+		FIsActionButtonVisible::CreateSP(this, &FAssetEditorToolkit::IsFindInContentBrowserButtonVisible));
 		
 	ToolkitCommands->MapAction(
 		FGlobalEditorCommonCommands::Get().OpenDocumentation,
diff --git a/Engine/Source/Editor/UnrealEd/Public/Toolkits/AssetEditorToolkit.h b/Engine/Source/Editor/UnrealEd/Public/Toolkits/AssetEditorToolkit.h
--- a/Engine/Source/Editor/UnrealEd/Public/Toolkits/AssetEditorToolkit.h
+++ b/Engine/Source/Editor/UnrealEd/Public/Toolkits/AssetEditorToolkit.h
@@ -334,6 +334,9 @@ protected:
 	/** Called to test if "Find in Content Browser" should be enabled for this asset */
 	virtual bool CanFindInContentBrowser() const { return true; }
 
+	/** Called to test if "Find in Content Browser" should be visible for this asset */
+	virtual bool IsFindInContentBrowserButtonVisible() const { return true; }
+
 	/** Called when "Find in Content Browser" is clicked for this asset */
 	virtual void FindInContentBrowser_Execute();
 
