Non-unity fix in PropertyTableCellPresenter

diff --git a/Engine/Source/Editor/PropertyEditor/Public/IPropertyTableCellPresenter.h b/Engine/Source/Editor/PropertyEditor/Public/IPropertyTableCellPresenter.h
--- a/Engine/Source/Editor/PropertyEditor/Public/IPropertyTableCellPresenter.h
+++ b/Engine/Source/Editor/PropertyEditor/Public/IPropertyTableCellPresenter.h
@@ -2,6 +2,7 @@
 #pragma once
 
 #include "CoreMinimal.h"
+#include "Styling/SlateBrush.h"
 
 class IPropertyTableCellPresenter
 {
