[BugFix] Details splitter doesn't highlight certain diffed properties and scrolling isn't synced with those rows

diff --git a/Engine/Source/Editor/PropertyEditor/Private/SDetailsViewBase.cpp b/Engine/Source/Editor/PropertyEditor/Private/SDetailsViewBase.cpp
--- a/Engine/Source/Editor/PropertyEditor/Private/SDetailsViewBase.cpp
+++ b/Engine/Source/Editor/PropertyEditor/Private/SDetailsViewBase.cpp
@@ -421,7 +421,7 @@ bool SDetailsViewBase::IsAncestorCollapsed(const TSharedRef<IDetailTreeNode>& No
 	TSharedPtr<FDetailTreeNode> Ancestor = StaticCastSharedRef<FDetailTreeNode>(Node)->GetParentNode().Pin();
 	while (!bIsCollapsed && Ancestor)
 	{
-		bIsCollapsed = !DetailTree->IsItemExpanded(Ancestor.ToSharedRef());
+		bIsCollapsed = !Ancestor->ShouldShowOnlyChildren() && !DetailTree->IsItemExpanded(Ancestor.ToSharedRef());
 		Ancestor = Ancestor->GetParentNode().Pin();
 	}
 	return bIsCollapsed;
