PropertyEditor animate multiple rows

diff --git a/Engine/Source/Editor/PropertyEditor/Private/IDetailsViewPrivate.h b/Engine/Source/Editor/PropertyEditor/Private/IDetailsViewPrivate.h
--- a/Engine/Source/Editor/PropertyEditor/Private/IDetailsViewPrivate.h
+++ b/Engine/Source/Editor/PropertyEditor/Private/IDetailsViewPrivate.h
@@ -144,10 +144,15 @@ public:
 	*/
 	virtual void RestoreExpandedItems(TSharedRef<FPropertyNode> StartNode) = 0;
 
-	/** Mark node as animating, useful if animating during behaviors that trigger widget reconstruction */
-	virtual void MarkNodeAnimating(TSharedPtr<FPropertyNode> InNode, float InAnimationDuration) = 0;
+	/**
+	* Mark node as animating, useful if animating during behaviors that trigger widget reconstruction.
+	* @param InNode				The slate property node to animate.
+	* @param InAnimationDuration	The animation duration in seconds.
+	* @param InAnimationBatchId	(Optional) A batch id to enable simultaneous (and performant) animation of multiple rows.
+	*/
+	virtual void MarkNodeAnimating(TSharedPtr<FPropertyNode> InNode, float InAnimationDuration, TOptional<FGuid> InAnimationBatchId = {}) = 0;
 
-	/** Returns true if node is animating, currently we only keep track of the last animated node */
+	/** @return true if node is animating */
 	virtual bool IsNodeAnimating(TSharedPtr<FPropertyNode> InNode) = 0;
 
 	/** Column width accessibility */
diff --git a/Engine/Source/Editor/PropertyEditor/Private/SDetailsViewBase.cpp b/Engine/Source/Editor/PropertyEditor/Private/SDetailsViewBase.cpp
--- a/Engine/Source/Editor/PropertyEditor/Private/SDetailsViewBase.cpp
+++ b/Engine/Source/Editor/PropertyEditor/Private/SDetailsViewBase.cpp
@@ -2,31 +2,25 @@
 
 #include "SDetailsViewBase.h"
 
+#include "Algo/AllOf.h"
+#include "Algo/AnyOf.h"
 #include "Engine/Engine.h"
 #include "GameFramework/Actor.h"
 #include "Misc/ConfigCacheIni.h"
 #include "Modules/ModuleManager.h"
 #include "Presentation/PropertyEditor/PropertyEditor.h"
-#include "Settings/EditorExperimentalSettings.h"
-#include "UserInterface/PropertyEditor/SPropertyEditorEditInline.h"
 #include "Widgets/Colors/SColorPicker.h"
 #include "Widgets/Input/SSearchBox.h"
-
-#include "CategoryPropertyNode.h"
+#include "Framework/Application/SlateApplication.h"
 #include "DetailCategoryBuilderImpl.h"
 #include "DetailLayoutBuilderImpl.h"
 #include "DetailLayoutHelpers.h"
-#include "DetailPropertyRow.h"
-#include "EditConditionParser.h"
 #include "Editor.h"
-#include "EditorConfigSubsystem.h"
 #include "IDetailCustomization.h"
-#include "ObjectEditorUtils.h"
 #include "ObjectPropertyNode.h"
+#include "PropertyPermissionList.h"
 #include "ScopedTransaction.h"
 #include "SDetailNameArea.h"
-#include "SDetailsView.h"
-#include "PropertyPermissionList.h"
 #include "ThumbnailRendering/ThumbnailManager.h"
 
 SDetailsViewBase::SDetailsViewBase() :
@@ -1546,30 +1540,102 @@ void SDetailsViewBase::RestoreExpandedItems(TSharedRef<FPropertyNode> StartNode)
 	}
 }
 
-void SDetailsViewBase::MarkNodeAnimating(TSharedPtr<FPropertyNode> InNode, float InAnimationDuration)
+void SDetailsViewBase::MarkNodeAnimating(TSharedPtr<FPropertyNode> InNode, float InAnimationDuration, TOptional<FGuid> InAnimationBatchId)
 {
 	if (InNode.IsValid() && InAnimationDuration > 0.0f)
 	{
-		CurrentlyAnimatingNodePath = FPropertyNode::CreatePropertyPath(InNode.ToSharedRef());
-		GEditor->GetTimerManager()->ClearTimer(AnimateNodeTimer);
-		GEditor->GetTimerManager()->SetTimer(AnimateNodeTimer, FTimerDelegate::CreateSP(this, &SDetailsViewBase::HandleNodeAnimationComplete), InAnimationDuration, false);
+		TSharedPtr<FPropertyPath> NodePropertyPath = FPropertyNode::CreatePropertyPath(InNode.ToSharedRef());
+		
+		// Try finding existing, ignore Batch Id
+		FAnimatingNodeCollection* AnimatingNode = CurrentlyAnimatingNodeCollections.FindByPredicate(
+			[NodePropertyPath](const FAnimatingNodeCollection& AnimatingNode)
+			{
+				return Algo::AnyOf(AnimatingNode.NodePaths, [NodePropertyPath](const TSharedPtr<FPropertyPath>& NodePath)
+				{
+					return FPropertyPath::AreEqual(NodePath.ToSharedRef(), NodePropertyPath.ToSharedRef());
+				});
+			});
+
+		// If this node is part of a batch animation, let it play out (don't reset it)
+		if (AnimatingNode)
+		{
+			if (AnimatingNode->NodePaths.Num() > 1
+				|| (InAnimationBatchId.IsSet() && AnimatingNode->BatchId == InAnimationBatchId))
+			{
+				return;
+			}
+		}
+		
+		// If not found, but batch id found, add it to that batch
+		if (!AnimatingNode && InAnimationBatchId.IsSet())
+		{
+			AnimatingNode = CurrentlyAnimatingNodeCollections.FindByPredicate(
+				[InAnimationBatchId](const FAnimatingNodeCollection& AnimatingNode)
+				{
+					return AnimatingNode.BatchId == InAnimationBatchId;
+				});
+			
+			if (AnimatingNode != nullptr)
+			{
+				AnimatingNode->NodePaths.Add(NodePropertyPath);
+
+				// Timer Handle already in progress, so no need to clear/set below
+				return;
+			}
+		}
+
+		// Otherwise add it
+		if (AnimatingNode == nullptr)
+		{
+			AnimatingNode = &CurrentlyAnimatingNodeCollections.Add_GetRef({{NodePropertyPath}});
+			if (InAnimationBatchId.IsSet())
+			{
+				AnimatingNode->BatchId = InAnimationBatchId.GetValue();
+			}
+		}
+		
+		GEditor->GetTimerManager()->ClearTimer(AnimatingNode->NodeTimer);
+		GEditor->GetTimerManager()->SetTimer(AnimatingNode->NodeTimer, FTimerDelegate::CreateSP(this, &SDetailsViewBase::HandleNodeAnimationComplete, AnimatingNode), InAnimationDuration, false);
 	}
 }
 
 bool SDetailsViewBase::IsNodeAnimating(TSharedPtr<FPropertyNode> InNode)
 {
-	if (CurrentlyAnimatingNodePath.IsValid() && InNode.IsValid())
+	if (!InNode.IsValid())
 	{
-		TSharedRef<FPropertyPath> InNodePath = FPropertyNode::CreatePropertyPath(InNode.ToSharedRef());
-		return FPropertyPath::AreEqual(CurrentlyAnimatingNodePath.ToSharedRef(), InNodePath);
+		return false;
 	}
 
-	return false;
+	TSharedRef<FPropertyPath> InNodePath = FPropertyNode::CreatePropertyPath(InNode.ToSharedRef());
+	return CurrentlyAnimatingNodeCollections.ContainsByPredicate(
+		[InNodePath](const FAnimatingNodeCollection& AnimatingNode)
+		{
+			return Algo::AnyOf(AnimatingNode.NodePaths, [InNodePath](const TSharedPtr<FPropertyPath>& NodePath)
+			{
+				return FPropertyPath::AreEqual(NodePath.ToSharedRef(), InNodePath);
+			});
+		});
+}
+
+bool SDetailsViewBase::FAnimatingNodeCollection::IsValid() const
+{
+	return Algo::AllOf(NodePaths,
+			[](const TSharedPtr<FPropertyPath>& InNodePath)
+			{
+				return InNodePath.IsValid();
+			});
 }
 
-void SDetailsViewBase::HandleNodeAnimationComplete()
+void SDetailsViewBase::HandleNodeAnimationComplete(FAnimatingNodeCollection* InAnimatedNode)
 {
-	CurrentlyAnimatingNodePath = nullptr;
+	if (InAnimatedNode)
+	{
+		GEditor->GetTimerManager()->ClearTimer(InAnimatedNode->NodeTimer);
+		CurrentlyAnimatingNodeCollections.RemoveAll([InAnimatedNode](FAnimatingNodeCollection& AnimatedNode)
+		{
+			return AnimatedNode == *InAnimatedNode;		
+		});
+	}
 }
 
 void SDetailsViewBase::FilterRootNode(const TSharedPtr<FComplexPropertyNode>& RootNode)
diff --git a/Engine/Source/Editor/PropertyEditor/Private/SDetailsViewBase.h b/Engine/Source/Editor/PropertyEditor/Private/SDetailsViewBase.h
--- a/Engine/Source/Editor/PropertyEditor/Private/SDetailsViewBase.h
+++ b/Engine/Source/Editor/PropertyEditor/Private/SDetailsViewBase.h
@@ -6,16 +6,15 @@
 #include "AssetThumbnail.h"
 #include "DetailColumnSizeData.h"
 #include "DetailFilter.h"
-#include "DetailTreeNode.h"
 #include "DetailsViewConfig.h"
+#include "DetailTreeNode.h"
 #include "Framework/Commands/UICommandList.h"
 #include "IDetailsView.h"
 #include "IDetailsViewPrivate.h"
-#include "IPropertyUtilities.h"
 #include "Input/Reply.h"
+#include "IPropertyUtilities.h"
 #include "Layout/Visibility.h"
 #include "Misc/Attribute.h"
-#include "PropertyCustomizationHelpers.h"
 #include "PropertyEditorModule.h"
 #include "PropertyNode.h"
 #include "PropertyPath.h"
@@ -144,7 +143,7 @@ public:
 	virtual const FCustomPropertyTypeLayoutMap& GetCustomPropertyTypeLayoutMap() const { return InstancedTypeToLayoutMap; }
 	virtual void SaveExpandedItems( TSharedRef<FPropertyNode> StartNode ) override;
 	virtual void RestoreExpandedItems(TSharedRef<FPropertyNode> StartNode) override;
-	virtual void MarkNodeAnimating(TSharedPtr<FPropertyNode> InNode, float InAnimationDuration) override;
+	virtual void MarkNodeAnimating(TSharedPtr<FPropertyNode> InNode, float InAnimationDuration, TOptional<FGuid> InAnimationBatchId) override;
 	virtual bool IsNodeAnimating(TSharedPtr<FPropertyNode> InNode) override;
 	virtual FDetailColumnSizeData& GetColumnSizeData() override { return ColumnSizeData; }
 	virtual bool IsFavoritingEnabled() const override { return DetailsViewArgs.bAllowFavoriteSystem; }
@@ -193,8 +192,55 @@ protected:
 	 */
 	void SetColorPropertyFromColorPicker(FLinearColor NewColor);
 
-	/** Called when node animation completes */
-	void HandleNodeAnimationComplete();
+	/** Contains one or more nodes being animated by a single timer. */
+	struct FAnimatingNodeCollection
+	{
+		TArray<TSharedPtr<FPropertyPath>> NodePaths;
+		FTimerHandle NodeTimer;
+
+		/** Optionally identify by an id, where multiple nodes can be part of the same animation. */
+		FGuid BatchId;
+
+		/** Get validity of this AnimatingNode, determined by it's node path. */
+		bool IsValid() const;
+
+		bool operator==(const FAnimatingNodeCollection& Other) const
+		{
+			if (NodePaths.Num() != Other.NodePaths.Num())
+			{
+				return false;
+			}
+
+			if (BatchId.IsValid() && BatchId == Other.BatchId)
+			{
+				return true;
+			}
+
+			for (int32 Idx = 0; Idx < NodePaths.Num(); ++Idx)
+			{
+				if (!NodePaths[Idx].IsValid()
+					|| !Other.NodePaths[Idx].IsValid())
+				{
+					continue;
+				}
+
+				if (!FPropertyPath::AreEqual(NodePaths[Idx].ToSharedRef(), Other.NodePaths[Idx].ToSharedRef()))
+				{
+					return false;					
+				}				
+			}
+			
+			return true;
+		}
+
+		bool operator!=(const FAnimatingNodeCollection& Other) const
+		{
+			return !(*this == Other);
+		}
+	};
+
+	/** Called when a node collection animation completes */
+	void HandleNodeAnimationComplete(FAnimatingNodeCollection* InAnimatedNode);
 
 	/** Updates the property map for access when customizing the details view.  Generates default layout for properties */
 	void UpdatePropertyMaps();
@@ -422,10 +468,8 @@ protected:
 	TSharedPtr<IDetailPropertyExtensionHandler> ExtensionHandler;
 	/** The tree node that is currently highlighted, may be none. */
 	TWeakPtr<FDetailTreeNode> CurrentlyHighlightedNode;
-	/** The property node whose widget-equivalent should be animating, may be none. */
-	TSharedPtr<FPropertyPath> CurrentlyAnimatingNodePath;
-	/** Timer for that current node's widget animation duration. */
-	FTimerHandle AnimateNodeTimer;
+	/** The list of nodes whose widgets should be animating. */
+	TArray<FAnimatingNodeCollection> CurrentlyAnimatingNodeCollections;
 
 	/** Current set of expanded detail nodes (by path) that should be saved when the details panel closes */
 	FStringPrefixTree ExpandedDetailNodes;
