PropertyEditor category based copy/paste

diff --git a/Engine/Source/Editor/DetailCustomizations/Private/ComponentTransformDetails.cpp b/Engine/Source/Editor/DetailCustomizations/Private/ComponentTransformDetails.cpp
--- a/Engine/Source/Editor/DetailCustomizations/Private/ComponentTransformDetails.cpp
+++ b/Engine/Source/Editor/DetailCustomizations/Private/ComponentTransformDetails.cpp
@@ -22,9 +22,9 @@
 #include "HAL/PlatformApplicationMisc.h"
 #include "IDetailChildrenBuilder.h"
 #include "IDetailPropertyRow.h"
-#include "IPropertyUtilities.h"
 #include "Input/Events.h"
 #include "Internationalization/Internationalization.h"
+#include "IPropertyUtilities.h"
 #include "Kismet2/ComponentEditorUtils.h"
 #include "Layout/Margin.h"
 #include "Math/Quat.h"
@@ -34,6 +34,7 @@
 #include "Misc/Attribute.h"
 #include "Misc/ConfigCacheIni.h"
 #include "Misc/NotifyHook.h"
+#include "PropertyEditorCopyPaste.h"
 #include "PropertyHandle.h"
 #include "ScopedTransaction.h"
 #include "Settings/EditorProjectSettings.h"
@@ -45,15 +46,15 @@
 #include "Templates/UnrealTemplate.h"
 #include "Textures/SlateIcon.h"
 #include "Types/SlateStructs.h"
+#include "UnrealEdGlobals.h"
 #include "UObject/Class.h"
 #include "UObject/Object.h"
 #include "UObject/ObjectMacros.h"
 #include "UObject/ObjectPtr.h"
 #include "UObject/Package.h"
-#include "UObject/UObjectGlobals.h"
 #include "UObject/UnrealNames.h"
 #include "UObject/UnrealType.h"
-#include "UnrealEdGlobals.h"
+#include "UObject/UObjectGlobals.h"
 #include "Widgets/DeclarativeSyntaxSupport.h"
 #include "Widgets/Images/SImage.h"
 #include "Widgets/Input/NumericUnitTypeInterface.inl"
@@ -73,6 +74,16 @@ struct FSlateBrush;
 
 #define LOCTEXT_NAMESPACE "FComponentTransformDetails"
 
+namespace UE::DetailsCustomizations::Internal
+{
+	/** Lookup to get the property name for the given TransformField. */
+	static TMap<ETransformField::Type, FString> TransformFieldToPropertyNameString = {
+		{ ETransformField::Location, USceneComponent::GetRelativeLocationPropertyName().ToString() },
+		{ ETransformField::Rotation, USceneComponent::GetRelativeRotationPropertyName().ToString() },
+		{ ETransformField::Scale, USceneComponent::GetRelativeScale3DPropertyName().ToString() }		
+	};	
+}
+
 class FScopedSwitchWorldForObject
 {
 public:
@@ -312,12 +323,46 @@ void FComponentTransformDetails::OnPaste( ETransformField::Type TransformField )
 	FString PastedText;
 	FPlatformApplicationMisc::ClipboardPaste(PastedText);
 
-	switch (TransformField)
+	PasteFromText(TEXT(""), PastedText, TransformField);
+}
+
+void FComponentTransformDetails::OnPasteFromText(
+	const FString& InTag,
+	const FString& InText,
+	const TOptional<FGuid>& InOperationId,
+	ETransformField::Type InTransformField)
+{
+	PasteFromText(InTag, InText, InTransformField);
+}
+
+void FComponentTransformDetails::PasteFromText(
+	const FString& InTag,
+	const FString& InText,
+	ETransformField::Type InTransformField)
+{
+	if (InText.IsEmpty())
+	{
+		return;
+	}
+
+	FString Text = InText;
+	if (!InTag.IsEmpty())
 	{
-		case ETransformField::Location:
+		const FString PropertyPath = UE::PropertyEditor::GetPropertyPath(GetPropertyHandle());
+
+		// ensure that if tag is specified, that it matches the subscriber
+		if (!InTag.Equals(UE::DetailsCustomizations::Internal::TransformFieldToPropertyNameString[InTransformField]))
+		{
+			return;
+		}
+	}
+
+	switch (InTransformField)
+	{
+	case ETransformField::Location:
 		{
 			FVector Location;
-			if (Location.InitFromString(PastedText))
+			if (Location.InitFromString(Text))
 			{
 				FScopedTransaction Transaction(LOCTEXT("PasteLocation", "Paste Location"));
 				OnSetTransform(ETransformField::Location, EAxisList::All, Location, false, true);
@@ -327,10 +372,10 @@ void FComponentTransformDetails::OnPaste( ETransformField::Type TransformField )
 	case ETransformField::Rotation:
 		{
 			FRotator Rotation;
-			PastedText.ReplaceInline(TEXT("Pitch="), TEXT("P="));
-			PastedText.ReplaceInline(TEXT("Yaw="), TEXT("Y="));
-			PastedText.ReplaceInline(TEXT("Roll="), TEXT("R="));
-			if (Rotation.InitFromString(PastedText))
+			Text.ReplaceInline(TEXT("Pitch="), TEXT("P="));
+			Text.ReplaceInline(TEXT("Yaw="), TEXT("Y="));
+			Text.ReplaceInline(TEXT("Roll="), TEXT("R="));
+			if (Rotation.InitFromString(Text))
 			{
 				FScopedTransaction Transaction(LOCTEXT("PasteRotation", "Paste Rotation"));
 				OnSetTransform(ETransformField::Rotation, EAxisList::All, Rotation.Euler(), false, true);
@@ -340,7 +385,7 @@ void FComponentTransformDetails::OnPaste( ETransformField::Type TransformField )
 	case ETransformField::Scale:
 		{
 			FVector Scale;
-			if (Scale.InitFromString(PastedText))
+			if (Scale.InitFromString(Text))
 			{
 				FScopedTransaction Transaction(LOCTEXT("PasteScale", "Paste Scale"));
 				OnSetTransform(ETransformField::Scale, EAxisList::All, Scale, false, true);
@@ -372,14 +417,16 @@ BEGIN_SLATE_FUNCTION_BUILD_OPTIMIZATION
 void FComponentTransformDetails::GenerateChildContent( IDetailChildrenBuilder& ChildrenBuilder )
 {
 	UClass* SceneComponentClass = USceneComponent::StaticClass();
-		
+
 	FSlateFontInfo FontInfo = IDetailLayoutBuilder::GetDetailFont();
 
 	const bool bHideLocationField = ( HiddenFieldMask & ( 1 << ETransformField::Location ) ) != 0;
 	const bool bHideRotationField = ( HiddenFieldMask & ( 1 << ETransformField::Rotation ) ) != 0;
 	const bool bHideScaleField = ( HiddenFieldMask & ( 1 << ETransformField::Scale ) ) != 0;
 
-	IDetailLayoutBuilder& LayoutBuilder = ChildrenBuilder.GetParentCategory().GetParentLayout();
+	IDetailCategoryBuilder& ParentCategory = ChildrenBuilder.GetParentCategory();
+	
+	IDetailLayoutBuilder& LayoutBuilder = ParentCategory.GetParentLayout();
 	TSharedPtr<IPropertyHandle> LocationPropertyHandle = LayoutBuilder.GetProperty(USceneComponent::GetRelativeLocationPropertyName(), USceneComponent::StaticClass());
 	TSharedPtr<IPropertyHandle> RotationPropertyHandle = LayoutBuilder.GetProperty(USceneComponent::GetRelativeRotationPropertyName(), USceneComponent::StaticClass());
 	TSharedPtr<IPropertyHandle> ScalePropertyHandle = LayoutBuilder.GetProperty(USceneComponent::GetRelativeScale3DPropertyName(), USceneComponent::StaticClass());
@@ -403,6 +450,8 @@ void FComponentTransformDetails::GenerateChildContent( IDetailChildrenBuilder& C
 			TypeInterface = SharedThis(this);
 		}
 
+		ParentCategory.OnPasteFromText()->AddSP(this, &FComponentTransformDetails::OnPasteFromText, ETransformField::Location);
+
 		ChildrenBuilder.AddCustomRow( LOCTEXT("LocationFilter", "Location") )
 		.RowTag("Location")
 		.CopyAction( CreateCopyAction( ETransformField::Location ) )
@@ -449,6 +498,8 @@ void FComponentTransformDetails::GenerateChildContent( IDetailChildrenBuilder& C
 			TypeInterface = MakeShareable( new TNumericUnitTypeInterface<FRotator::FReal>(EUnit::Degrees) );
 		}
 
+		ParentCategory.OnPasteFromText()->AddSP(this, &FComponentTransformDetails::OnPasteFromText, ETransformField::Rotation);
+
 		ChildrenBuilder.AddCustomRow( LOCTEXT("RotationFilter", "Rotation") )
 		.RowTag("Rotation")
 		.CopyAction( CreateCopyAction(ETransformField::Rotation) )
@@ -494,6 +545,8 @@ void FComponentTransformDetails::GenerateChildContent( IDetailChildrenBuilder& C
 	// Scale
 	if(!bHideScaleField)
 	{
+		ParentCategory.OnPasteFromText()->AddSP(this, &FComponentTransformDetails::OnPasteFromText, ETransformField::Scale);
+		
 		ChildrenBuilder.AddCustomRow( LOCTEXT("ScaleFilter", "Scale") )
 		.RowTag("Scale")
 		.CopyAction( CreateCopyAction(ETransformField::Scale) )
diff --git a/Engine/Source/Editor/DetailCustomizations/Private/ComponentTransformDetails.h b/Engine/Source/Editor/DetailCustomizations/Private/ComponentTransformDetails.h
--- a/Engine/Source/Editor/DetailCustomizations/Private/ComponentTransformDetails.h
+++ b/Engine/Source/Editor/DetailCustomizations/Private/ComponentTransformDetails.h
@@ -188,6 +188,10 @@ private:
 	 */
 	void OnPaste( ETransformField::Type TransformField );
 
+	void OnPasteFromText(const FString& InTag, const FString& InText, const TOptional<FGuid>& InOperationId, ETransformField::Type InTransformField);
+
+	void PasteFromText(const FString& InTag, const FString& InText, ETransformField::Type InTransformField);
+	
 	/**
 	 * Creates a UI action for copying a specified transform field
 	 */
@@ -264,7 +268,7 @@ private:
 	struct FOptionalVector
 	{
 		/**
-		 * Sets the value from an FVector                   
+		 * Sets the value from an FVector
 		 */
 		void Set( const FVector& InVec )
 		{
diff --git a/Engine/Source/Editor/DetailCustomizations/Private/MathStructProxyCustomizations.cpp b/Engine/Source/Editor/DetailCustomizations/Private/MathStructProxyCustomizations.cpp
--- a/Engine/Source/Editor/DetailCustomizations/Private/MathStructProxyCustomizations.cpp
+++ b/Engine/Source/Editor/DetailCustomizations/Private/MathStructProxyCustomizations.cpp
@@ -163,6 +163,11 @@ void FMatrixStructCustomization<T>::CustomizeLocation(TSharedRef<class IProperty
 {
 	TWeakPtr<IPropertyHandle> WeakHandlePtr = StructPropertyHandle;
 
+	if (Row.IsPasteFromTextBound())
+	{
+		Row.OnPasteFromTextDelegate.Pin()->AddSP(this, &FMatrixStructCustomization<T>::OnPasteFromText, FTransformField::Location, WeakHandlePtr);
+	}
+	
 	Row
 	.CopyAction(FUIAction(FExecuteAction::CreateSP(this, &FMatrixStructCustomization<T>::OnCopy, FTransformField::Location, WeakHandlePtr)))
 	.PasteAction(FUIAction(FExecuteAction::CreateSP(this, &FMatrixStructCustomization<T>::OnPaste, FTransformField::Location, WeakHandlePtr)))
@@ -198,6 +203,11 @@ void FMatrixStructCustomization<T>::CustomizeRotation(TSharedRef<class IProperty
 {
 	TWeakPtr<IPropertyHandle> WeakHandlePtr = StructPropertyHandle;
 
+	if (Row.IsPasteFromTextBound())
+	{
+		Row.OnPasteFromTextDelegate.Pin()->AddSP(this, &FMatrixStructCustomization<T>::OnPasteFromText, FTransformField::Rotation, WeakHandlePtr);
+	}
+
 	Row
 	.CopyAction(FUIAction(FExecuteAction::CreateSP(this, &FMatrixStructCustomization<T>::OnCopy, FTransformField::Rotation, WeakHandlePtr)))
 	.PasteAction(FUIAction(FExecuteAction::CreateSP(this, &FMatrixStructCustomization<T>::OnPaste, FTransformField::Rotation, WeakHandlePtr)))
@@ -233,6 +243,11 @@ void FMatrixStructCustomization<T>::CustomizeScale(TSharedRef<class IPropertyHan
 {
 	TWeakPtr<IPropertyHandle> WeakHandlePtr = StructPropertyHandle;
 
+	if (Row.IsPasteFromTextBound())
+	{
+		Row.OnPasteFromTextDelegate.Pin()->AddSP(this, &FMatrixStructCustomization<T>::OnPasteFromText, FTransformField::Scale, WeakHandlePtr);
+	}
+
 	Row
 	.CopyAction(FUIAction(FExecuteAction::CreateSP(this, &FMatrixStructCustomization<T>::OnCopy, FTransformField::Scale, WeakHandlePtr)))
 	.PasteAction(FUIAction(FExecuteAction::CreateSP(this, &FMatrixStructCustomization<T>::OnPaste, FTransformField::Scale, WeakHandlePtr)))
@@ -321,15 +336,37 @@ void FMatrixStructCustomization<T>::OnCopy(FTransformField::Type Type, TWeakPtr<
 template<typename T>
 void FMatrixStructCustomization<T>::OnPaste(FTransformField::Type Type, TWeakPtr<IPropertyHandle> PropertyHandlePtr)
 {
-	auto PropertyHandle = PropertyHandlePtr.Pin();
+	FString PastedText;
+	FPlatformApplicationMisc::ClipboardPaste(PastedText);
+
+	PasteFromText(TEXT(""), PastedText, Type, PropertyHandlePtr);
+}
+
+template <typename T>
+void FMatrixStructCustomization<T>::OnPasteFromText(
+	const FString& InTag,
+	const FString& InText,
+	const TOptional<FGuid>& InOperationId,
+	FTransformField::Type Type,
+	TWeakPtr<IPropertyHandle> PropertyHandlePtr)
+{
+	PasteFromText(InTag, InText, Type, PropertyHandlePtr);
+}
 
+template <typename T>
+void FMatrixStructCustomization<T>::PasteFromText(
+	const FString& InTag,
+	const FString& InText,
+	FTransformField::Type Type,
+	TWeakPtr<IPropertyHandle> PropertyHandlePtr)
+{
+	auto PropertyHandle = PropertyHandlePtr.Pin();
 	if (!PropertyHandle.IsValid())
 	{
 		return;
 	}
 
-	FString PastedText;
-	FPlatformApplicationMisc::ClipboardPaste(PastedText);
+	FString PastedText = InText;
 
 	switch (Type)
 	{
diff --git a/Engine/Source/Editor/DetailCustomizations/Private/PhysicsConstraintComponentDetails.cpp b/Engine/Source/Editor/DetailCustomizations/Private/PhysicsConstraintComponentDetails.cpp
--- a/Engine/Source/Editor/DetailCustomizations/Private/PhysicsConstraintComponentDetails.cpp
+++ b/Engine/Source/Editor/DetailCustomizations/Private/PhysicsConstraintComponentDetails.cpp
@@ -121,6 +121,11 @@ void FConstraintTransformCustomization::MakeRotationRow(TSharedRef<class IProper
 
 	const FText ComponentToolTipFormatText = LOCTEXT("ConstraintTransformRotationToolTip", "{0} component of the constraint rotation relative to the {1} bone (in degrees).");
 
+	if (Row.IsPasteFromTextBound())
+	{
+		Row.OnPasteFromTextDelegate.Pin()->AddSP(this, &FConstraintTransformCustomization::OnPasteRotationFromText, WeakPriAxisHandlePtr, WeakSecAxisHandlePtr);
+	}
+	
 	Row
 	.CopyAction(FUIAction(FExecuteAction::CreateSP(this, &FConstraintTransformCustomization::OnCopyRotation, WeakPriAxisHandlePtr, WeakSecAxisHandlePtr)))
 	.PasteAction(FUIAction(FExecuteAction::CreateSP(this, &FConstraintTransformCustomization::OnPasteRotation, WeakPriAxisHandlePtr, WeakSecAxisHandlePtr)))
@@ -172,6 +177,11 @@ void FConstraintTransformCustomization::MakePositionRow(TSharedRef<class IProper
 
 	const FText ComponentToolTipFormatText = LOCTEXT("ConstraintTransformPositionToolTip", "{0} component of the constraint position relative to the {1} bone.");
 
+	if (Row.IsPasteFromTextBound())
+	{
+		Row.OnPasteFromTextDelegate.Pin()->AddSP(this, &FConstraintTransformCustomization::OnPastePositionFromText, WeakPositionHandlePtr);
+	}
+	
 	Row
 	.CopyAction(FUIAction(FExecuteAction::CreateSP(this, &FConstraintTransformCustomization::OnCopyPosition, WeakPositionHandlePtr)))
 	.PasteAction(FUIAction(FExecuteAction::CreateSP(this, &FConstraintTransformCustomization::OnPastePosition, WeakPositionHandlePtr)))
@@ -479,6 +489,20 @@ void FConstraintTransformCustomization::OnCopy(TWeakPtr<IPropertyHandle> Positio
 }
 
 void FConstraintTransformCustomization::OnPaste(TWeakPtr<IPropertyHandle> PositionPropertyHandlePtr, TWeakPtr<IPropertyHandle> PriAxisPropertyHandlePtr, TWeakPtr<IPropertyHandle> SecAxisPropertyHandlePtr)
+{
+	FString PastedText;
+	FPlatformApplicationMisc::ClipboardPaste(PastedText);
+
+	OnPasteFromText(TEXT(""), PastedText, {}, PositionPropertyHandlePtr, PriAxisPropertyHandlePtr, SecAxisPropertyHandlePtr);
+}
+
+void FConstraintTransformCustomization::OnPasteFromText(
+	const FString& InTag,
+	const FString& InText,
+	const TOptional<FGuid>& InOperationId,
+	TWeakPtr<IPropertyHandle> PositionPropertyHandlePtr,
+	TWeakPtr<IPropertyHandle> PriAxisPropertyHandlePtr,
+	TWeakPtr<IPropertyHandle> SecAxisPropertyHandlePtr)
 {
 	TSharedPtr<IPropertyHandle> PositionPropertyHandle = PositionPropertyHandlePtr.Pin();
 	TSharedPtr<IPropertyHandle> PriAxisPropertyHandle = PriAxisPropertyHandlePtr.Pin();
@@ -489,12 +513,9 @@ void FConstraintTransformCustomization::OnPaste(TWeakPtr<IPropertyHandle> Positi
 		return;
 	}
 
-	FString PastedText;
-	FPlatformApplicationMisc::ClipboardPaste(PastedText);
-
 	{
 		FScopedTransaction Transaction(LOCTEXT("PastePosition", "Paste Position"));
-		SetValueFromFormattedString(PastedText);
+		SetValueFromFormattedString(InText);
 		FlushValues(PriAxisPropertyHandle, SecAxisPropertyHandle);
 		FlushValues(PositionPropertyHandle);
 	}
@@ -542,6 +563,26 @@ void FConstraintTransformCustomization::OnCopyRotation(TWeakPtr<IPropertyHandle>
 }
 
 void FConstraintTransformCustomization::OnPastePosition(TWeakPtr<IPropertyHandle> PositionPropertyHandlePtr)
+{
+	FString PastedText;
+	FPlatformApplicationMisc::ClipboardPaste(PastedText);
+
+	PastePositionFromText(TEXT(""), PastedText, PositionPropertyHandlePtr);
+}
+
+void FConstraintTransformCustomization::OnPastePositionFromText(
+	const FString& InTag,
+	const FString& InText,
+	const TOptional<FGuid>& InOperationId,
+	TWeakPtr<IPropertyHandle> PositionPropertyHandlePtr)
+{
+	PastePositionFromText(InTag, InText, PositionPropertyHandlePtr);
+}
+
+void FConstraintTransformCustomization::PastePositionFromText(
+	const FString& InTag,
+	const FString& InText,
+	TWeakPtr<IPropertyHandle> PositionPropertyHandlePtr)
 {
 	TSharedPtr<IPropertyHandle> PositionPropertyHandle = PositionPropertyHandlePtr.Pin();
 
@@ -550,17 +591,36 @@ void FConstraintTransformCustomization::OnPastePosition(TWeakPtr<IPropertyHandle
 		return;
 	}
 
-	FString PastedText;
-	FPlatformApplicationMisc::ClipboardPaste(PastedText);
-
 	{
 		FScopedTransaction Transaction(LOCTEXT("PastePosition", "Paste Position"));
-		SetPositionFromFormattedString(PastedText);
+		SetPositionFromFormattedString(InText);
 		FlushValues(PositionPropertyHandle);
 	}
 }
 
 void FConstraintTransformCustomization::OnPasteRotation(TWeakPtr<IPropertyHandle> PriAxisPropertyHandlePtr, TWeakPtr<IPropertyHandle> SecAxisPropertyHandlePtr)
+{
+	FString PastedText;
+	FPlatformApplicationMisc::ClipboardPaste(PastedText);
+
+	PasteRotationFromText(TEXT(""), PastedText, PriAxisPropertyHandlePtr, SecAxisPropertyHandlePtr);
+}
+
+void FConstraintTransformCustomization::OnPasteRotationFromText(
+	const FString& InTag,
+	const FString& InText,
+	const TOptional<FGuid>& InOperationId,
+	TWeakPtr<IPropertyHandle> PriAxisPropertyHandlePtr,
+	TWeakPtr<IPropertyHandle> SecAxisPropertyHandlePtr)
+{
+	PasteRotationFromText(InTag, InText, PriAxisPropertyHandlePtr, SecAxisPropertyHandlePtr);
+}
+
+void FConstraintTransformCustomization::PasteRotationFromText(
+	const FString& InTag,
+	const FString& InText,
+	TWeakPtr<IPropertyHandle> PriAxisPropertyHandlePtr,
+	TWeakPtr<IPropertyHandle> SecAxisPropertyHandlePtr)
 {
 	TSharedPtr<IPropertyHandle> PriAxisPropertyHandle = PriAxisPropertyHandlePtr.Pin();
 	TSharedPtr<IPropertyHandle> SecAxisPropertyHandle = SecAxisPropertyHandlePtr.Pin();
@@ -570,12 +630,9 @@ void FConstraintTransformCustomization::OnPasteRotation(TWeakPtr<IPropertyHandle
 		return;
 	}
 
-	FString PastedText;
-	FPlatformApplicationMisc::ClipboardPaste(PastedText);
-	
 	{
 		FScopedTransaction Transaction(LOCTEXT("PasteRotation", "Paste Rotation"));
-		SetRotationFromFormattedString(PastedText);
+		SetRotationFromFormattedString(InText);
 		FlushValues(PriAxisPropertyHandle, SecAxisPropertyHandle);
 	}
 }
@@ -1024,6 +1081,11 @@ void FPhysicsConstraintComponentDetails::AddConstraintFrameTransform(const ECons
 		HeaderRow.CopyAction(FUIAction(FExecuteAction::CreateSP(this, &FPhysicsConstraintComponentDetails::OnCopyConstraintTransform, PositionPropertyHandle, PriAxisPropertyHandle, SecAxisPropertyHandle, TransformProxy)));
 		HeaderRow.PasteAction(FUIAction(FExecuteAction::CreateSP(this, &FPhysicsConstraintComponentDetails::OnPasteConstraintTransform, PositionPropertyHandle, PriAxisPropertyHandle, SecAxisPropertyHandle, TransformProxy)));
 
+		if (HeaderRow.IsPasteFromTextBound())
+		{
+			HeaderRow.OnPasteFromTextDelegate.Pin()->AddSP(this, &FPhysicsConstraintComponentDetails::OnPasteConstraintTransformFromText, PositionPropertyHandle, PriAxisPropertyHandle, SecAxisPropertyHandle, TransformProxy);
+		}
+		
 		HeaderRow.OverrideResetToDefault(
 			FResetToDefaultOverride::Create(
 				FIsResetToDefaultVisible::CreateSP(this, &FPhysicsConstraintComponentDetails::IsSnapConstraintTransformComponentVisible, SnapFlag),
@@ -1191,13 +1253,39 @@ void FPhysicsConstraintComponentDetails::OnCopyConstraintTransform(TSharedPtr<IP
 }
 
 void FPhysicsConstraintComponentDetails::OnPasteConstraintTransform(TSharedPtr<IPropertyHandle> PositionPropertyHandle, TSharedPtr<IPropertyHandle> PriAxisPropertyHandle, TSharedPtr<IPropertyHandle> SecAxisPropertyHandle, TSharedPtr<FConstraintTransformCustomization> TransformProxy)
+{
+	FString PastedText;
+	FPlatformApplicationMisc::ClipboardPaste(PastedText);
+	
+	PasteConstraintTransformFromText(TEXT(""), PastedText, PositionPropertyHandle, PriAxisPropertyHandle, SecAxisPropertyHandle, TransformProxy);
+}
+
+void FPhysicsConstraintComponentDetails::OnPasteConstraintTransformFromText(
+	const FString& InTag,
+	const FString& InText,
+	const TOptional<FGuid>& InOperationId,
+	TSharedPtr<IPropertyHandle> PositionPropertyHandle,
+	TSharedPtr<IPropertyHandle> PriAxisPropertyHandle,
+	TSharedPtr<IPropertyHandle> SecAxisPropertyHandle,
+	TSharedPtr<FConstraintTransformCustomization> TransformProxy)
 {
 	if (TransformProxy.IsValid())
 	{
-		TransformProxy->OnPaste(PositionPropertyHandle, PriAxisPropertyHandle, SecAxisPropertyHandle);
+		TransformProxy->OnPasteFromText(InTag, InText, InOperationId, PositionPropertyHandle, PriAxisPropertyHandle, SecAxisPropertyHandle);
 	}
 }
 
+void FPhysicsConstraintComponentDetails::PasteConstraintTransformFromText(
+	const FString& InTag,
+	const FString& InText,
+	TSharedPtr<IPropertyHandle> PositionPropertyHandle,
+	TSharedPtr<IPropertyHandle> PriAxisPropertyHandle,
+	TSharedPtr<IPropertyHandle> SecAxisPropertyHandle,
+	TSharedPtr<FConstraintTransformCustomization> TransformProxy)
+{
+	OnPasteConstraintTransformFromText(InTag, InText, {}, PositionPropertyHandle, PriAxisPropertyHandle, SecAxisPropertyHandle, TransformProxy);
+}
+
 void FPhysicsConstraintComponentDetails::AddConstraintBehaviorProperties(IDetailLayoutBuilder& DetailBuilder, TSharedPtr<IPropertyHandle> ConstraintInstance, TSharedPtr<IPropertyHandle> ProfilePropertiesProperty)
 {
 	IDetailCategoryBuilder& ConstraintCat = DetailBuilder.EditCategory("Constraint Behavior");
diff --git a/Engine/Source/Editor/DetailCustomizations/Private/PhysicsConstraintComponentDetails.h b/Engine/Source/Editor/DetailCustomizations/Private/PhysicsConstraintComponentDetails.h
--- a/Engine/Source/Editor/DetailCustomizations/Private/PhysicsConstraintComponentDetails.h
+++ b/Engine/Source/Editor/DetailCustomizations/Private/PhysicsConstraintComponentDetails.h
@@ -64,6 +64,9 @@ public:
 	void OnPastePosition(TWeakPtr<IPropertyHandle> PositionPropertyHandlePtr);
 	void OnPasteRotation(TWeakPtr<IPropertyHandle> PriAxisPropertyHandlePtr, TWeakPtr<IPropertyHandle> SecAxisPropertyHandlePtr);
 
+	void OnPasteFromText(const FString& InTag, const FString& InText, const TOptional<FGuid>& InOperationId,
+		TWeakPtr<IPropertyHandle> PositionPropertyHandlePtr, TWeakPtr<IPropertyHandle> PriAxisPropertyHandlePtr, TWeakPtr<IPropertyHandle> SecAxisPropertyHandlePtr);
+	
 	bool IsRotationValueEnabled(TWeakPtr<IPropertyHandle> PriAxisPropertyHandlePtr, TWeakPtr<IPropertyHandle> SecAxisPropertyHandlePtr) const;
 	template<typename ProxyType, typename NumericType> TOptional<NumericType> OnGetRotationValue(TWeakPtr<IPropertyHandle> PriAxisPropertyHandlePtr, TWeakPtr<IPropertyHandle> SecAxisPropertyHandlePtr, TSharedRef< TProxyProperty<ProxyType, NumericType> > ProxyValue) const;
 	template<typename ProxyType, typename NumericType> void OnRotationValueCommitted(NumericType NewValue, ETextCommit::Type CommitType, TWeakPtr<IPropertyHandle> PriAxisPropertyHandlePtr, TWeakPtr<IPropertyHandle> SecAxisPropertyHandlePtr, TSharedRef< TProxyProperty<ProxyType, NumericType> > ProxyValue);
@@ -99,6 +102,12 @@ protected:
 	void OrthonormalVectorPairToDisplayedRotator(const FVector& PriAxis, const FVector& SecAxis, FRotator& OutRotator) const;
 	void DisplayedRotatorToOrthonormalVectorPair(const FRotator& InRotator, FVector& OutPriAxis, FVector& OutSecAxis) const;
 
+	void OnPastePositionFromText(const FString& InTag, const FString& InText, const TOptional<FGuid>& InOperationId, TWeakPtr<IPropertyHandle> PositionPropertyHandlePtr);
+	void PastePositionFromText(const FString& InTag, const FString& InText, TWeakPtr<IPropertyHandle> PositionPropertyHandlePtr);
+	
+	void OnPasteRotationFromText(const FString& InTag, const FString& InText, const TOptional<FGuid>& InOperationId, TWeakPtr<IPropertyHandle> PriAxisPropertyHandlePtr, TWeakPtr<IPropertyHandle> SecAxisPropertyHandlePtr);
+	void PasteRotationFromText(const FString& InTag, const FString& InText, TWeakPtr<IPropertyHandle> PriAxisPropertyHandlePtr, TWeakPtr<IPropertyHandle> SecAxisPropertyHandlePtr);
+
 	mutable TSharedRef< TProxyValue<FRotator> > CachedRotation;
 	mutable TSharedRef< TProxyProperty<FRotator, FRotator::FReal> > CachedRotationYaw;
 	mutable TSharedRef< TProxyProperty<FRotator, FRotator::FReal> > CachedRotationPitch;
@@ -180,7 +189,23 @@ private:
 	FSlateColor GetConstraintTransformColorAndOpacity(const EConstraintTransformComponentFlags ComponentFlags) const;
 
 private:
-
+	void OnPasteConstraintTransformFromText(
+		const FString& InTag,
+		const FString& InText,
+		const TOptional<FGuid>& InOperationId,
+		TSharedPtr<IPropertyHandle> PositionPropertyHandle,
+		TSharedPtr<IPropertyHandle> PriAxisPropertyHandle,
+		TSharedPtr<IPropertyHandle> SecAxisPropertyHandle,
+		TSharedPtr<FConstraintTransformCustomization> TransformProxy);
+
+	void PasteConstraintTransformFromText(
+		const FString& InTag,
+		const FString& InText,
+		TSharedPtr<IPropertyHandle> PositionPropertyHandle,
+		TSharedPtr<IPropertyHandle> PriAxisPropertyHandle,
+		TSharedPtr<IPropertyHandle> SecAxisPropertyHandle,
+		TSharedPtr<FConstraintTransformCustomization> TransformProxy);
+	
 	TSharedRef<SWidget> MakeEditSpaceToggleButtonWidget(TSharedPtr<IPropertyHandle> ConstraintInstancePropertyHandle, const EConstraintTransformComponentFlags ComponentFlags); // TODO - get this in the right place in src
 	void SnapConstraintTransformComponentsToDefault(TSharedPtr<IPropertyHandle> ConstraintInstancePropertyHandle, const EConstraintTransformComponentFlags SnapFlags);
 	void UpdateTransformProxyDisplayRelativeToDefault(TSharedPtr<IPropertyHandle> ConstraintInstancePropertyHandle);
diff --git a/Engine/Source/Editor/DetailCustomizations/Private/SplineComponentDetails.cpp b/Engine/Source/Editor/DetailCustomizations/Private/SplineComponentDetails.cpp
--- a/Engine/Source/Editor/DetailCustomizations/Private/SplineComponentDetails.cpp
+++ b/Engine/Source/Editor/DetailCustomizations/Private/SplineComponentDetails.cpp
@@ -282,6 +282,9 @@ private:
 	void OnCopy(ESplinePointProperty SplinePointProp);
 	void OnPaste(ESplinePointProperty SplinePointProp);
 
+	void OnPasteFromText(const FString& InTag, const FString& InText, const TOptional<FGuid>& InOperationId, ESplinePointProperty SplinePointProp);
+	void PasteFromText(const FString& InTag, const FString& InText, ESplinePointProperty SplinePointProp);
+
 	void OnBeginPositionSlider();
 	void OnBeginScaleSlider();
 	void OnEndSlider(float);
@@ -516,15 +519,21 @@ void FSplinePointDetails::GenerateChildContent(IDetailChildrenBuilder& ChildrenB
 		[
 			SNew(SNumericEntryBox<float>)
 			.IsEnabled(TAttribute<bool>(this, &FSplinePointDetails::IsOnePointSelected))
-		.Value(this, &FSplinePointDetails::GetInputKey)
-		.UndeterminedString(LOCTEXT("Multiple", "Multiple"))
-		.OnValueCommitted(this, &FSplinePointDetails::OnSetInputKey)
-		.Font(IDetailLayoutBuilder::GetDetailFont())
+			.Value(this, &FSplinePointDetails::GetInputKey)
+			.UndeterminedString(LOCTEXT("Multiple", "Multiple"))
+			.OnValueCommitted(this, &FSplinePointDetails::OnSetInputKey)
+			.Font(IDetailLayoutBuilder::GetDetailFont())
 		];
 
+	IDetailCategoryBuilder& ParentCategory = ChildrenBuilder.GetParentCategory();
+	TSharedPtr<FOnPasteFromText> PasteFromTextDelegate = ParentCategory.OnPasteFromText();
+	const bool bUsePasteFromText = PasteFromTextDelegate.IsValid();	
+
 	// Position
 	if (SplineComp->AllowsSpinePointLocationEditing())
 	{
+		PasteFromTextDelegate->AddSP(this, &FSplinePointDetails::OnPasteFromText, ESplinePointProperty::Location);
+		
 		ChildrenBuilder.AddCustomRow(LOCTEXT("Location", "Location"))
 			.RowTag(TEXT("Location"))
 			.CopyAction(CreateCopyAction(ESplinePointProperty::Location))
@@ -542,26 +551,28 @@ void FSplinePointDetails::GenerateChildContent(IDetailChildrenBuilder& ChildrenB
 			[
 				SNew(SVectorInputBox)
 				.X(this, &FSplinePointDetails::GetPositionX)
-			.Y(this, &FSplinePointDetails::GetPositionY)
-			.Z(this, &FSplinePointDetails::GetPositionZ)
-			.AllowSpin(true)
-			.bColorAxisLabels(true)
-			.SpinDelta(1.f)
-			.OnXChanged(this, &FSplinePointDetails::OnSetPosition, ETextCommit::Default, EAxis::X)
-			.OnYChanged(this, &FSplinePointDetails::OnSetPosition, ETextCommit::Default, EAxis::Y)
-			.OnZChanged(this, &FSplinePointDetails::OnSetPosition, ETextCommit::Default, EAxis::Z)
-			.OnXCommitted(this, &FSplinePointDetails::OnSetPosition, EAxis::X)
-			.OnYCommitted(this, &FSplinePointDetails::OnSetPosition, EAxis::Y)
-			.OnZCommitted(this, &FSplinePointDetails::OnSetPosition, EAxis::Z)
-			.OnBeginSliderMovement(this, &FSplinePointDetails::OnBeginPositionSlider)
-			.OnEndSliderMovement(this, &FSplinePointDetails::OnEndSlider)
-			.Font(IDetailLayoutBuilder::GetDetailFont())
+				.Y(this, &FSplinePointDetails::GetPositionY)
+				.Z(this, &FSplinePointDetails::GetPositionZ)
+				.AllowSpin(true)
+				.bColorAxisLabels(true)
+				.SpinDelta(1.f)
+				.OnXChanged(this, &FSplinePointDetails::OnSetPosition, ETextCommit::Default, EAxis::X)
+				.OnYChanged(this, &FSplinePointDetails::OnSetPosition, ETextCommit::Default, EAxis::Y)
+				.OnZChanged(this, &FSplinePointDetails::OnSetPosition, ETextCommit::Default, EAxis::Z)
+				.OnXCommitted(this, &FSplinePointDetails::OnSetPosition, EAxis::X)
+				.OnYCommitted(this, &FSplinePointDetails::OnSetPosition, EAxis::Y)
+				.OnZCommitted(this, &FSplinePointDetails::OnSetPosition, EAxis::Z)
+				.OnBeginSliderMovement(this, &FSplinePointDetails::OnBeginPositionSlider)
+				.OnEndSliderMovement(this, &FSplinePointDetails::OnEndSlider)
+				.Font(IDetailLayoutBuilder::GetDetailFont())
 			];
 	}
 
 	// Rotation
 	if (SplineComp->AllowsSplinePointRotationEditing())
 	{
+		PasteFromTextDelegate->AddSP(this, &FSplinePointDetails::OnPasteFromText, ESplinePointProperty::Rotation);	
+		
 		ChildrenBuilder.AddCustomRow(LOCTEXT("Rotation", "Rotation"))
 			.RowTag(TEXT("Rotation"))
 			.CopyAction(CreateCopyAction(ESplinePointProperty::Rotation))
@@ -579,20 +590,22 @@ void FSplinePointDetails::GenerateChildContent(IDetailChildrenBuilder& ChildrenB
 			[
 				SNew(SRotatorInputBox)
 				.Roll(this, &FSplinePointDetails::GetRotationRoll)
-			.Pitch(this, &FSplinePointDetails::GetRotationPitch)
-			.Yaw(this, &FSplinePointDetails::GetRotationYaw)
-			.AllowSpin(false)
-			.bColorAxisLabels(false)
-			.OnRollCommitted(this, &FSplinePointDetails::OnSetRotation, EAxis::X)
-			.OnPitchCommitted(this, &FSplinePointDetails::OnSetRotation, EAxis::Y)
-			.OnYawCommitted(this, &FSplinePointDetails::OnSetRotation, EAxis::Z)
-			.Font(IDetailLayoutBuilder::GetDetailFont())
+				.Pitch(this, &FSplinePointDetails::GetRotationPitch)
+				.Yaw(this, &FSplinePointDetails::GetRotationYaw)
+				.AllowSpin(false)
+				.bColorAxisLabels(false)
+				.OnRollCommitted(this, &FSplinePointDetails::OnSetRotation, EAxis::X)
+				.OnPitchCommitted(this, &FSplinePointDetails::OnSetRotation, EAxis::Y)
+				.OnYawCommitted(this, &FSplinePointDetails::OnSetRotation, EAxis::Z)
+				.Font(IDetailLayoutBuilder::GetDetailFont())
 			];
 	}
 
 	// Scale
 	if (SplineComp->AllowsSplinePointScaleEditing())
 	{
+		PasteFromTextDelegate->AddSP(this, &FSplinePointDetails::OnPasteFromText, ESplinePointProperty::Scale);
+		
 		ChildrenBuilder.AddCustomRow(LOCTEXT("Scale", "Scale"))
 			.RowTag(TEXT("Scale"))
 			.Visibility(TAttribute<EVisibility>(this, &FSplinePointDetails::IsEnabled))
@@ -631,6 +644,8 @@ void FSplinePointDetails::GenerateChildContent(IDetailChildrenBuilder& ChildrenB
 	// ArriveTangent
 	if (SplineComp->AllowsSplinePointArriveTangentEditing())
 	{
+		PasteFromTextDelegate->AddSP(this, &FSplinePointDetails::OnPasteFromText, ESplinePointProperty::ArriveTangent);
+	
 		ChildrenBuilder.AddCustomRow(LOCTEXT("ArriveTangent", "Arrive Tangent"))
 			.RowTag(TEXT("ArriveTangent"))
 			.Visibility(TAttribute<EVisibility>(this, &FSplinePointDetails::IsEnabled))
@@ -642,7 +657,7 @@ void FSplinePointDetails::GenerateChildContent(IDetailChildrenBuilder& ChildrenB
 			[
 				SNew(STextBlock)
 				.Text(LOCTEXT("ArriveTangent", "Arrive Tangent"))
-			.Font(IDetailLayoutBuilder::GetDetailFont())
+				.Font(IDetailLayoutBuilder::GetDetailFont())
 			]
 		.ValueContent()
 			.MinDesiredWidth(375.0f)
@@ -650,21 +665,22 @@ void FSplinePointDetails::GenerateChildContent(IDetailChildrenBuilder& ChildrenB
 			[
 				SNew(SVectorInputBox)
 				.X(this, &FSplinePointDetails::GetArriveTangentX)
-			.Y(this, &FSplinePointDetails::GetArriveTangentY)
-			.Z(this, &FSplinePointDetails::GetArriveTangentZ)
-			.AllowSpin(false)
-			.bColorAxisLabels(false)
-			.OnXCommitted(this, &FSplinePointDetails::OnSetArriveTangent, EAxis::X)
-			.OnYCommitted(this, &FSplinePointDetails::OnSetArriveTangent, EAxis::Y)
-			.OnZCommitted(this, &FSplinePointDetails::OnSetArriveTangent, EAxis::Z)
-			.Font(IDetailLayoutBuilder::GetDetailFont())
+				.Y(this, &FSplinePointDetails::GetArriveTangentY)
+				.Z(this, &FSplinePointDetails::GetArriveTangentZ)
+				.AllowSpin(false)
+				.bColorAxisLabels(false)
+				.OnXCommitted(this, &FSplinePointDetails::OnSetArriveTangent, EAxis::X)
+				.OnYCommitted(this, &FSplinePointDetails::OnSetArriveTangent, EAxis::Y)
+				.OnZCommitted(this, &FSplinePointDetails::OnSetArriveTangent, EAxis::Z)
+				.Font(IDetailLayoutBuilder::GetDetailFont())
 			];
 	}
 
-
 	// LeaveTangent
 	if (SplineComp->AllowsSplinePointLeaveTangentEditing())
 	{
+		PasteFromTextDelegate->AddSP(this, &FSplinePointDetails::OnPasteFromText, ESplinePointProperty::LeaveTangent);
+	
 		ChildrenBuilder.AddCustomRow(LOCTEXT("LeaveTangent", "Leave Tangent"))
 			.RowTag(TEXT("LeaveTangent"))
 			.Visibility(TAttribute<EVisibility>(this, &FSplinePointDetails::IsEnabled))
@@ -676,7 +692,7 @@ void FSplinePointDetails::GenerateChildContent(IDetailChildrenBuilder& ChildrenB
 			[
 				SNew(STextBlock)
 				.Text(LOCTEXT("LeaveTangent", "Leave Tangent"))
-			.Font(IDetailLayoutBuilder::GetDetailFont())
+				.Font(IDetailLayoutBuilder::GetDetailFont())
 			]
 		.ValueContent()
 			.MinDesiredWidth(375.0f)
@@ -684,14 +700,14 @@ void FSplinePointDetails::GenerateChildContent(IDetailChildrenBuilder& ChildrenB
 			[
 				SNew(SVectorInputBox)
 				.X(this, &FSplinePointDetails::GetLeaveTangentX)
-			.Y(this, &FSplinePointDetails::GetLeaveTangentY)
-			.Z(this, &FSplinePointDetails::GetLeaveTangentZ)
-			.AllowSpin(false)
-			.bColorAxisLabels(false)
-			.OnXCommitted(this, &FSplinePointDetails::OnSetLeaveTangent, EAxis::X)
-			.OnYCommitted(this, &FSplinePointDetails::OnSetLeaveTangent, EAxis::Y)
-			.OnZCommitted(this, &FSplinePointDetails::OnSetLeaveTangent, EAxis::Z)
-			.Font(IDetailLayoutBuilder::GetDetailFont())
+				.Y(this, &FSplinePointDetails::GetLeaveTangentY)
+				.Z(this, &FSplinePointDetails::GetLeaveTangentZ)
+				.AllowSpin(false)
+				.bColorAxisLabels(false)
+				.OnXCommitted(this, &FSplinePointDetails::OnSetLeaveTangent, EAxis::X)
+				.OnYCommitted(this, &FSplinePointDetails::OnSetLeaveTangent, EAxis::Y)
+				.OnZCommitted(this, &FSplinePointDetails::OnSetLeaveTangent, EAxis::Z)
+				.Font(IDetailLayoutBuilder::GetDetailFont())
 			];
 	}
 
@@ -715,13 +731,13 @@ void FSplinePointDetails::GenerateChildContent(IDetailChildrenBuilder& ChildrenB
 			[
 				SNew(SComboBox<TSharedPtr<FString>>)
 				.OptionsSource(&SplinePointTypes)
-			.OnGenerateWidget(this, &FSplinePointDetails::OnGenerateComboWidget)
-			.OnSelectionChanged(this, &FSplinePointDetails::OnSplinePointTypeChanged)
-			[
-				SNew(STextBlock)
-				.Font(IDetailLayoutBuilder::GetDetailFont())
-			.Text(this, &FSplinePointDetails::GetPointType)
-			]
+				.OnGenerateWidget(this, &FSplinePointDetails::OnGenerateComboWidget)
+				.OnSelectionChanged(this, &FSplinePointDetails::OnSplinePointTypeChanged)
+				[
+					SNew(STextBlock)
+					.Font(IDetailLayoutBuilder::GetDetailFont())
+					.Text(this, &FSplinePointDetails::GetPointType)
+				]
 			];
 	}
 
@@ -1563,61 +1579,79 @@ void FSplinePointDetails::OnPaste(ESplinePointProperty SplinePointProp)
 	FString PastedText;
 	FPlatformApplicationMisc::ClipboardPaste(PastedText);
 
+	PasteFromText(TEXT(""), PastedText, SplinePointProp);
+}
+
+void FSplinePointDetails::OnPasteFromText(
+	const FString& InTag,
+	const FString& InText,
+	const TOptional<FGuid>& InOperationId,
+	ESplinePointProperty SplinePointProp)
+{
+	PasteFromText(InTag, InText, SplinePointProp);
+}
+
+void FSplinePointDetails::PasteFromText(
+	const FString& InTag,
+	const FString& InText,
+	ESplinePointProperty SplinePointProp)
+{
+	FString PastedText = InText;
 	switch (SplinePointProp)
 	{
 	case ESplinePointProperty::Location:
-	{
-		FVector NewLocation;
-		if (NewLocation.InitFromString(PastedText))
 		{
-			FScopedTransaction Transaction(LOCTEXT("PasteLocation", "Paste Location"));
-			SetSplinePointProperty(ESplinePointProperty::Location, NewLocation, EAxisList::All, true);
+			FVector NewLocation;
+			if (NewLocation.InitFromString(PastedText))
+			{
+				FScopedTransaction Transaction(LOCTEXT("PasteLocation", "Paste Location"));
+				SetSplinePointProperty(ESplinePointProperty::Location, NewLocation, EAxisList::All, true);
+			}
+			break;
 		}
-		break;
-	}
 	case ESplinePointProperty::Rotation:
-	{
-		FVector NewRotation;
-		PastedText.ReplaceInline(TEXT("Pitch="), TEXT("X="));
-		PastedText.ReplaceInline(TEXT("Yaw="), TEXT("Y="));
-		PastedText.ReplaceInline(TEXT("Roll="), TEXT("Z="));
-		if (NewRotation.InitFromString(PastedText))
 		{
-			FScopedTransaction Transaction(LOCTEXT("PasteRotation", "Paste Rotation"));
-			SetSplinePointProperty(ESplinePointProperty::Rotation, NewRotation, EAxisList::All, true);
+			FVector NewRotation;
+			PastedText.ReplaceInline(TEXT("Pitch="), TEXT("X="));
+			PastedText.ReplaceInline(TEXT("Yaw="), TEXT("Y="));
+			PastedText.ReplaceInline(TEXT("Roll="), TEXT("Z="));
+			if (NewRotation.InitFromString(PastedText))
+			{
+				FScopedTransaction Transaction(LOCTEXT("PasteRotation", "Paste Rotation"));
+				SetSplinePointProperty(ESplinePointProperty::Rotation, NewRotation, EAxisList::All, true);
+			}
+			break;
 		}
-		break;
-	}
 	case ESplinePointProperty::Scale:
-	{
-		FVector NewScale;
-		if (NewScale.InitFromString(PastedText))
 		{
-			FScopedTransaction Transaction(LOCTEXT("PasteScale", "Paste Scale"));
-			SetSplinePointProperty(ESplinePointProperty::Scale, NewScale, EAxisList::All, true);
+			FVector NewScale;
+			if (NewScale.InitFromString(PastedText))
+			{
+				FScopedTransaction Transaction(LOCTEXT("PasteScale", "Paste Scale"));
+				SetSplinePointProperty(ESplinePointProperty::Scale, NewScale, EAxisList::All, true);
+			}
+			break;
 		}
-		break;
-	}
 	case ESplinePointProperty::ArriveTangent:
-	{
-		FVector NewArrive;
-		if (NewArrive.InitFromString(PastedText))
 		{
-			FScopedTransaction Transaction(LOCTEXT("PasteArriveTangent", "Paste Arrive Tangent"));
-			SetSplinePointProperty(ESplinePointProperty::ArriveTangent, NewArrive, EAxisList::All, true);
+			FVector NewArrive;
+			if (NewArrive.InitFromString(PastedText))
+			{
+				FScopedTransaction Transaction(LOCTEXT("PasteArriveTangent", "Paste Arrive Tangent"));
+				SetSplinePointProperty(ESplinePointProperty::ArriveTangent, NewArrive, EAxisList::All, true);
+			}
+			break;
 		}
-		break;
-	}
 	case ESplinePointProperty::LeaveTangent:
-	{
-		FVector NewLeave;
-		if (NewLeave.InitFromString(PastedText))
 		{
-			FScopedTransaction Transaction(LOCTEXT("PasteLeaveTangent", "Paste Leave Tangent"));
-			SetSplinePointProperty(ESplinePointProperty::LeaveTangent, NewLeave, EAxisList::All, true);
+			FVector NewLeave;
+			if (NewLeave.InitFromString(PastedText))
+			{
+				FScopedTransaction Transaction(LOCTEXT("PasteLeaveTangent", "Paste Leave Tangent"));
+				SetSplinePointProperty(ESplinePointProperty::LeaveTangent, NewLeave, EAxisList::All, true);
+			}
+			break;
 		}
-		break;
-	}
 	default:
 		break;
 	}
diff --git a/Engine/Source/Editor/DetailCustomizations/Public/Customizations/MathStructProxyCustomizations.h b/Engine/Source/Editor/DetailCustomizations/Public/Customizations/MathStructProxyCustomizations.h
--- a/Engine/Source/Editor/DetailCustomizations/Public/Customizations/MathStructProxyCustomizations.h
+++ b/Engine/Source/Editor/DetailCustomizations/Public/Customizations/MathStructProxyCustomizations.h
@@ -292,6 +292,19 @@ protected:
 	void OnCopy(FTransformField::Type Type, TWeakPtr<IPropertyHandle> PropertyHandlePtr);
 	void OnPaste(FTransformField::Type Type, TWeakPtr<IPropertyHandle> PropertyHandlePtr);
 
+	void PasteFromText(
+		const FString& InTag,
+		const FString& InText,
+		FTransformField::Type Type,
+		TWeakPtr<IPropertyHandle> PropertyHandlePtr);
+
+	void OnPasteFromText(
+		const FString& InTag,
+		const FString& InText,
+		const TOptional<FGuid>& InOperationId,
+		FTransformField::Type Type,
+		TWeakPtr<IPropertyHandle> PropertyHandlePtr);
+
 protected:
 	/** Cached rotation values */
 	mutable TSharedRef< TProxyValue<UE::Math::TRotator<T>> > CachedRotation;
diff --git a/Engine/Source/Editor/PropertyEditor/Private/CustomChildBuilder.cpp b/Engine/Source/Editor/PropertyEditor/Private/CustomChildBuilder.cpp
--- a/Engine/Source/Editor/PropertyEditor/Private/CustomChildBuilder.cpp
+++ b/Engine/Source/Editor/PropertyEditor/Private/CustomChildBuilder.cpp
@@ -29,10 +29,17 @@ IDetailGroup& FCustomChildrenBuilder::AddGroup( FName GroupName, const FText& Lo
 
 FDetailWidgetRow& FCustomChildrenBuilder::AddCustomRow( const FText& SearchString )
 {
-	TSharedRef<FDetailWidgetRow> NewRow = MakeShareable( new FDetailWidgetRow );
+	const TSharedRef<FDetailWidgetRow> NewRow = MakeShared<FDetailWidgetRow>();
 	FDetailLayoutCustomization NewCustomization;
 
 	NewRow->FilterString( SearchString );
+
+	// Bind to PasteFromText if specified
+	if (const TSharedPtr<FOnPasteFromText> PasteFromTextDelegate = GetParentCategory().OnPasteFromText())
+	{
+		NewRow->OnPasteFromTextDelegate = PasteFromTextDelegate;
+	}
+	
 	NewCustomization.WidgetDecl = NewRow;
 
 	ChildCustomizations.Add( NewCustomization );
diff --git a/Engine/Source/Editor/PropertyEditor/Private/DetailAdvancedDropdownNode.h b/Engine/Source/Editor/PropertyEditor/Private/DetailAdvancedDropdownNode.h
--- a/Engine/Source/Editor/PropertyEditor/Private/DetailAdvancedDropdownNode.h
+++ b/Engine/Source/Editor/PropertyEditor/Private/DetailAdvancedDropdownNode.h
@@ -30,7 +30,7 @@ private:
 	virtual IDetailsViewPrivate* GetDetailsView() const override { return ParentCategory.GetDetailsView(); }
 	virtual TSharedRef< ITableRow > GenerateWidgetForTableView( const TSharedRef<STableViewBase>& OwnerTable, bool bAllowFavoriteSystem) override;
 	virtual bool GenerateStandaloneWidget(FDetailWidgetRow& OutRow) const override;
-	virtual void GetChildren(FDetailNodeList& OutChildren) override {}
+	virtual void GetChildren(FDetailNodeList& OutChildren, const bool& bInIgnoreVisibility) override {}
 	virtual void OnItemExpansionChanged( bool bIsExpanded, bool bShouldSaveState) override {}
 	virtual bool ShouldBeExpanded() const override { return false; }
 	virtual ENodeVisibility GetVisibility() const override;
diff --git a/Engine/Source/Editor/PropertyEditor/Private/DetailCategoryBuilderImpl.cpp b/Engine/Source/Editor/PropertyEditor/Private/DetailCategoryBuilderImpl.cpp
--- a/Engine/Source/Editor/PropertyEditor/Private/DetailCategoryBuilderImpl.cpp
+++ b/Engine/Source/Editor/PropertyEditor/Private/DetailCategoryBuilderImpl.cpp
@@ -98,7 +98,6 @@ bool FDetailLayoutCustomization::IsHidden() const
 		|| (HasPropertyNode() && PropertyRow->GetPropertyVisibility() != EVisibility::Visible);
 }
 
-
 TSharedPtr<FPropertyNode> FDetailLayoutCustomization::GetPropertyNode() const
 {
 	return PropertyRow.IsValid() ? PropertyRow->GetPropertyNode() : nullptr;
@@ -166,6 +165,7 @@ TOptional<FResetToDefaultOverride> FDetailLayoutCustomization::GetCustomResetToD
 FDetailCategoryImpl::FDetailCategoryImpl(FName InCategoryName, TSharedRef<FDetailLayoutBuilderImpl> InDetailLayout)
 	: HeaderContentWidget(nullptr)
 	, DetailLayoutBuilder(InDetailLayout)
+	, PasteFromTextDelegate(MakeShared<FOnPasteFromText>())
 	, CategoryName(InCategoryName)
 	, SortOrder(0)
 	, bRestoreExpansionState(!ContainsOnlyAdvanced())
@@ -783,12 +783,12 @@ TSharedRef<ITableRow> FDetailCategoryImpl::GenerateWidgetForTableView(const TSha
 	TSharedPtr<FDetailLayoutBuilderImpl> ParentLayout = GetParentLayoutImpl();
 
 	return SNew(SDetailCategoryTableRow, AsShared(), OwnerTable)
+		.PasteFromText(OnPasteFromText())
 		.InnerCategory(ParentLayout.IsValid() ? ParentLayout->IsLayoutForExternalRoot() : false)
 		.DisplayName(GetDisplayName())
 		.HeaderContent(HeaderContent);
 }
 
-
 bool FDetailCategoryImpl::GenerateStandaloneWidget(FDetailWidgetRow& OutRow) const
 {
 	TSharedPtr<SWidget> HeaderContent = HeaderContentWidget;
@@ -1040,9 +1040,9 @@ void FDetailCategoryImpl::GenerateChildrenForLayouts()
 	}
 }
 
-void FDetailCategoryImpl::GetChildren(FDetailNodeList& OutChildren)
+void FDetailCategoryImpl::GetChildren(FDetailNodeList& OutChildren, const bool& bInIgnoreVisibility)
 {
-	GetGeneratedChildren(OutChildren, false, false);
+	GetGeneratedChildren(OutChildren, bInIgnoreVisibility, false);
 }
 
 void FDetailCategoryImpl::GetGeneratedChildren(FDetailNodeList& OutChildren, bool bIgnoreVisibility, bool bIgnoreAdvanced)
@@ -1053,7 +1053,7 @@ void FDetailCategoryImpl::GetGeneratedChildren(FDetailNodeList& OutChildren, boo
 		{
 			if (Child->ShouldShowOnlyChildren())
 			{
-				Child->GetChildren(OutChildren);
+				Child->GetChildren(OutChildren, bIgnoreVisibility);
 			}
 			else
 			{
diff --git a/Engine/Source/Editor/PropertyEditor/Private/DetailCategoryBuilderImpl.h b/Engine/Source/Editor/PropertyEditor/Private/DetailCategoryBuilderImpl.h
--- a/Engine/Source/Editor/PropertyEditor/Private/DetailCategoryBuilderImpl.h
+++ b/Engine/Source/Editor/PropertyEditor/Private/DetailCategoryBuilderImpl.h
@@ -186,6 +186,9 @@ public:
 	virtual int32 GetSortOrder() const override;
 	virtual void SetSortOrder(int32 InSortOrder) override;
 	virtual void AddPropertyDisableInstancedReference(TSharedPtr<IPropertyHandle> PropertyHandle) override;
+	
+	/** Delegate handling pasting an optionally tagged text snippet */
+	virtual TSharedPtr<FOnPasteFromText> OnPasteFromText() const override { return PasteFromTextDelegate; }
 
 	/** FDetailTreeNode interface */
 	virtual IDetailsView* GetNodeDetailsView() const override { return GetDetailsView(); }
@@ -199,7 +202,7 @@ public:
 	virtual void GetFilterStrings(TArray<FString>& OutFilterStrings) const override;
 	virtual bool GetInitiallyCollapsed() const override;
 
-	virtual void GetChildren(FDetailNodeList& OutChildren) override;
+	virtual void GetChildren(FDetailNodeList& OutChildren, const bool& bInIgnoreVisibility = false) override;
 	virtual bool ShouldBeExpanded() const override;
 	virtual ENodeVisibility GetVisibility() const override;
 	virtual void FilterNode(const FDetailFilter& DetailFilter) override;
@@ -408,6 +411,9 @@ private:
 	TSharedPtr<FDetailTreeNode> InlinePropertyNode;
 	/** The parent detail builder */
 	TWeakPtr<FDetailLayoutBuilderImpl> DetailLayoutBuilder;
+	/** Delegate handling pasting an optionally tagged text snippet */
+	TSharedPtr<FOnPasteFromText> PasteFromTextDelegate;
+
 	/** The category identifier */
 	FName CategoryName;
 	/** The sort order of this category (amongst all categories) */
diff --git a/Engine/Source/Editor/PropertyEditor/Private/DetailCategoryGroupNode.cpp b/Engine/Source/Editor/PropertyEditor/Private/DetailCategoryGroupNode.cpp
--- a/Engine/Source/Editor/PropertyEditor/Private/DetailCategoryGroupNode.cpp
+++ b/Engine/Source/Editor/PropertyEditor/Private/DetailCategoryGroupNode.cpp
@@ -43,7 +43,7 @@ bool FDetailCategoryGroupNode::GenerateStandaloneWidget(FDetailWidgetRow& OutRow
 	return true;
 }
 
-void FDetailCategoryGroupNode::GetChildren( FDetailNodeList& OutChildren )
+void FDetailCategoryGroupNode::GetChildren(FDetailNodeList& OutChildren, const bool& bInIgnoreVisibility)
 {
 	for( int32 ChildIndex = 0; ChildIndex < ChildNodes.Num(); ++ChildIndex )
 	{
@@ -52,7 +52,7 @@ void FDetailCategoryGroupNode::GetChildren( FDetailNodeList& OutChildren )
 		{
 			if( Child->ShouldShowOnlyChildren() )
 			{
-				Child->GetChildren( OutChildren );
+				Child->GetChildren( OutChildren, bInIgnoreVisibility );
 			}
 			else
 			{
diff --git a/Engine/Source/Editor/PropertyEditor/Private/DetailCategoryGroupNode.h b/Engine/Source/Editor/PropertyEditor/Private/DetailCategoryGroupNode.h
--- a/Engine/Source/Editor/PropertyEditor/Private/DetailCategoryGroupNode.h
+++ b/Engine/Source/Editor/PropertyEditor/Private/DetailCategoryGroupNode.h
@@ -37,7 +37,7 @@ private:
 	virtual TSharedPtr<IPropertyHandle> CreatePropertyHandle() const override { return nullptr; }
 	virtual void GetFilterStrings(TArray<FString>& OutFilterStrings) const override { OutFilterStrings.Add(GroupName.ToString()); };
 
-	virtual void GetChildren(FDetailNodeList& OutChildren ) override;
+	virtual void GetChildren(FDetailNodeList& OutChildren, const bool& bInIgnoreVisibility) override;
 	virtual void FilterNode( const FDetailFilter& InFilter ) override;
 	virtual void Tick( float DeltaTime ) override {}
 	virtual bool ShouldShowOnlyChildren() const override { return false; }
diff --git a/Engine/Source/Editor/PropertyEditor/Private/DetailGroup.cpp b/Engine/Source/Editor/PropertyEditor/Private/DetailGroup.cpp
--- a/Engine/Source/Editor/PropertyEditor/Private/DetailGroup.cpp
+++ b/Engine/Source/Editor/PropertyEditor/Private/DetailGroup.cpp
@@ -15,8 +15,8 @@ FDetailGroup::FDetailGroup( const FName InGroupName, TSharedRef<FDetailCategoryI
 	, bStartExpanded( bInStartExpanded )
 	, bResetEnabled( false )
 	, DisplayMode( EDetailGroupDisplayMode::Group )
+	, PasteFromTextDelegate(MakeShared<FOnPasteFromText>())
 {
-
 }
 
 FDetailWidgetRow& FDetailGroup::HeaderRow()
diff --git a/Engine/Source/Editor/PropertyEditor/Private/DetailGroup.h b/Engine/Source/Editor/PropertyEditor/Private/DetailGroup.h
--- a/Engine/Source/Editor/PropertyEditor/Private/DetailGroup.h
+++ b/Engine/Source/Editor/PropertyEditor/Private/DetailGroup.h
@@ -20,7 +20,7 @@ class FDetailGroup : public IDetailGroup, public TSharedFromThis<FDetailGroup>
 public:
 	FDetailGroup( const FName InGroupName, TSharedRef<FDetailCategoryImpl> InParentCategory, const FText& InLocalizedDisplayName, const bool bInStartExpanded = false );
 
-	/** IDetailGroup interface */     
+	/** IDetailGroup interface */
 	virtual FDetailWidgetRow& HeaderRow() override;
 	virtual IDetailPropertyRow& HeaderProperty( TSharedRef<IPropertyHandle> PropertyHandle ) override;
 	virtual FDetailWidgetRow& AddWidgetRow() override;
@@ -96,6 +96,11 @@ public:
 	*/
 	virtual FDetailGroupReset& GetOnDetailGroupReset() override { return OnDetailGroupReset; }
 
+	/**
+	 * Return an optional PasteFromText delegate
+	 */
+	virtual TSharedPtr<FOnPasteFromText> OnPasteFromText() const override { return PasteFromTextDelegate; }
+
 private:
 	/**
 	 * Called when the name of the group is clicked to expand the group
@@ -136,4 +141,6 @@ private:
 	EDetailGroupDisplayMode DisplayMode;
 	/**	Delegate called when user press the Group Reset ui */
 	FDetailGroupReset OnDetailGroupReset;
+	/** Delegate handling pasting an optionally tagged text snippet */
+	TSharedPtr<FOnPasteFromText> PasteFromTextDelegate;
 };
diff --git a/Engine/Source/Editor/PropertyEditor/Private/DetailItemNode.cpp b/Engine/Source/Editor/PropertyEditor/Private/DetailItemNode.cpp
--- a/Engine/Source/Editor/PropertyEditor/Private/DetailItemNode.cpp
+++ b/Engine/Source/Editor/PropertyEditor/Private/DetailItemNode.cpp
@@ -365,7 +365,7 @@ bool FDetailItemNode::GenerateStandaloneWidget(FDetailWidgetRow& OutRow) const
 	return bResult;
 }
 
-void FDetailItemNode::GetChildren(FDetailNodeList& OutChildren)
+void FDetailItemNode::GetChildren(FDetailNodeList& OutChildren, const bool& bInIgnoreVisibility)
 {
 	OutChildren.Reserve(Children.Num());
 
@@ -377,12 +377,13 @@ void FDetailItemNode::GetChildren(FDetailNodeList& OutChildren)
 		// If we are visible due to filtering and so is a child, we only show that child.  
 		// If we are visible due to filtering and no child is visible, we show all children
 
-		if( ChildVisibility == ENodeVisibility::Visible ||
-			( !bShouldBeVisibleDueToChildFiltering && bShouldBeVisibleDueToFiltering && ChildVisibility != ENodeVisibility::ForcedHidden ) )
+		if( ChildVisibility == ENodeVisibility::Visible
+			|| bInIgnoreVisibility
+			|| ( !bShouldBeVisibleDueToChildFiltering && bShouldBeVisibleDueToFiltering && ChildVisibility != ENodeVisibility::ForcedHidden ) )
 		{
 			if( Child->ShouldShowOnlyChildren() )
 			{
-				Child->GetChildren( OutChildren );
+				Child->GetChildren( OutChildren, bInIgnoreVisibility );
 			}
 			else
 			{
diff --git a/Engine/Source/Editor/PropertyEditor/Private/DetailItemNode.h b/Engine/Source/Editor/PropertyEditor/Private/DetailItemNode.h
--- a/Engine/Source/Editor/PropertyEditor/Private/DetailItemNode.h
+++ b/Engine/Source/Editor/PropertyEditor/Private/DetailItemNode.h
@@ -69,7 +69,7 @@ public:
 	virtual IDetailsViewPrivate* GetDetailsView() const override{ TSharedPtr<FDetailCategoryImpl> PC = GetParentCategory(); return PC.IsValid() ? PC->GetDetailsView() : nullptr; }
 	virtual TSharedRef< ITableRow > GenerateWidgetForTableView( const TSharedRef<STableViewBase>& OwnerTable, bool bAllowFavoriteSystem) override;
 	virtual bool GenerateStandaloneWidget(FDetailWidgetRow& OutRow) const override;
-	virtual void GetChildren( FDetailNodeList& OutChildren )  override;
+	virtual void GetChildren(FDetailNodeList& OutChildren, const bool& bInIgnoreVisibility = false) override;
 	virtual void OnItemExpansionChanged( bool bInIsExpanded, bool bShouldSaveState) override;
 	virtual bool ShouldBeExpanded() const override;
 	virtual ENodeVisibility GetVisibility() const override;
diff --git a/Engine/Source/Editor/PropertyEditor/Private/DetailMultiTopLevelObjectRootNode.cpp b/Engine/Source/Editor/PropertyEditor/Private/DetailMultiTopLevelObjectRootNode.cpp
--- a/Engine/Source/Editor/PropertyEditor/Private/DetailMultiTopLevelObjectRootNode.cpp
+++ b/Engine/Source/Editor/PropertyEditor/Private/DetailMultiTopLevelObjectRootNode.cpp
@@ -165,8 +165,7 @@ bool FDetailMultiTopLevelObjectRootNode::GenerateStandaloneWidget(FDetailWidgetR
 	return true;
 }
 
-
-void FDetailMultiTopLevelObjectRootNode::GetChildren(FDetailNodeList& OutChildren )
+void FDetailMultiTopLevelObjectRootNode::GetChildren(FDetailNodeList& OutChildren, const bool& bInIgnoreVisibility)
 {
 	for( int32 ChildIndex = 0; ChildIndex < ChildNodes.Num(); ++ChildIndex )
 	{
@@ -175,7 +174,7 @@ void FDetailMultiTopLevelObjectRootNode::GetChildren(FDetailNodeList& OutChildre
 		{
 			if( Child->ShouldShowOnlyChildren() )
 			{
-				Child->GetChildren( OutChildren );
+				Child->GetChildren( OutChildren, bInIgnoreVisibility );
 			}
 			else
 			{
diff --git a/Engine/Source/Editor/PropertyEditor/Private/DetailMultiTopLevelObjectRootNode.h b/Engine/Source/Editor/PropertyEditor/Private/DetailMultiTopLevelObjectRootNode.h
--- a/Engine/Source/Editor/PropertyEditor/Private/DetailMultiTopLevelObjectRootNode.h
+++ b/Engine/Source/Editor/PropertyEditor/Private/DetailMultiTopLevelObjectRootNode.h
@@ -55,7 +55,7 @@ private:
 	virtual ENodeVisibility GetVisibility() const override;
 	virtual TSharedRef<ITableRow> GenerateWidgetForTableView(const TSharedRef<STableViewBase>& OwnerTable, bool bAllowFavoriteSystem) override;
 	virtual bool GenerateStandaloneWidget(FDetailWidgetRow& OutRow) const override;
-	virtual void GetChildren(FDetailNodeList& OutChildren)  override;
+	virtual void GetChildren(FDetailNodeList& OutChildren, const bool& bInIgnoreVisibility) override;
 	virtual void FilterNode(const FDetailFilter& InFilter) override;
 	virtual void Tick(float DeltaTime) override {}
 	virtual bool ShouldShowOnlyChildren() const override;
diff --git a/Engine/Source/Editor/PropertyEditor/Private/DetailPropertyRow.cpp b/Engine/Source/Editor/PropertyEditor/Private/DetailPropertyRow.cpp
--- a/Engine/Source/Editor/PropertyEditor/Private/DetailPropertyRow.cpp
+++ b/Engine/Source/Editor/PropertyEditor/Private/DetailPropertyRow.cpp
@@ -738,6 +738,7 @@ void FDetailPropertyRow::SetWidgetRowProperties(FDetailWidgetRow& Row) const
 		Row.CopyMenuAction = CustomPropertyWidget->CopyMenuAction;
 		Row.PasteMenuAction = CustomPropertyWidget->PasteMenuAction;
 		Row.CustomMenuItems = CustomPropertyWidget->CustomMenuItems;
+        Row.OnPasteFromTextDelegate = CustomPropertyWidget->OnPasteFromTextDelegate;
 
 		if (CustomPropertyWidget->CustomResetToDefault.IsSet())
 		{
diff --git a/Engine/Source/Editor/PropertyEditor/Private/DetailTreeNode.cpp b/Engine/Source/Editor/PropertyEditor/Private/DetailTreeNode.cpp
--- a/Engine/Source/Editor/PropertyEditor/Private/DetailTreeNode.cpp
+++ b/Engine/Source/Editor/PropertyEditor/Private/DetailTreeNode.cpp
@@ -46,10 +46,10 @@ FNodeWidgets FDetailTreeNode::CreateNodeWidgets() const
 	return Widgets;
 }
 
-void FDetailTreeNode::GetChildren(TArray<TSharedRef<IDetailTreeNode>>& OutChildren)
+void FDetailTreeNode::GetChildren(TArray<TSharedRef<IDetailTreeNode>>& OutChildren, const bool& bInIgnoreVisibility)
 {
 	FDetailNodeList Children;
-	GetChildren(Children);
+	GetChildren(Children, bInIgnoreVisibility);
 
 	OutChildren.Reset(Children.Num());
 
diff --git a/Engine/Source/Editor/PropertyEditor/Private/MaterialList.cpp b/Engine/Source/Editor/PropertyEditor/Private/MaterialList.cpp
--- a/Engine/Source/Editor/PropertyEditor/Private/MaterialList.cpp
+++ b/Engine/Source/Editor/PropertyEditor/Private/MaterialList.cpp
@@ -560,6 +560,12 @@ void FMaterialList::GenerateHeaderRowContent( FDetailWidgetRow& NodeRow )
 	NodeRow.CopyAction(FUIAction(FExecuteAction::CreateSP(this, &FMaterialList::OnCopyMaterialList), FCanExecuteAction::CreateSP(this, &FMaterialList::OnCanCopyMaterialList)));
 	NodeRow.PasteAction(FUIAction(FExecuteAction::CreateSP(this, &FMaterialList::OnPasteMaterialList)));
 
+	if (const TSharedPtr<FOnPasteFromText> OnPasteFromTextDelegate = NodeRow.OnPasteFromTextDelegate.Pin();
+		OnPasteFromTextDelegate.IsValid())
+	{
+		OnPasteFromTextDelegate->AddSP(this, &FMaterialList::OnPasteMaterialListFromText);
+	}
+
 	if (bAllowCollpase)
 	{
 		NodeRow.NameContent()
@@ -692,6 +698,21 @@ void FMaterialList::OnPasteMaterialList()
 	}
 }
 
+void FMaterialList::OnPasteMaterialListFromText(
+	const FString& InTag,
+	const FString& InText,
+	const TOptional<FGuid>& InOperationId)
+{
+	if (const TSharedPtr<FOnPasteFromText> OnPasteFromTextDelegate = MaterialListDelegates.OnPasteFromText.Pin();
+		OnPasteFromTextDelegate.IsValid())
+	{
+		if (OnPasteFromTextDelegate->IsBound())
+		{
+			OnPasteFromTextDelegate->Broadcast(InTag, InText, InOperationId);
+		}
+	}
+}
+
 bool FMaterialList::OnCanCopyMaterialItem(int32 CurrentSlot) const
 {
 	if (MaterialListDelegates.OnCanCopyMaterialItem.IsBound())
diff --git a/Engine/Source/Editor/PropertyEditor/Private/PropertyEditorClipboard.cpp b/Engine/Source/Editor/PropertyEditor/Private/PropertyEditorClipboard.cpp
new file mode 100644
--- /dev/null
+++ b/Engine/Source/Editor/PropertyEditor/Private/PropertyEditorClipboard.cpp
@@ -0,0 +1,175 @@
+﻿// Copyright Epic Games, Inc. All Rights Reserved.
+
+#include "PropertyEditorClipboard.h"
+
+#include "HAL/PlatformApplicationMisc.h"
+#include "PropertyNode.h"
+#include "Serialization/JsonReader.h"
+#include "Serialization/JsonSerializer.h"
+
+
+namespace UE::PropertyEditor::Internal
+{
+	constexpr const TCHAR* ClipboardMapKeyName = TEXT("Tagged");
+
+	bool IsJsonString(const FString& InStr)
+	{
+		if (InStr.IsEmpty())
+		{
+			return false;
+		}
+
+		// Str has contents but isn't necessarily json
+		TSharedPtr<FJsonObject> UnusedJsonObject;
+		const TSharedRef<TJsonReader<>> JsonReader = TJsonReaderFactory<>::Create(*InStr);
+		return FJsonSerializer::Deserialize(JsonReader, UnusedJsonObject); // This will return false if not parsed
+	}
+
+	bool TryWriteClipboard(const TMap<FName, FString>& InTaggedClipboardPairs, FString& OutClipboardStr)
+	{
+		const TSharedPtr<FJsonObject> RootJsonObject = MakeShared<FJsonObject>();
+		TArray<TSharedPtr<FJsonValue>> TaggedValuesJson;
+		TaggedValuesJson.Reserve(InTaggedClipboardPairs.Num());
+		for (const TPair<FName, FString>& TaggedPair : InTaggedClipboardPairs)
+		{
+			TArray<TSharedPtr<FJsonValue>> TaggedPairJson;
+			TaggedPairJson.Reserve(2);
+
+			TaggedPairJson.Add(MakeShared<FJsonValueString>(TaggedPair.Key.ToString()));
+			TaggedPairJson.Add(MakeShared<FJsonValueString>(TaggedPair.Value));
+
+			TaggedValuesJson.Add(MakeShared<FJsonValueArray>(MoveTemp(TaggedPairJson)));
+		}
+		RootJsonObject->SetArrayField(ClipboardMapKeyName, TaggedValuesJson);
+
+		const TSharedRef<TJsonWriter<>> JsonWriter = TJsonWriterFactory<>::Create(&OutClipboardStr);
+		return FJsonSerializer::Serialize(RootJsonObject.ToSharedRef(), JsonWriter);
+	}
+
+	bool TryParseClipboard(const FString& InClipboardStr, TMap<FName, FString>& OutTaggedClipboard)
+	{
+		TSharedPtr<FJsonObject> RootJsonObject;
+		const TSharedRef<TJsonReader<>> JsonReader = TJsonReaderFactory<>::Create(InClipboardStr);
+		if (FJsonSerializer::Deserialize(JsonReader, RootJsonObject))
+		{
+			const TArray<TSharedPtr<FJsonValue>>* TaggedValueArray = nullptr;
+			if (RootJsonObject->TryGetArrayField(ClipboardMapKeyName, TaggedValueArray))
+			{
+				OutTaggedClipboard.Reserve(TaggedValueArray->Num());				
+				for (const TSharedPtr<FJsonValue>& TaggedValueEntry : *TaggedValueArray)
+				{
+					const TArray<TSharedPtr<FJsonValue>>* TaggedValuePair = nullptr;
+					if (TaggedValueEntry->TryGetArray(TaggedValuePair))
+					{
+						// Expects at least key, value entries
+						if (TaggedValuePair->Num() < 2)
+						{
+							continue;
+						}
+						
+						FString TagName = (*TaggedValuePair)[0]->AsString();
+						FString TagValue = (*TaggedValuePair)[1]->AsString();
+
+						OutTaggedClipboard.Add(FName(TagName), TagValue);
+					}
+				}
+				return true;
+			}
+		}
+
+		return false;
+	}
+}
+
+TSharedRef<FPropertyEditorClipboard> FPropertyEditorClipboard::Get()
+{
+	static TSharedRef<FPropertyEditorClipboard> Instance = MakeShared<FPropertyEditorClipboard>();
+	return Instance;
+}
+
+void FPropertyEditorClipboard::ClipboardCopy(const TCHAR* Str)
+{
+	FPropertyEditorClipboard::Get()->ClipboardContents.Reset();
+	FPlatformApplicationMisc::ClipboardCopy(Str);
+}
+
+static bool EncodeClipboardCopy(const TMap<FName, FString>& InTaggedClipboard)
+{
+	FString ClipboardStr;
+	if (UE::PropertyEditor::Internal::TryWriteClipboard(InTaggedClipboard, ClipboardStr))
+	{
+		// Tagged content written, now encode and write to clipboard
+		FPlatformApplicationMisc::ClipboardCopy(*ClipboardStr);
+
+		// Encoding (to json) successful
+		return true;
+	}
+	
+	// Encoding failed - silently fail and do regular copy
+	UE_LOG(LogPropertyNode, Warning, TEXT("Failed to encode tagged clipboard as json."))
+
+	// Only if default exists (it always should, caller!)
+	if (const FString* DefaultClipboardStr = InTaggedClipboard.Find(NAME_None))
+	{
+		FPlatformApplicationMisc::ClipboardCopy(**DefaultClipboardStr);
+	}
+	return false;
+}
+
+void FPropertyEditorClipboard::ClipboardCopy(
+	const TCHAR* Str, 
+	FName Tag)
+{
+	if (Tag == NAME_None)
+	{
+		return ClipboardCopy(Str);
+	}
+
+	TMap<FName, FString>& TaggedContents = FPropertyEditorClipboard::Get()->ClipboardContents;
+	TaggedContents.Reset();
+	// Populate both default and tagged, makes it easy to check OnPaste if the contents is intended to be "specialized"
+	TaggedContents.Add(NAME_None, Str);
+	TaggedContents.Add(Tag, Str);
+
+	EncodeClipboardCopy(TaggedContents);
+}
+
+void FPropertyEditorClipboard::ClipboardCopy(
+	TUniqueFunction<void(TMap<FName, FString>&)>&& TagMappingFunc)
+{
+	TMap<FName, FString>& TaggedContents = FPropertyEditorClipboard::Get()->ClipboardContents;
+	TaggedContents.Reset();
+
+	ClipboardAppend(MoveTemp(TagMappingFunc));
+}
+
+void FPropertyEditorClipboard::ClipboardAppend(TUniqueFunction<void(TMap<FName, FString>&)>&& TagMappingFunc)
+{
+	TMap<FName, FString>& TaggedContents = FPropertyEditorClipboard::Get()->ClipboardContents;
+	
+	TagMappingFunc(TaggedContents);
+
+    EncodeClipboardCopy(TaggedContents);
+}
+
+bool FPropertyEditorClipboard::ClipboardPaste(FString& Dest)
+{
+	FPlatformApplicationMisc::ClipboardPaste(Dest);
+	return true;
+}
+
+bool FPropertyEditorClipboard::ClipboardPaste(FString& Dest, FName Tag)
+{
+	if (Tag != NAME_None)
+	{
+		// Check for tagged entry, return true if found
+		if (FString* ClipboardText = FPropertyEditorClipboard::Get()->ClipboardContents.Find(Tag))
+		{
+			Dest = *ClipboardText;
+			return true;
+		}
+	}
+
+	// Otherwise perform default operation
+	return ClipboardPaste(Dest);
+}
diff --git a/Engine/Source/Editor/PropertyEditor/Private/PropertyEditorClipboardPrivate.h b/Engine/Source/Editor/PropertyEditor/Private/PropertyEditorClipboardPrivate.h
new file mode 100644
--- /dev/null
+++ b/Engine/Source/Editor/PropertyEditor/Private/PropertyEditorClipboardPrivate.h
@@ -0,0 +1,49 @@
+﻿// Copyright Epic Games, Inc. All Rights Reserved.
+
+#pragma once
+
+#include "Containers/Map.h"
+#include "Containers/UnrealString.h"
+#include "CoreTypes.h"
+#include "Templates/Function.h"
+#include "Templates/SharedPointerFwd.h"
+#include "UObject/NameTypes.h"
+
+/** Utilities to encode and decode the clipboard. */
+namespace UE::PropertyEditor::Internal
+{
+	bool IsJsonString(const FString& InStr);
+	bool TryWriteClipboard(const TMap<FName, FString>& InTaggedClipboardPairs, FString& OutClipboardStr);
+	bool TryParseClipboard(const FString& InClipboardStr, TMap<FName, FString>& OutTaggedClipboard);
+
+	/** Contains parsed clipboard data, used to cache subsequent validity checks. */
+	struct FClipboardData
+	{
+		/** Clipboard content string. */
+		TOptional<FString> Content;
+
+		/** Applicability of the clipboard content to this row. */ 
+		bool bIsApplicable = false;
+
+		/** Deserialized property values (as text). */
+		TMap<FName, FString> PropertyValues;
+
+		/** List of unique property names being pasted. */
+		TArray<FName> PropertyNames;
+
+		/** Resets/invalidates data. */
+		void Reset()
+		{
+			Content.Reset();
+			bIsApplicable = false;
+			PropertyValues.Reset();
+			PropertyNames.Reset();
+		}
+
+		void Reserve(const int32& InNum)
+		{
+			PropertyValues.Reserve(InNum);
+			PropertyNames.Reserve(InNum);
+		}
+	};
+}
diff --git a/Engine/Source/Editor/PropertyEditor/Private/PropertyEditorCopyPaste.cpp b/Engine/Source/Editor/PropertyEditor/Private/PropertyEditorCopyPaste.cpp
new file mode 100644
--- /dev/null
+++ b/Engine/Source/Editor/PropertyEditor/Private/PropertyEditorCopyPaste.cpp
@@ -0,0 +1,61 @@
+﻿// Copyright Epic Games, Inc. All Rights Reserved.
+
+#include "PropertyEditorCopyPaste.h"
+#include "PropertyEditorCopyPastePrivate.h"
+#include "PropertyHandleImpl.h"
+#include "PropertyNode.h"
+
+namespace UE::PropertyEditor
+{
+	FString GetPropertyPath(const TSharedPtr<IPropertyHandle>& InPropertyHandle)
+	{
+		return Private::GetPropertyPath(
+			[InPropertyHandle] { return InPropertyHandle; },
+			{});
+	}
+
+	bool TagMatchesProperty(const FString& InTag, const TSharedPtr<IPropertyHandle>& InPropertyHandle)
+	{
+		if (InTag.IsEmpty())
+		{
+			return true;
+		}
+
+		const FString PropertyPath = UE::PropertyEditor::GetPropertyPath(InPropertyHandle);
+		
+		// If tag is specified, ensure that it matches the property (by path)
+		return InTag.Equals(PropertyPath);
+	}
+
+	namespace Private
+	{
+		FString GetPropertyPath(
+			TUniqueFunction<const TSharedPtr<IPropertyHandle>()>&& GetPropertyHandle,
+			TUniqueFunction<const TSharedPtr<FPropertyNode>()>&& GetPropertyNode)
+		{
+			const TSharedPtr<IPropertyHandle> PropertyHandle = GetPropertyHandle();
+			if (PropertyHandle.IsValid() && PropertyHandle->IsValidHandle())
+			{
+				FString PropertyPath = PropertyHandle->GeneratePathToProperty();
+				if (!PropertyPath.IsEmpty())
+				{
+					return PropertyPath;
+				}
+			}
+
+			TSharedPtr<FPropertyNode> PropertyNode = nullptr;
+			if (GetPropertyNode)
+			{
+				PropertyNode = GetPropertyNode();
+			}
+
+			if (!PropertyNode && PropertyHandle.IsValid() && PropertyHandle->IsValidHandle())
+			{
+				PropertyNode = StaticCastSharedPtr<FPropertyHandleBase>(PropertyHandle)->GetPropertyNode();			
+			}
+
+			const static FString EmptyString = TEXT("");
+			return PropertyNode.IsValid() ? PropertyNode->GetPropertyPath() : EmptyString;
+		}
+	}
+}
diff --git a/Engine/Source/Editor/PropertyEditor/Private/PropertyEditorCopyPastePrivate.h b/Engine/Source/Editor/PropertyEditor/Private/PropertyEditorCopyPastePrivate.h
new file mode 100644
--- /dev/null
+++ b/Engine/Source/Editor/PropertyEditor/Private/PropertyEditorCopyPastePrivate.h
@@ -0,0 +1,16 @@
+﻿// Copyright Epic Games, Inc. All Rights Reserved.
+
+#pragma once
+
+#include "CoreTypes.h"
+#include "PropertyHandle.h"
+#include "PropertyNode.h"
+
+/** Utilities related to copy/paste. */
+namespace UE::PropertyEditor::Private
+{
+	/** Get property path from handle, or node if the handle isn't valid. */
+	[[nodiscard]] PROPERTYEDITOR_API FString GetPropertyPath(
+		TUniqueFunction<const TSharedPtr<IPropertyHandle>()>&& GetPropertyHandle,
+		TUniqueFunction<const TSharedPtr<FPropertyNode>()>&& GetPropertyNode);
+}
diff --git a/Engine/Source/Editor/PropertyEditor/Private/SDetailCategoryTableRow.cpp b/Engine/Source/Editor/PropertyEditor/Private/SDetailCategoryTableRow.cpp
--- a/Engine/Source/Editor/PropertyEditor/Private/SDetailCategoryTableRow.cpp
+++ b/Engine/Source/Editor/PropertyEditor/Private/SDetailCategoryTableRow.cpp
@@ -2,20 +2,43 @@
 
 #include "SDetailCategoryTableRow.h"
 
-#include "UserInterface/PropertyEditor/PropertyEditorConstants.h"
+#include "Algo/AnyOf.h"
+#include "Algo/Compare.h"
+#include "Async/ParallelFor.h"
+#include "DetailCategoryBuilderImpl.h"
+#include "PropertyEditorClipboard.h"
+#include "PropertyEditorClipboardPrivate.h"
+#include "ScopedTransaction.h"
 #include "SDetailExpanderArrow.h"
 #include "SDetailRowIndent.h"
+#include "Serialization/JsonSerializer.h"
 #include "Styling/StyleColors.h"
+#include "UserInterface/PropertyEditor/PropertyEditorConstants.h"
 
 void SDetailCategoryTableRow::Construct(const FArguments& InArgs, TSharedRef<FDetailTreeNode> InOwnerTreeNode, const TSharedRef<STableViewBase>& InOwnerTableView)
 {
 	OwnerTreeNode = InOwnerTreeNode;
 
+	DisplayName = InArgs._DisplayName;
 	bIsInnerCategory = InArgs._InnerCategory;
 	bShowBorder = InArgs._ShowBorder;
 
-	FDetailColumnSizeData& ColumnSizeData = InOwnerTreeNode->GetDetailsView()->GetColumnSizeData();
+	IDetailsViewPrivate* DetailsView = InOwnerTreeNode->GetDetailsView();
+	FDetailColumnSizeData& ColumnSizeData = DetailsView->GetColumnSizeData();
+	
+	PulseAnimation.AddCurve(0.0f, UE::PropertyEditor::Private::PulseAnimationLength, ECurveEaseFunction::CubicInOut);
 
+	CopyAction.ExecuteAction = FExecuteAction::CreateSP(this, &SDetailCategoryTableRow::OnCopyCategory);
+	CopyAction.CanExecuteAction = FCanExecuteAction::CreateSP(this, &SDetailCategoryTableRow::CanCopyCategory);
+
+	if (InArgs._PasteFromText.IsValid())
+	{
+		OnPasteFromTextDelegate = InArgs._PasteFromText;
+	}
+	
+	PasteAction.ExecuteAction = FExecuteAction::CreateSP(this, &SDetailCategoryTableRow::OnPasteCategory);
+	PasteAction.CanExecuteAction = FCanExecuteAction::CreateSP(this, &SDetailCategoryTableRow::CanPasteCategory);
+	
 	TSharedRef<SHorizontalBox> HeaderBox = SNew(SHorizontalBox)
 		+ SHorizontalBox::Slot()
 		.HAlign(HAlign_Left)
@@ -109,6 +132,79 @@ void SDetailCategoryTableRow::Construct(const FArguments& InArgs, TSharedRef<FDe
 	);
 }
 
+FReply SDetailCategoryTableRow::OnMouseButtonUp(const FGeometry& MyGeometry, const FPointerEvent& MouseEvent)
+{
+	if (MouseEvent.GetModifierKeys().IsShiftDown())
+	{
+		bool bIsHandled = false;
+		if (CopyAction.CanExecute() && MouseEvent.GetEffectingButton() == EKeys::RightMouseButton)
+		{
+			CopyAction.Execute();
+			PulseAnimation.Play(SharedThis(this));
+			bIsHandled = true;
+		}
+		else if (PasteAction.CanExecute() && MouseEvent.GetEffectingButton() == EKeys::LeftMouseButton)
+		{
+			PasteAction.Execute();
+			PulseAnimation.Play(SharedThis(this));
+			bIsHandled = true;
+		}
+
+		if (bIsHandled)
+		{
+			return FReply::Handled();
+		}
+	}
+
+	return SDetailTableRowBase::OnMouseButtonUp(MyGeometry, MouseEvent);
+}
+
+bool SDetailCategoryTableRow::OnContextMenuOpening(FMenuBuilder& MenuBuilder)
+{
+	// Don't add anything if neither actions are bound
+	if (!CopyAction.IsBound() || !PasteAction.IsBound())
+	{
+		return true;
+	}
+	
+	MenuBuilder.BeginSection(NAME_None, NSLOCTEXT("PropertyView", "EditHeading", "Edit"));
+	{
+		bool bLongDisplayName = false;
+
+		FMenuEntryParams CopyContentParams;
+		CopyContentParams.LabelOverride = NSLOCTEXT("PropertyView", "CopyCategoryProperties", "Copy All Properties in Category");
+		CopyContentParams.ToolTipOverride = TAttribute<FText>::CreateLambda([this]()
+		{
+			return CanCopyCategory()
+				? NSLOCTEXT("PropertyView", "CopyCategoryProperties_ToolTip", "Copy all properties in this category")
+				: NSLOCTEXT("PropertyView", "CantCopyCategoryProperties_ToolTip", "None of the properties in this category can be copied");
+		});
+		CopyContentParams.InputBindingOverride = FInputChord(EModifierKey::Shift, EKeys::RightMouseButton).GetInputText(bLongDisplayName);
+		CopyContentParams.IconOverride = FSlateIcon(FCoreStyle::Get().GetStyleSetName(), "GenericCommands.Copy");
+		CopyContentParams.DirectActions = CopyAction;
+		
+		
+		MenuBuilder.AddMenuEntry(CopyContentParams);
+
+		FMenuEntryParams PasteContentParams;
+		PasteContentParams.LabelOverride = NSLOCTEXT("PropertyView", "PasteCategoryProperties", "Paste All Properties in Category");
+		PasteContentParams.ToolTipOverride = TAttribute<FText>::CreateLambda([this]()
+		{
+			return CanPasteCategory()
+				? NSLOCTEXT("PropertyView", "PasteCategoryProperties_ToolTip", "Paste the copied property values here")
+				// @note: this is specific to the constraint that the destination category has to match the source category (copied from) exactly 
+				: NSLOCTEXT("PropertyView", "CantPasteCategoryProperties_ToolTip", "The properties in this category don't match the contents of the clipboard");
+		});
+		PasteContentParams.InputBindingOverride = FInputChord(EModifierKey::Shift, EKeys::LeftMouseButton).GetInputText(bLongDisplayName);
+		PasteContentParams.IconOverride = FSlateIcon(FCoreStyle::Get().GetStyleSetName(), "GenericCommands.Paste");
+		PasteContentParams.DirectActions = PasteAction;
+		MenuBuilder.AddMenuEntry(PasteContentParams);
+	}
+	MenuBuilder.EndSection();
+	
+	return true;
+}
+
 EVisibility SDetailCategoryTableRow::IsSeparatorVisible() const
 {
 	return bIsInnerCategory || IsItemExpanded() ? EVisibility::Collapsed : EVisibility::Visible;
@@ -132,6 +228,8 @@ const FSlateBrush* SDetailCategoryTableRow::GetBackgroundImage() const
 
 FSlateColor SDetailCategoryTableRow::GetInnerBackgroundColor() const
 {
+	FSlateColor Color = FSlateColor(FLinearColor::White);
+	
 	if (bShowBorder && bIsInnerCategory)
 	{
 		int32 IndentLevel = -1;
@@ -142,10 +240,16 @@ FSlateColor SDetailCategoryTableRow::GetInnerBackgroundColor() const
 
 		IndentLevel = FMath::Max(IndentLevel - 1, 0);
 
-		return PropertyEditorConstants::GetRowBackgroundColor(IndentLevel, this->IsHovered());
+		Color = PropertyEditorConstants::GetRowBackgroundColor(IndentLevel, this->IsHovered());
 	}
 
-	return FSlateColor(FLinearColor::White);
+	if (PulseAnimation.IsPlaying())
+	{
+		float Lerp = PulseAnimation.GetLerp();
+		return FMath::Lerp(FAppStyle::Get().GetSlateColor("Colors.Hover2").GetSpecifiedColor(), Color.GetSpecifiedColor(), Lerp);
+	}
+
+	return Color;
 }
 
 FSlateColor SDetailCategoryTableRow::GetOuterBackgroundColor() const
@@ -158,17 +262,189 @@ FSlateColor SDetailCategoryTableRow::GetOuterBackgroundColor() const
 	return FAppStyle::Get().GetSlateColor("Colors.Panel");
 }
 
+void SDetailCategoryTableRow::OnCopyCategory()
+{
+	if (!OwnerTreeNode.IsValid())
+	{
+		return;
+	}
+
+	if (TArray<TSharedPtr<IPropertyHandle>> CategoryProperties = GetPropertyHandles(true);
+		!CategoryProperties.IsEmpty())
+	{
+		TArray<FString> PropertiesNotCopied;
+		PropertiesNotCopied.Reserve(CategoryProperties.Num());
+		
+		TMap<FString, FString> PropertyValues;
+		PropertyValues.Reserve(CategoryProperties.Num());
+		
+		for (TSharedPtr<IPropertyHandle> PropertyHandle : CategoryProperties)
+		{
+			if (PropertyHandle.IsValid() && PropertyHandle->IsValidHandle())
+			{
+				FString PropertyPath = UE::PropertyEditor::GetPropertyPath(PropertyHandle);
+				
+				FString PropertyValueStr;
+				if (PropertyHandle->GetValueAsFormattedString(PropertyValueStr, PPF_Copy) == FPropertyAccess::Success)
+				{
+					PropertyValues.Add(PropertyPath, PropertyValueStr);
+				}
+				else
+				{
+					PropertiesNotCopied.Add(PropertyHandle->GetPropertyDisplayName().ToString());
+				}
+			}
+		}
+
+		if (!PropertiesNotCopied.IsEmpty())
+		{
+			UE_LOG(
+				LogPropertyNode,
+				Warning,
+				TEXT("One or more of the properties in category \"%s\" was not copied:\n%s"),
+				*DisplayName.ToString(),
+				*FString::Join(PropertiesNotCopied, TEXT("\n")));
+		}
+
+		FPropertyEditorClipboard::ClipboardCopy([&PropertyValues](TMap<FName, FString>& OutTaggedClipboard)
+		{
+			for (const TPair<FString, FString>& PropertyValuePair : PropertyValues)
+			{
+				OutTaggedClipboard.Add(FName(PropertyValuePair.Key), PropertyValuePair.Value);
+			}
+		});
+
+		PulseAnimation.Play(SharedThis(this));
+	}
+}
+
+bool SDetailCategoryTableRow::CanCopyCategory() const
+{
+	if (!OwnerTreeNode.IsValid())
+	{
+		return false;
+	}
+
+	TArray<TSharedPtr<IPropertyHandle>> PropertyHandles = GetPropertyHandles(true);
+	return !PropertyHandles.IsEmpty();
+}
+
+void SDetailCategoryTableRow::OnPasteCategory()
+{
+	if (!OwnerTreeNode.IsValid()
+		|| !CanPasteCategory())
+	{
+		return;
+	}
+
+	if (const TArray<TSharedPtr<IPropertyHandle>> CategoryProperties = GetPropertyHandles(true);
+		!CategoryProperties.IsEmpty())
+	{
+		FScopedTransaction Transaction(NSLOCTEXT("UnrealEd", "PasteCategoryProperties", "Paste Category Properties"));
+
+		TArray<FString> PropertiesNotPasted;
+		PropertiesNotPasted.Reserve(CategoryProperties.Num());
+
+		{
+			// Assign a unique id to this paste operation, enabling batching of various operations
+			const FGuid OperationGuid = FGuid::NewGuid();			
+			for (const TPair<FName, FString>& KVP : PreviousClipboardData.PropertyValues)
+			{
+				OnPasteFromTextDelegate->Broadcast(KVP.Key.ToString(), KVP.Value, OperationGuid);
+			}
+		}
+		
+		if (!PropertiesNotPasted.IsEmpty())
+		{
+			UE_LOG(
+				LogPropertyNode,
+				Warning,
+				TEXT("One or more of the properties in category \"%s\" was not pasted:\n%s"),
+				*DisplayName.ToString(),
+				*FString::Join(PropertiesNotPasted, TEXT("\n")));
+		}
+
+		IDetailsViewPrivate* DetailsView = OwnerTreeNode.Pin()->GetDetailsView();
+		DetailsView->ForceRefresh();
+	}
+}
+
+bool SDetailCategoryTableRow::CanPasteCategory()
+{
+	if (!OwnerTreeNode.IsValid())
+	{
+		return false;
+	}
+
+	const TArray<TSharedPtr<IPropertyHandle>> PropertyHandles = GetPropertyHandles(true);
+
+	// @note: Usually we'd check for IsEditConst or IsEditConditionMet, but if used for PP settings (for example),
+	// by default no properties are editable unless explicitly overridden, so this check would in all cases would prevent paste.
+	constexpr bool bHasEditables = true;
+	
+	// No editable properties to write to
+	if constexpr (!bHasEditables)
+	{
+		return false;
+	}
+
+	FString ClipboardContent;
+	FPropertyEditorClipboard::ClipboardPaste(ClipboardContent);
+
+	// If same as last, return previously resolved applicability
+	if (PreviousClipboardData.Content.Get({}).Equals(ClipboardContent))
+	{
+		return PreviousClipboardData.bIsApplicable;
+	}
+
+	// New clipboard contents, non-applicable by default
+	PreviousClipboardData.Reset();
+
+	// Can't be empty, must be json
+	if (!UE::PropertyEditor::Internal::IsJsonString(ClipboardContent))
+	{
+		return false;
+	}
+
+	PreviousClipboardData.Reserve(PropertyHandles.Num());
+
+	if (!UE::PropertyEditor::Internal::TryParseClipboard(ClipboardContent, PreviousClipboardData.PropertyValues))
+	{
+		return false;
+	}
+
+	PreviousClipboardData.PropertyValues.GenerateKeyArray(PreviousClipboardData.PropertyNames);
+
+	TArray<FString> PropertyNames;
+	Algo::Transform(PropertyHandles, PropertyNames, [](const TSharedPtr<IPropertyHandle>& InPropertyHandle)
+	{
+		return UE::PropertyEditor::GetPropertyPath(InPropertyHandle);
+	});
+
+	PreviousClipboardData.PropertyNames.Sort(FNameLexicalLess());
+	PropertyNames.Sort();
+
+	// @note: properties must all match to be applicable
+	PreviousClipboardData.Content = ClipboardContent;
+
+	return PreviousClipboardData.bIsApplicable = Algo::Compare(PreviousClipboardData.PropertyNames, PropertyNames);
+}
+
 FReply SDetailCategoryTableRow::OnMouseButtonDown(const FGeometry& MyGeometry, const FPointerEvent& MouseEvent)
 {
+	if (MouseEvent.GetModifierKeys().IsShiftDown())
+	{
+		// Allows this to be used for the paste properties shortcut
+		return FReply::Unhandled();
+	}
+	
 	if (MouseEvent.GetEffectingButton() == EKeys::LeftMouseButton)
 	{
 		ToggleExpansion();
 		return FReply::Handled();
 	}
-	else
-	{
-		return FReply::Unhandled();
-	}
+
+	return FReply::Unhandled();
 }
 
 FReply SDetailCategoryTableRow::OnMouseButtonDoubleClick(const FGeometry& InMyGeometry, const FPointerEvent& InMouseEvent)
diff --git a/Engine/Source/Editor/PropertyEditor/Private/SDetailCategoryTableRow.h b/Engine/Source/Editor/PropertyEditor/Private/SDetailCategoryTableRow.h
--- a/Engine/Source/Editor/PropertyEditor/Private/SDetailCategoryTableRow.h
+++ b/Engine/Source/Editor/PropertyEditor/Private/SDetailCategoryTableRow.h
@@ -5,6 +5,7 @@
 #include "DetailTreeNode.h"
 #include "Input/Reply.h"
 #include "Layout/Visibility.h"
+#include "PropertyEditorClipboardPrivate.h"
 #include "SDetailTableRowBase.h"
 #include "Templates/SharedPointer.h"
 #include "Widgets/DeclarativeSyntaxSupport.h"
@@ -22,20 +23,47 @@ public:
 		SLATE_ARGUMENT(bool, InnerCategory)
 		SLATE_ARGUMENT(TSharedPtr<SWidget>, HeaderContent)
 		SLATE_ARGUMENT(bool, ShowBorder)
+		SLATE_ARGUMENT(TSharedPtr<FOnPasteFromText>, PasteFromText)
 	SLATE_END_ARGS()
 
 	void Construct(const FArguments& InArgs, TSharedRef<FDetailTreeNode> InOwnerTreeNode, const TSharedRef<STableViewBase>& InOwnerTableView);
 
+	// ~Begin SWidget Interface
+	virtual FReply OnMouseButtonUp(const FGeometry& MyGeometry, const FPointerEvent& MouseEvent) override;
+	// ~Begin End Interface
+
+protected:
+	virtual bool OnContextMenuOpening( FMenuBuilder& MenuBuilder ) override;
+	
 private:
 	EVisibility IsSeparatorVisible() const;
 	const FSlateBrush* GetBackgroundImage() const;
 	FSlateColor GetInnerBackgroundColor() const;
 	FSlateColor GetOuterBackgroundColor() const;
 
-	FReply OnMouseButtonDown(const FGeometry& MyGeometry, const FPointerEvent& MouseEvent) override;
-	FReply OnMouseButtonDoubleClick(const FGeometry& InMyGeometry, const FPointerEvent& InMouseEvent) override;
+	void OnCopyCategory();
+	bool CanCopyCategory() const;
+	void OnPasteCategory();
+	bool CanPasteCategory();
+
+	virtual FReply OnMouseButtonDown(const FGeometry& MyGeometry, const FPointerEvent& MouseEvent) override;
+	virtual FReply OnMouseButtonDoubleClick(const FGeometry& InMyGeometry, const FPointerEvent& InMouseEvent) override;
 
 private:
-	bool bIsInnerCategory;
-	bool bShowBorder;
+	/** Cached category name. */
+	FText DisplayName;
+
+	/** Previously parsed clipboard data. */
+	UE::PropertyEditor::Internal::FClipboardData PreviousClipboardData;
+
+	bool bIsInnerCategory = false;
+	bool bShowBorder = false;
+	FUIAction CopyAction;
+	FUIAction PasteAction;
+
+	/** Delegate handling pasting an optionally tagged text snippet */
+	TSharedPtr<FOnPasteFromText> OnPasteFromTextDelegate;
+
+	/** Animation curve for displaying pulse */
+	FCurveSequence PulseAnimation;
 };
diff --git a/Engine/Source/Editor/PropertyEditor/Private/SDetailSingleItemRow.cpp b/Engine/Source/Editor/PropertyEditor/Private/SDetailSingleItemRow.cpp
--- a/Engine/Source/Editor/PropertyEditor/Private/SDetailSingleItemRow.cpp
+++ b/Engine/Source/Editor/PropertyEditor/Private/SDetailSingleItemRow.cpp
@@ -2,36 +2,35 @@
 
 #include "SDetailSingleItemRow.h"
 
+#include "Algo/AnyOf.h"
+#include "Algo/Compare.h"
 #include "DetailGroup.h"
 #include "DetailPropertyRow.h"
 #include "DetailWidgetRow.h"
 #include "Editor.h"
-#include "EditorMetadataOverrides.h"
 #include "IDetailDragDropHandler.h"
 #include "IDetailPropertyExtensionHandler.h"
-#include "ObjectPropertyNode.h"
-#include "UserInterface/PropertyEditor/PropertyEditorConstants.h"
+#include "Modules/ModuleInterface.h"
+#include "Modules/ModuleManager.h"
+#include "PropertyEditorClipboard.h"
+#include "PropertyEditorClipboardPrivate.h"
+#include "PropertyEditorCopyPastePrivate.h"
 #include "PropertyEditorModule.h"
 #include "PropertyHandleImpl.h"
+#include "PropertyPermissionList.h"
 #include "SConstrainedBox.h"
 #include "SDetailExpanderArrow.h"
 #include "SDetailRowIndent.h"
-
-#include "HAL/PlatformApplicationMisc.h"
-#include "Modules/ModuleInterface.h"
-#include "Modules/ModuleManager.h"
-#include "Settings/EditorExperimentalSettings.h"
 #include "Styling/StyleColors.h"
+#include "UserInterface/PropertyEditor/PropertyEditorConstants.h"
 #include "Widgets/Images/SImage.h"
 #include "Widgets/Input/SButton.h"
 #include "Widgets/Input/SComboButton.h"
-#include "PropertyPermissionList.h"
 
 namespace DetailWidgetConstants
 {
 	const FMargin LeftRowPadding( 20.0f, 0.0f, 10.0f, 0.0f );
 	const FMargin RightRowPadding( 12.0f, 0.0f, 2.0f, 0.0f );
-	const float PulseAnimationLength = 0.5f;
 }
 
 namespace SDetailSingleItemRow_Helper
@@ -248,10 +247,10 @@ TSharedPtr<FPropertyNode> SDetailSingleItemRow::GetPropertyNode() const
 TSharedPtr<IPropertyHandle> SDetailSingleItemRow::GetPropertyHandle() const
 {
 	TSharedPtr<IPropertyHandle> Handle;
-	TSharedPtr<FPropertyNode> PropertyNode = GetPropertyNode();
+	const TSharedPtr<FPropertyNode> PropertyNode = GetPropertyNode();
 	if (PropertyNode.IsValid())
 	{
-		TSharedPtr<FDetailTreeNode> OwnerTreeNodePtr = OwnerTreeNode.Pin();
+		const TSharedPtr<FDetailTreeNode> OwnerTreeNodePtr = OwnerTreeNode.Pin();
 		if (OwnerTreeNodePtr.IsValid())
 		{
 			IDetailsViewPrivate* DetailsView = OwnerTreeNodePtr->GetDetailsView();
@@ -308,7 +307,7 @@ void SDetailSingleItemRow::Construct( const FArguments& InArgs, FDetailLayoutCus
 	IDetailsViewPrivate* DetailsView = InOwnerTreeNode->GetDetailsView();
 	FDetailColumnSizeData& ColumnSizeData = DetailsView->GetColumnSizeData();
 
-	PulseAnimation.AddCurve(0.0f, DetailWidgetConstants::PulseAnimationLength, ECurveEaseFunction::CubicInOut);
+	PulseAnimation.AddCurve(0.0f, UE::PropertyEditor::Private::PulseAnimationLength, ECurveEaseFunction::CubicInOut);
 
 	// Play on construction if animation was started from a behavior the re-constructs this widget
 	if (DetailsView->IsNodeAnimating(GetPropertyNode()))
@@ -316,11 +315,14 @@ void SDetailSingleItemRow::Construct( const FArguments& InArgs, FDetailLayoutCus
 		PulseAnimation.Play(SharedThis(this));
 	}
 
+	TSharedPtr<FDetailCategoryImpl> Category = InOwnerTreeNode->GetParentCategory();
+
 	const bool bIsValidTreeNode = InOwnerTreeNode->GetParentCategory().IsValid() && InOwnerTreeNode->GetParentCategory()->IsParentLayoutValid();
 	if (bIsValidTreeNode)
 	{
 		if (Customization->IsValidCustomization())
 		{
+			TSharedPtr<FDetailGroup> Group = nullptr;
 			WidgetRow = Customization->GetWidgetRow();
 
 			// Setup copy / paste actions
@@ -340,6 +342,15 @@ void SDetailSingleItemRow::Construct( const FArguments& InArgs, FDetailLayoutCus
 						PasteAction.ExecuteAction = FExecuteAction::CreateSP(this, &SDetailSingleItemRow::OnPasteProperty);
 						PasteAction.CanExecuteAction = FCanExecuteAction::CreateSP(this, &SDetailSingleItemRow::CanPasteProperty);
 					}
+					else if (Group = Customization->DetailGroup;
+						Group.IsValid())
+					{
+						CopyAction.ExecuteAction = FExecuteAction::CreateSP(this, &SDetailSingleItemRow::OnCopyGroup);
+						CopyAction.CanExecuteAction = FCanExecuteAction::CreateSP(this, &SDetailSingleItemRow::CanCopyGroup);
+
+						PasteAction.ExecuteAction = FExecuteAction::CreateSP(this, &SDetailSingleItemRow::OnPasteGroup);
+						PasteAction.CanExecuteAction = FCanExecuteAction::CreateSP(this, &SDetailSingleItemRow::CanPasteGroup);
+					}
 					else
 					{
 						CopyAction.ExecuteAction = FExecuteAction::CreateLambda([]() {});
@@ -348,6 +359,28 @@ void SDetailSingleItemRow::Construct( const FArguments& InArgs, FDetailLayoutCus
 						PasteAction.CanExecuteAction = FCanExecuteAction::CreateLambda([]() { return false; });
 					}
 				}
+
+				if (WidgetRow.IsPasteFromTextBound())
+				{
+					OnPasteFromTextDelegate = WidgetRow.OnPasteFromTextDelegate.Pin();					
+				}
+				else if (Category.IsValid())
+				{
+					OnPasteFromTextDelegate = Category->OnPasteFromText();
+				}
+				// If still not set, but this is a group, initialize
+				else if (Group.IsValid())
+				{
+					if (TSharedPtr<FOnPasteFromText> GroupOnPasteFromTextDelegate = Group->OnPasteFromText())
+					{
+						OnPasteFromTextDelegate = GroupOnPasteFromTextDelegate;
+					}
+				}
+
+				if (OnPasteFromTextDelegate.IsValid())
+				{
+					OnPasteFromTextDelegate->AddSP(this, &SDetailSingleItemRow::OnPasteFromText);
+				}
 			}
 
 			// Populate the extension content in the WidgetRow if there's an extension handler.
@@ -710,7 +743,7 @@ FReply SDetailSingleItemRow::OnMouseButtonUp(const FGeometry& MyGeometry, const
 			{
 				// Mark property node as animating so we will animate after any potential re-construction
 				IDetailsViewPrivate* DetailsView = OwnerTreeNode.Pin()->GetDetailsView();
-				DetailsView->MarkNodeAnimating(PropertyNode, DetailWidgetConstants::PulseAnimationLength);
+				DetailsView->MarkNodeAnimating(PropertyNode, UE::PropertyEditor::Private::PulseAnimationLength);
 			}
 			bIsHandled = true;
 		}
@@ -777,6 +810,177 @@ FSlateColor SDetailSingleItemRow::GetInnerBackgroundColor() const
 	return Color;
 }
 
+void SDetailSingleItemRow::OnCopyGroup()
+{
+	if (!OwnerTreeNode.IsValid())
+	{
+		return;
+	}
+
+	if (TArray<TSharedPtr<IPropertyHandle>> GroupProperties = GetPropertyHandles(true);
+		!GroupProperties.IsEmpty())
+	{
+		TArray<FString> PropertiesNotCopied;
+		PropertiesNotCopied.Reserve(GroupProperties.Num());
+		
+		TMap<FString, FString> PropertyValues;
+		PropertyValues.Reserve(GroupProperties.Num());
+		
+		for (TSharedPtr<IPropertyHandle> PropertyHandle : GroupProperties)
+		{
+			if (PropertyHandle.IsValid() && PropertyHandle->IsValidHandle())
+			{
+				FString PropertyPath = UE::PropertyEditor::GetPropertyPath(PropertyHandle);
+				
+				FString PropertyValueStr;
+				if (PropertyHandle->GetValueAsFormattedString(PropertyValueStr, PPF_Copy) == FPropertyAccess::Success)
+				{
+					PropertyValues.Add(PropertyPath, PropertyValueStr);
+				}
+				else
+				{
+					PropertiesNotCopied.Add(PropertyHandle->GetPropertyDisplayName().ToString());
+				}
+			}
+		}
+
+		if (!PropertiesNotCopied.IsEmpty())
+		{
+			UE_LOG(
+				LogPropertyNode,
+				Warning,
+				TEXT("One or more of the properties in group \"%s\" was not copied:\n%s"),
+				*GetRowNameText(),
+				*FString::Join(PropertiesNotCopied, TEXT("\n")));
+		}
+
+		FPropertyEditorClipboard::ClipboardCopy([&PropertyValues](TMap<FName, FString>& OutTaggedClipboard)
+		{
+			for (const TPair<FString, FString>& PropertyValuePair : PropertyValues)
+			{
+				OutTaggedClipboard.Add(FName(PropertyValuePair.Key), PropertyValuePair.Value);
+			}
+		});
+
+		PulseAnimation.Play(SharedThis(this));
+	}
+}
+
+bool SDetailSingleItemRow::CanCopyGroup() const
+{
+	if (!OwnerTreeNode.IsValid())
+	{
+		return false;
+	}
+
+	TArray<TSharedPtr<IPropertyHandle>> GroupPropertyHandles = GetPropertyHandles(true);
+	return !GroupPropertyHandles.IsEmpty();
+}
+
+void SDetailSingleItemRow::OnPasteGroup()
+{
+	if (!OwnerTreeNode.IsValid()
+		|| !CanPasteGroup())
+	{
+		return;
+	}
+
+	if (OnPasteFromTextDelegate.IsValid())
+	{
+		if (const TArray<TSharedPtr<IPropertyHandle>> GroupProperties = GetPropertyHandles(true);
+			!GroupProperties.IsEmpty())
+		{
+			FScopedTransaction Transaction(NSLOCTEXT("UnrealEd", "PasteGroupProperties", "Paste Group Properties"));
+
+			TArray<FString> PropertiesNotPasted;
+			PropertiesNotPasted.Reserve(GroupProperties.Num());
+
+			{
+				const FGuid OperationGuid = FGuid::NewGuid();
+			
+				for (const TPair<FName, FString>& PropertyNameValuePair : PreviousClipboardData.PropertyValues)
+				{
+					OnPasteFromTextDelegate->Broadcast(PropertyNameValuePair.Key.ToString(), PropertyNameValuePair.Value, OperationGuid);
+				}
+			}
+		
+			if (!PropertiesNotPasted.IsEmpty())
+			{
+				UE_LOG(
+					LogPropertyNode,
+					Warning,
+					TEXT("One or more of the properties in group \"%s\" was not pasted:\n%s"),
+					*GetRowNameText(),
+					*FString::Join(PropertiesNotPasted, TEXT("\n")));
+			}
+
+			IDetailsViewPrivate* DetailsView = OwnerTreeNode.Pin()->GetDetailsView();
+			DetailsView->ForceRefresh();
+		}
+	}
+}
+
+bool SDetailSingleItemRow::CanPasteGroup()
+{
+	if (!OwnerTreeNode.IsValid())
+	{
+		return false;
+	}
+
+	const TArray<TSharedPtr<IPropertyHandle>> GroupPropertyHandles = GetPropertyHandles(true);
+
+	// @note: Usually we'd check for IsEditConst or IsEditConditionMet, but if used for PP settings,
+	// by default no properties are editable unless explicitly overridden, so this check would in all cases would prevent paste.
+	constexpr bool bHasEditables = true;
+	
+	// No editable properties to write to
+	if constexpr (!bHasEditables)
+	{
+		return false;
+	}
+
+	FString ClipboardContent;
+	FPropertyEditorClipboard::ClipboardPaste(ClipboardContent);
+
+	// If same as last, return previously resolved applicability
+	if (PreviousClipboardData.Content.Get({}).Equals(ClipboardContent))
+	{
+		return PreviousClipboardData.bIsApplicable;
+	}
+
+	// New clipboard contents, non-applicable by default
+	PreviousClipboardData.Reset();
+
+	// Can't be empty, must be json
+	if (!UE::PropertyEditor::Internal::IsJsonString(ClipboardContent))
+	{
+		return false;
+	}
+
+	PreviousClipboardData.Reserve(GroupPropertyHandles.Num());
+
+	if (!UE::PropertyEditor::Internal::TryParseClipboard(ClipboardContent, PreviousClipboardData.PropertyValues))
+	{
+		return false;
+	}
+
+	PreviousClipboardData.PropertyValues.GenerateKeyArray(PreviousClipboardData.PropertyNames);
+
+	TArray<FString> PropertyNames;
+	Algo::Transform(GroupPropertyHandles, PropertyNames, [](const TSharedPtr<IPropertyHandle>& InPropertyHandle)
+	{
+		return UE::PropertyEditor::GetPropertyPath(InPropertyHandle);
+	});
+
+	PreviousClipboardData.PropertyNames.Sort(FNameLexicalLess());
+	PropertyNames.Sort();
+
+	// @note: properties must all match to be applicable
+	PreviousClipboardData.Content = ClipboardContent;
+
+	return PreviousClipboardData.bIsApplicable = Algo::Compare(PreviousClipboardData.PropertyNames, PropertyNames);
+}
+
 bool SDetailSingleItemRow::OnContextMenuOpening(FMenuBuilder& MenuBuilder)
 {
 	FUIAction CopyDisplayNameAction = FExecuteAction::CreateSP(this, &SDetailSingleItemRow::OnCopyPropertyDisplayName);
@@ -792,19 +996,47 @@ bool SDetailSingleItemRow::OnContextMenuOpening(FMenuBuilder& MenuBuilder)
 
 	if (CopyAction.IsBound() && PasteAction.IsBound())
 	{
-		bool bLongDisplayName = false;
+		const bool bLongDisplayName = false;
+		const bool bIsGroup = Customization->IsValidCustomization() && Customization->DetailGroup.IsValid();
 
 		FMenuEntryParams CopyContentParams;
-		CopyContentParams.LabelOverride = NSLOCTEXT("PropertyView", "CopyProperty", "Copy");
-		CopyContentParams.ToolTipOverride = NSLOCTEXT("PropertyView", "CopyProperty_ToolTip", "Copy this property value");
+		if (bIsGroup)
+		{
+			CopyContentParams.LabelOverride = NSLOCTEXT("PropertyView", "CopyGroupProperties", "Copy All Properties in Group");
+			CopyContentParams.ToolTipOverride = TAttribute<FText>::CreateLambda([this]()
+			{
+				return CanCopyGroup()
+					? NSLOCTEXT("PropertyView", "CopyGroupProperties_ToolTip", "Copy all properties in this group")
+					: NSLOCTEXT("PropertyView", "CantCopyGroupProperties_ToolTip", "None of the properties in this group can be copied");
+			});
+		}
+		else
+		{
+			CopyContentParams.LabelOverride = NSLOCTEXT("PropertyView", "CopyProperty", "Copy");
+			CopyContentParams.ToolTipOverride = NSLOCTEXT("PropertyView", "CopyProperty_ToolTip", "Copy this property value");
+		}
 		CopyContentParams.InputBindingOverride = FInputChord(EModifierKey::Shift, EKeys::RightMouseButton).GetInputText(bLongDisplayName);
 		CopyContentParams.IconOverride = FSlateIcon(FCoreStyle::Get().GetStyleSetName(), "GenericCommands.Copy");
 		CopyContentParams.DirectActions = CopyAction;
 		MenuBuilder.AddMenuEntry(CopyContentParams);
 
 		FMenuEntryParams PasteContentParams;
-		PasteContentParams.LabelOverride = NSLOCTEXT("PropertyView", "PasteProperty", "Paste");
-		PasteContentParams.ToolTipOverride = NSLOCTEXT("PropertyView", "PasteProperty_ToolTip", "Paste the copied value here");
+		if (bIsGroup)
+		{
+			PasteContentParams.LabelOverride = NSLOCTEXT("PropertyView", "PasteGroupProperties", "Paste All Properties in Group");
+			PasteContentParams.ToolTipOverride = TAttribute<FText>::CreateLambda([this]()
+			{
+				return CanPasteGroup()
+					? NSLOCTEXT("PropertyView", "PasteGroupProperties_ToolTip", "Paste the copied property values here")
+					// @note: this is specific to the constraint that the destination group has to match the source group (copied from) exactly 
+					: NSLOCTEXT("PropertyView", "CantPasteGroupProperties_ToolTip", "The properties in this group don't match the contents of the clipboard");
+			});
+		}
+		else
+		{
+			PasteContentParams.LabelOverride = NSLOCTEXT("PropertyView", "PasteProperty", "Paste");
+			PasteContentParams.ToolTipOverride = NSLOCTEXT("PropertyView", "PasteProperty_ToolTip", "Paste the copied value here");	
+		}
 		PasteContentParams.InputBindingOverride = FInputChord(EModifierKey::Shift, EKeys::LeftMouseButton).GetInputText(bLongDisplayName);
 		PasteContentParams.IconOverride = FSlateIcon(FCoreStyle::Get().GetStyleSetName(), "GenericCommands.Paste");
 		PasteContentParams.DirectActions = PasteAction;
@@ -915,7 +1147,7 @@ void SDetailSingleItemRow::OnCopyProperty()
 			FString Value;
 			if (Handle->GetValueAsFormattedString(Value, PPF_Copy) == FPropertyAccess::Success)
 			{
-				FPlatformApplicationMisc::ClipboardCopy(*Value);
+				FPropertyEditorClipboard::ClipboardCopy(*Value);
 				PulseAnimation.Play(SharedThis(this));
 			}
 		}
@@ -929,7 +1161,7 @@ void SDetailSingleItemRow::OnCopyPropertyDisplayName()
 		TSharedPtr<FPropertyNode> PropertyNode = GetPropertyNode();
 		if (PropertyNode.IsValid())
 		{
-			FPlatformApplicationMisc::ClipboardCopy(*PropertyNode->GetDisplayName().ToString());
+			FPropertyEditorClipboard::ClipboardCopy(*PropertyNode->GetDisplayName().ToString());
 		}
 	}
 }
@@ -962,7 +1194,7 @@ void SDetailSingleItemRow::OnCopyPropertyInternalName()
 			{
 				if (const UStruct* OwnerStruct = Property->GetOwnerStruct())
 				{
-					FPlatformApplicationMisc::ClipboardCopy(*OwnerStruct->GetAuthoredNameForField(Property));
+					FPropertyEditorClipboard::ClipboardCopy(*OwnerStruct->GetAuthoredNameForField(Property));
 				}
 			}
 		}
@@ -991,61 +1223,125 @@ bool SDetailSingleItemRow::CanCopyPropertyInternalName()
 	return false;
 }
 
-void SDetailSingleItemRow::OnPasteProperty()
+bool SDetailSingleItemRow::CanPasteProperty() const
 {
 	FString ClipboardContent;
-	FPlatformApplicationMisc::ClipboardPaste(ClipboardContent);
+	if (OwnerTreeNode.IsValid())
+	{
+		FPropertyEditorClipboard::ClipboardPaste(ClipboardContent);
+	}
+
+	return CanPasteFromText(TEXT(""), ClipboardContent);
+}
 
-	if (!ClipboardContent.IsEmpty() && OwnerTreeNode.IsValid())
+bool SDetailSingleItemRow::CanPasteFromText(const FString& InTag, const FString& InText) const
+{
+	if (InText.IsEmpty())
 	{
-		TSharedPtr<FPropertyNode> PropertyNode = GetPropertyNode();
-		if (!PropertyNode.IsValid() && Customization->DetailGroup.IsValid())
+		return false;
+	}
+	
+	// Prevent paste from working if the property's edit condition is not met.
+	TSharedPtr<FDetailPropertyRow> PropertyRow = Customization->PropertyRow;
+	if (!PropertyRow.IsValid() && Customization->DetailGroup.IsValid())
+	{
+		PropertyRow = Customization->DetailGroup->GetHeaderPropertyRow();
+	}
+
+	if (PropertyRow.IsValid())
+	{
+		if (const FPropertyEditor* PropertyEditor = PropertyRow->GetPropertyEditor().Get())
 		{
-			PropertyNode = Customization->DetailGroup->GetHeaderPropertyNode();
+			return !PropertyEditor->IsEditConst();
 		}
-		if (PropertyNode.IsValid())
-		{
-			FScopedTransaction Transaction(NSLOCTEXT("UnrealEd", "PasteProperty", "Paste Property"));
+	}
 
-			IDetailsViewPrivate* DetailsView = OwnerTreeNode.Pin()->GetDetailsView();
-			TSharedPtr<IPropertyHandle> Handle = PropertyEditorHelpers::GetPropertyHandle(PropertyNode.ToSharedRef(), DetailsView->GetNotifyHook(), DetailsView->GetPropertyUtilities());
+	return false;
+}
 
-			Handle->SetValueFromFormattedString(ClipboardContent, EPropertyValueSetFlags::InstanceObjects);
+void SDetailSingleItemRow::OnPasteProperty()
+{
+	FString ClipboardContent;
 
-			// Need to refresh the details panel in case a property was pasted over another.
-			OwnerTreeNode.Pin()->GetDetailsView()->ForceRefresh();
+	const TSharedPtr<FPropertyNode> PropertyNode = GetPropertyNode();
+	const FString PropertyPath = UE::PropertyEditor::Private::GetPropertyPath(
+		[&]() { return GetPropertyHandle(); },
+		[&]() { return PropertyNode; });
 
+	if (PropertyPath.IsEmpty())
+	{
+		FPropertyEditorClipboard::ClipboardPaste(ClipboardContent);
+	}
+	else
+	{
+		FPropertyEditorClipboard::ClipboardPaste(ClipboardContent, FName(PropertyPath));
+	}
+
+	if (PasteFromText(TEXT(""), ClipboardContent))
+	{
+		// Mark property node as animating so we will animate after re-construction
+		IDetailsViewPrivate* DetailsView = OwnerTreeNode.Pin()->GetDetailsView();
+		DetailsView->MarkNodeAnimating(PropertyNode, UE::PropertyEditor::Private::PulseAnimationLength);
+
+		// Need to refresh the details panel in case a property was pasted over another.
+		OwnerTreeNode.Pin()->GetDetailsView()->ForceRefresh();
+	}
+}
+
+void SDetailSingleItemRow::OnPasteFromText(const FString& InTag, const FString& InText, const TOptional<FGuid>& InOperationId)
+{
+	if (PasteFromText(InTag, InText))
+	{
+		TSharedPtr<FPropertyNode> PropertyNode = GetPropertyNode();
+		if (PropertyNode.IsValid())
+		{
 			// Mark property node as animating so we will animate after re-construction
-			DetailsView->MarkNodeAnimating(PropertyNode, DetailWidgetConstants::PulseAnimationLength);
+			IDetailsViewPrivate* DetailsView = OwnerTreeNode.Pin()->GetDetailsView();
+			DetailsView->MarkNodeAnimating(PropertyNode, UE::PropertyEditor::Private::PulseAnimationLength, InOperationId);
 		}
 	}
 }
 
-bool SDetailSingleItemRow::CanPasteProperty() const
+bool SDetailSingleItemRow::PasteFromText(const FString& InTag, const FString& InText)
 {
-	// Prevent paste from working if the property's edit condition is not met.
-	TSharedPtr<FDetailPropertyRow> PropertyRow = Customization->PropertyRow;
-	if (!PropertyRow.IsValid() && Customization->DetailGroup.IsValid())
+	if (InText.IsEmpty() || !OwnerTreeNode.IsValid())
 	{
-		PropertyRow = Customization->DetailGroup->GetHeaderPropertyRow();
+		return false;
+	}
+	
+	TSharedPtr<FPropertyNode> PropertyNode = GetPropertyNode();
+	if (!PropertyNode.IsValid() && Customization->DetailGroup.IsValid())
+	{
+		PropertyNode = Customization->DetailGroup->GetHeaderPropertyNode();
 	}
 
-	if (PropertyRow.IsValid())
+	// If still invalid, there's nothing to paste to
+	if (!PropertyNode.IsValid())
+	{
+		return false;
+	}
+	
+	const TSharedPtr<IPropertyHandle> PropertyHandle = GetPropertyHandle();
+	if (!InTag.IsEmpty())
 	{
-		FPropertyEditor* PropertyEditor = PropertyRow->GetPropertyEditor().Get();
-		if (PropertyEditor)
+		const FString PropertyPath = UE::PropertyEditor::Private::GetPropertyPath(
+			[&]() { return GetPropertyHandle(); },
+			[&]() { return PropertyNode; });
+		
+		// ensure that if tag is specified, that it matches the subscriber
+		if (!InTag.Equals(PropertyPath))
 		{
-			return !PropertyEditor->IsEditConst();
+			return false;
 		}
 	}
 
-	FString ClipboardContent;
-	if (OwnerTreeNode.IsValid())
+	FScopedTransaction Transaction(NSLOCTEXT("UnrealEd", "PasteProperty", "Paste Property"));
+	if (PropertyHandle->SetValueFromFormattedString(InText, EPropertyValueSetFlags::InstanceObjects) != FPropertyAccess::Success)
 	{
-		FPlatformApplicationMisc::ClipboardPaste(ClipboardContent);
+		return false;
 	}
 
-	return !ClipboardContent.IsEmpty();
+	return true;
 }
 
 // Helper function to determine which parent Struct a property comes from. If PropertyName is not a real property,
@@ -1078,7 +1374,7 @@ void SDetailSingleItemRow::CopyRowNameText() const
 	const FString RowNameText = GetRowNameText();
 	if (!RowNameText.IsEmpty())
 	{
-		FPlatformApplicationMisc::ClipboardCopy(*RowNameText);
+		FPropertyEditorClipboard::ClipboardCopy(*RowNameText);
 	}
 }
 
@@ -1282,9 +1578,9 @@ void SDetailSingleItemRow::CreateGlobalExtensionWidgets(TArray<FPropertyRowExten
 	{
 		Args.PropertyHandle = PropertyEditorHelpers::GetPropertyHandle(Customization->GetPropertyNode().ToSharedRef(), nullptr, nullptr);
 	}
-	else if (WidgetRow.GetPropertyHandles().Num() && WidgetRow.GetPropertyHandles()[0] && WidgetRow.GetPropertyHandles()[0]->IsValidHandle())
+	else if (GetPropertyHandles(true).Num() && GetPropertyHandles(true)[0] && GetPropertyHandles(true)[0]->IsValidHandle())
 	{
-		Args.PropertyHandle = WidgetRow.GetPropertyHandles()[0];
+		Args.PropertyHandle = GetPropertyHandles(true)[0];
 	}
 
 	PropertyEditorModule.GetGlobalRowExtensionDelegate().Broadcast(Args, OutExtensions);
diff --git a/Engine/Source/Editor/PropertyEditor/Private/SDetailSingleItemRow.h b/Engine/Source/Editor/PropertyEditor/Private/SDetailSingleItemRow.h
--- a/Engine/Source/Editor/PropertyEditor/Private/SDetailSingleItemRow.h
+++ b/Engine/Source/Editor/PropertyEditor/Private/SDetailSingleItemRow.h
@@ -12,6 +12,7 @@
 #include "DragAndDrop/DecoratedDragDropOp.h"
 #include "Framework/Commands/Commands.h"
 #include "Input/Reply.h"
+#include "PropertyEditorClipboardPrivate.h"
 #include "Widgets/DeclarativeSyntaxSupport.h"
 #include "Widgets/SWidget.h"
 #include "Widgets/Views/STableViewBase.h"
@@ -50,7 +51,6 @@ class SDetailSingleItemRow : public SDetailTableRowBase
 public:
 	SLATE_BEGIN_ARGS( SDetailSingleItemRow )
 		: _ColumnSizeData() {}
-
 		SLATE_ARGUMENT( FDetailColumnSizeData, ColumnSizeData )
 		SLATE_ARGUMENT( bool, AllowFavoriteSystem)
 	SLATE_END_ARGS()
@@ -67,8 +67,10 @@ public:
 	TSharedPtr<FDragDropOperation> CreateDragDropOperation();
 
 	virtual void Tick(const FGeometry& AllottedGeometry, const double InCurrentTime, const float InDeltaTime);
+
 protected:
 	virtual bool OnContextMenuOpening( FMenuBuilder& MenuBuilder ) override;
+	
 private:
 	void OnCopyProperty();
 	void OnCopyPropertyDisplayName();
@@ -77,6 +79,27 @@ private:
 	bool CanCopyPropertyInternalName();
 	void OnPasteProperty();
 	bool CanPasteProperty() const;
+
+	void OnCopyGroup();
+	bool CanCopyGroup() const;
+	void OnPasteGroup();
+	bool CanPasteGroup();
+
+	/**
+	 * @return True if the (optionally tagged) input contents can be pasted
+	 */
+	bool CanPasteFromText(const FString& InTag, const FString& InText) const;
+
+	/**
+	 * Delegate handling pasting an optionally tagged text snippet
+	 */
+	void OnPasteFromText(const FString& InTag, const FString& InText, const TOptional<FGuid>& InOperationId);
+
+	/**
+	 * Handle pasting an optionally tagged text snippet
+	 */
+	bool PasteFromText(const FString& InTag, const FString& InText);
+	
 	FSlateColor GetOuterBackgroundColor() const;
 	FSlateColor GetInnerBackgroundColor() const;
 
@@ -129,6 +152,12 @@ private:
 	FUIAction CopyAction;
 	FUIAction PasteAction;
 
+	/** Previously parsed clipboard data. */
+	UE::PropertyEditor::Internal::FClipboardData PreviousClipboardData;
+
+	/** Delegate handling pasting an optionally tagged text snippet */
+	TSharedPtr<FOnPasteFromText> OnPasteFromTextDelegate;
+
 	/** Animation curve for displaying pulse */
 	FCurveSequence PulseAnimation;
 };
diff --git a/Engine/Source/Editor/PropertyEditor/Private/SDetailTableRowBase.cpp b/Engine/Source/Editor/PropertyEditor/Private/SDetailTableRowBase.cpp
--- a/Engine/Source/Editor/PropertyEditor/Private/SDetailTableRowBase.cpp
+++ b/Engine/Source/Editor/PropertyEditor/Private/SDetailTableRowBase.cpp
@@ -2,6 +2,8 @@
 
 #include "SDetailTableRowBase.h"
 
+#include "PropertyHandleImpl.h"
+
 const float SDetailTableRowBase::ScrollBarPadding = 16.0f;
 
 int32 SDetailTableRowBase::GetIndentLevelForBackgroundColor() const
@@ -34,3 +36,62 @@ bool SDetailTableRowBase::IsScrollBarVisible(TWeakPtr<STableViewBase> OwnerTable
 	}
 	return false;
 }
+
+TArray<TSharedPtr<FPropertyNode>> SDetailTableRowBase::GetPropertyNodes(const bool& bRecursive) const
+{
+	TArray<TSharedPtr<IPropertyHandle>> PropertyHandles = GetPropertyHandles(bRecursive);
+
+	TArray<TSharedPtr<FPropertyNode>> PropertyNodes;
+	PropertyNodes.Reserve(PropertyHandles.Num());
+	
+	for (const TSharedPtr<IPropertyHandle>& PropertyHandle : PropertyHandles)
+	{
+		if (PropertyHandle->IsValidHandle())
+		{
+			PropertyNodes.Add(StaticCastSharedPtr<FPropertyHandleBase>(PropertyHandle)->GetPropertyNode());
+		}
+	}
+
+	return PropertyNodes;
+}
+
+TArray<TSharedPtr<IPropertyHandle>> SDetailTableRowBase::GetPropertyHandles(const bool& bRecursive) const
+{
+	TFunction<bool(const TSharedPtr<IDetailTreeNode>&, TArray<TSharedPtr<IPropertyHandle>>&)> AppendPropertyHandles;
+	AppendPropertyHandles = [&AppendPropertyHandles, bRecursive]
+		(const TSharedPtr<IDetailTreeNode>& InParent, TArray<TSharedPtr<IPropertyHandle>>& OutPropertyHandles)
+	{
+		// Parent in first call is actually parent of this row, so these are the nodes for this row 
+		TArray<TSharedRef<IDetailTreeNode>> ChildNodes;
+		InParent->GetChildren(ChildNodes, true);
+
+		if (ChildNodes.IsEmpty() || !bRecursive)
+		{
+			return false;
+		}
+
+		OutPropertyHandles.Reserve(OutPropertyHandles.Num() + ChildNodes.Num());		
+		for (const TSharedRef<IDetailTreeNode>& ChildNode : ChildNodes)
+		{
+			// @fixme: this won't return when there's multiple properties in a single row, or the row is custom 
+			TSharedPtr<IPropertyHandle> ChildPropertyHandle = ChildNode->CreatePropertyHandle();
+			if (ChildPropertyHandle.IsValid() && ChildPropertyHandle->IsValidHandle())
+			{
+				OutPropertyHandles.Add(ChildPropertyHandle);
+			}
+			AppendPropertyHandles(ChildNode, OutPropertyHandles);
+		}
+
+		return true;
+	};
+
+	TArray<TSharedPtr<IPropertyHandle>> PropertyHandles;	
+	if (const TSharedPtr<FDetailTreeNode> OwnerTreeNodePtr = OwnerTreeNode.Pin())
+	{
+		AppendPropertyHandles(OwnerTreeNodePtr, PropertyHandles);
+	}
+
+	PropertyHandles.Remove(nullptr);
+
+	return PropertyHandles;
+}
diff --git a/Engine/Source/Editor/PropertyEditor/Private/SDetailTableRowBase.h b/Engine/Source/Editor/PropertyEditor/Private/SDetailTableRowBase.h
--- a/Engine/Source/Editor/PropertyEditor/Private/SDetailTableRowBase.h
+++ b/Engine/Source/Editor/PropertyEditor/Private/SDetailTableRowBase.h
@@ -73,6 +73,12 @@ protected:
 	 */
 	virtual bool OnContextMenuOpening( FMenuBuilder& MenuBuilder ) { return false; }
 
+	/** Retrieve all property nodes represented by this row, and it's children if Recursive specified. */
+	TArray<TSharedPtr<FPropertyNode>> GetPropertyNodes(const bool& bRecursive = false) const;
+
+	/** Retrieve all property handles represented by this row, and it's children if Recursive specified. */
+	TArray<TSharedPtr<IPropertyHandle>> GetPropertyHandles(const bool& bRecursive = false) const;
+
 private:
 	void OnExpandAllClicked()
 	{
@@ -95,6 +101,7 @@ private:
 			OwnerTreeNodePin->GetDetailsView()->SetNodeExpansionState( OwnerTreeNodePin.ToSharedRef(), bIsExpanded, bRecursive );
 		}
 	}
+	
 protected:
 	TWeakPtr<FDetailTreeNode> OwnerTreeNode;
 };
diff --git a/Engine/Source/Editor/PropertyEditor/Private/UserInterface/PropertyEditor/PropertyEditorConstants.h b/Engine/Source/Editor/PropertyEditor/Private/UserInterface/PropertyEditor/PropertyEditorConstants.h
--- a/Engine/Source/Editor/PropertyEditor/Private/UserInterface/PropertyEditor/PropertyEditorConstants.h
+++ b/Engine/Source/Editor/PropertyEditor/Private/UserInterface/PropertyEditor/PropertyEditorConstants.h
@@ -8,6 +8,11 @@
 class FName;
 struct FSlateBrush;
 
+namespace UE::PropertyEditor::Private
+{
+	constexpr float PulseAnimationLength = 0.5f;
+}
+
 class PropertyEditorConstants
 {
 public:
diff --git a/Engine/Source/Editor/PropertyEditor/Private/UserInterface/PropertyEditor/SPropertyEditorAsset.cpp b/Engine/Source/Editor/PropertyEditor/Private/UserInterface/PropertyEditor/SPropertyEditorAsset.cpp
--- a/Engine/Source/Editor/PropertyEditor/Private/UserInterface/PropertyEditor/SPropertyEditorAsset.cpp
+++ b/Engine/Source/Editor/PropertyEditor/Private/UserInterface/PropertyEditor/SPropertyEditorAsset.cpp
@@ -1308,14 +1308,32 @@ void SPropertyEditorAsset::OnPaste()
 {
 	FString DestPath;
 	FPlatformApplicationMisc::ClipboardPaste(DestPath);
+	
+	PasteFromText(TEXT(""), DestPath);
+}
+
+void SPropertyEditorAsset::OnPasteFromText(
+	const FString& InTag,
+	const FString& InText,
+	const TOptional<FGuid>& InOperationId)
+{
+	if (CanPasteFromText(InTag, InText))
+	{
+		PasteFromText(InTag, InText);
+	}
+}
 
-	if(DestPath == TEXT("None"))
+void SPropertyEditorAsset::PasteFromText(
+	const FString& InTag,
+	const FString& InText)
+{
+	if(InText == TEXT("None"))
 	{
 		SetValue(nullptr);
 	}
 	else
 	{
-		UObject* Object = LoadObject<UObject>(nullptr, *DestPath);
+		UObject* Object = LoadObject<UObject>(nullptr, *InText);
 		if(Object && Object->IsA(ObjectClass))
 		{
 			// Check against custom asset filter
@@ -1333,7 +1351,19 @@ bool SPropertyEditorAsset::CanPaste()
 	FString ClipboardText;
 	FPlatformApplicationMisc::ClipboardPaste(ClipboardText);
 
-	const FString PossibleObjectPath = FPackageName::ExportTextPathToObjectPath(ClipboardText);
+	return CanPasteFromText(TEXT(""), ClipboardText);
+}
+
+bool SPropertyEditorAsset::CanPasteFromText(
+	const FString& InTag,
+	const FString& InText) const
+{
+	if (!UE::PropertyEditor::TagMatchesProperty(InTag, PropertyHandle))
+	{
+		return false;
+	}
+
+	const FString PossibleObjectPath = FPackageName::ExportTextPathToObjectPath(InText);
 
 	bool bCanPaste = false;
 
diff --git a/Engine/Source/Editor/PropertyEditor/Private/UserInterface/PropertyEditor/SPropertyEditorAsset.h b/Engine/Source/Editor/PropertyEditor/Private/UserInterface/PropertyEditor/SPropertyEditorAsset.h
--- a/Engine/Source/Editor/PropertyEditor/Private/UserInterface/PropertyEditor/SPropertyEditorAsset.h
+++ b/Engine/Source/Editor/PropertyEditor/Private/UserInterface/PropertyEditor/SPropertyEditorAsset.h
@@ -297,6 +297,21 @@ private:
 	 */
 	void OnPaste();
 
+	/**
+	 * @return True if the (optionally tagged) input contents can be pasted
+	 */
+	bool CanPasteFromText(const FString& InTag, const FString& InText) const;
+
+	/**
+	 * Delegate handling pasting an optionally tagged text snippet
+	 */
+	void OnPasteFromText(const FString& InTag, const FString& InText, const TOptional<FGuid>& InOperationId);
+
+	/**
+	 * Handle pasting an optionally tagged text snippet
+	 */
+	void PasteFromText(const FString& InTag, const FString& InText);
+
 	/**
 	 * @return True of the current clipboard contents can be pasted
 	 */
diff --git a/Engine/Source/Editor/PropertyEditor/Private/UserInterface/PropertyEditor/SPropertyMenuAssetPicker.cpp b/Engine/Source/Editor/PropertyEditor/Private/UserInterface/PropertyEditor/SPropertyMenuAssetPicker.cpp
--- a/Engine/Source/Editor/PropertyEditor/Private/UserInterface/PropertyEditor/SPropertyMenuAssetPicker.cpp
+++ b/Engine/Source/Editor/PropertyEditor/Private/UserInterface/PropertyEditor/SPropertyMenuAssetPicker.cpp
@@ -1,20 +1,23 @@
 // Copyright Epic Games, Inc. All Rights Reserved.
 
 #include "UserInterface/PropertyEditor/SPropertyMenuAssetPicker.h"
+
+#include "AssetRegistry/AssetRegistryModule.h"
+#include "AssetToolsModule.h"
+#include "ContentBrowserModule.h"
+#include "Editor.h"
 #include "Factories/Factory.h"
 #include "Framework/MultiBox/MultiBoxBuilder.h"
-#include "Editor.h"
-#include "Modules/ModuleManager.h"
-#include "Widgets/Layout/SBox.h"
-#include "AssetRegistry/AssetRegistryModule.h"
+#include "HAL/PlatformApplicationMisc.h"
 #include "IAssetTools.h"
-#include "Layout/WidgetPath.h"
-#include "AssetToolsModule.h"
 #include "IContentBrowserSingleton.h"
-#include "ContentBrowserModule.h"
-#include "UserInterface/PropertyEditor/PropertyEditorAssetConstants.h"
+#include "Layout/WidgetPath.h"
+#include "Modules/ModuleManager.h"
+#include "PropertyEditorClipboard.h"
+#include "PropertyEditorCopyPastePrivate.h"
 #include "Styling/SlateIconFinder.h"
-#include "HAL/PlatformApplicationMisc.h"
+#include "UserInterface/PropertyEditor/PropertyEditorAssetConstants.h"
+#include "Widgets/Layout/SBox.h"
 
 #define LOCTEXT_NAMESPACE "PropertyEditor"
 
@@ -227,7 +230,7 @@ void SPropertyMenuAssetPicker::OnCopy()
 {
 	if( CurrentObject.IsValid() )
 	{
-		FPlatformApplicationMisc::ClipboardCopy(*CurrentObject.GetExportTextName());
+		FPropertyEditorClipboard::ClipboardCopy(*CurrentObject.GetExportTextName());
 	}
 	OnClose.ExecuteIfBound();
 }
@@ -235,15 +238,32 @@ void SPropertyMenuAssetPicker::OnCopy()
 void SPropertyMenuAssetPicker::OnPaste()
 {
 	FString DestPath;
-	FPlatformApplicationMisc::ClipboardPaste(DestPath);
+	FPropertyEditorClipboard::ClipboardPaste(DestPath);
 
-	if(DestPath == TEXT("None"))
+	PasteFromText(TEXT(""), DestPath);
+}
+
+void SPropertyMenuAssetPicker::OnPasteFromText(
+	const FString& InTag,
+	const FString& InText,
+	const TOptional<FGuid>& InOperationId)
+{
+	// Naive check done elsewhere, guard with proper check here 
+	if (CanPasteFromText(InTag, InText))
 	{
-		SetValue(NULL);
+		PasteFromText(InTag, InText);
+	}
+}
+
+void SPropertyMenuAssetPicker::PasteFromText(const FString& InTag, const FString& InText)
+{
+	if(InText == TEXT("None"))
+	{
+		SetValue(nullptr);
 	}
 	else
 	{
-		UObject* Object = LoadObject<UObject>(NULL, *DestPath);
+		UObject* Object = LoadObject<UObject>(nullptr, *InText);
 		bool PassesAllowedClassesFilter = true;
 		if (Object && AllowedClasses.Num())
 		{
@@ -277,25 +297,39 @@ void SPropertyMenuAssetPicker::OnPaste()
 bool SPropertyMenuAssetPicker::CanPaste()
 {
 	FString ClipboardText;
-	FPlatformApplicationMisc::ClipboardPaste(ClipboardText);
+	FPropertyEditorClipboard::ClipboardPaste(ClipboardText);
+	
+	return CanPasteFromText(TEXT(""), ClipboardText);
+}
+
+bool SPropertyMenuAssetPicker::CanPasteFromText(const FString& InTag, const FString& InText) const
+{
+	if (!bAllowCopyPaste)
+	{
+		return false;
+	}
+
+	if (!UE::PropertyEditor::TagMatchesProperty(InTag, PropertyHandle))
+	{
+		return false;
+	}
 
 	FString Class;
-	FString PossibleObjectPath = ClipboardText;
-	if( ClipboardText.Split( TEXT("'"), &Class, &PossibleObjectPath, ESearchCase::CaseSensitive) )
+	FString PossibleObjectPath = InText;
+	if( InText.Split( TEXT("'"), &Class, &PossibleObjectPath, ESearchCase::CaseSensitive) )
 	{
 		// Remove the last item
 		PossibleObjectPath.LeftChopInline( 1, false );
 	}
 
 	bool bCanPaste = false;
-
 	if( PossibleObjectPath == TEXT("None") )
 	{
 		bCanPaste = true;
 	}
 	else
 	{
-		FAssetRegistryModule& AssetRegistryModule = FModuleManager::Get().LoadModuleChecked<FAssetRegistryModule>(TEXT("AssetRegistry"));
+		const FAssetRegistryModule& AssetRegistryModule = FModuleManager::Get().LoadModuleChecked<FAssetRegistryModule>(TEXT("AssetRegistry"));
 		bCanPaste = PossibleObjectPath.Len() < NAME_SIZE && AssetRegistryModule.Get().GetAssetByObjectPath( FSoftObjectPath(PossibleObjectPath) ).IsValid();
 	}
 
@@ -304,7 +338,7 @@ bool SPropertyMenuAssetPicker::CanPaste()
 
 void SPropertyMenuAssetPicker::OnClear()
 {
-	SetValue(NULL);
+	SetValue(nullptr);
 	OnClose.ExecuteIfBound();
 }
 
diff --git a/Engine/Source/Editor/PropertyEditor/Private/UserInterface/PropertyEditor/SPropertyMenuAssetPicker.h b/Engine/Source/Editor/PropertyEditor/Private/UserInterface/PropertyEditor/SPropertyMenuAssetPicker.h
--- a/Engine/Source/Editor/PropertyEditor/Private/UserInterface/PropertyEditor/SPropertyMenuAssetPicker.h
+++ b/Engine/Source/Editor/PropertyEditor/Private/UserInterface/PropertyEditor/SPropertyMenuAssetPicker.h
@@ -53,6 +53,21 @@ private:
 	 */
 	void OnPaste();
 
+	/**
+	 * @return True if the (optionally tagged) input contents can be pasted
+	 */
+	bool CanPasteFromText(const FString& InTag, const FString& InText) const;
+
+	/**
+	 * Delegate handling pasting an optionally tagged text snippet
+	 */
+	void OnPasteFromText(const FString& InTag, const FString& InText, const TOptional<FGuid>& InOperationId);
+
+	/**
+	 * Handle pasting an optionally tagged text snippet
+	 */
+	void PasteFromText(const FString& InTag, const FString& InText);
+
 	/**
 	 * @return True of the current clipboard contents can be pasted
 	 */
diff --git a/Engine/Source/Editor/PropertyEditor/PropertyEditor.Build.cs b/Engine/Source/Editor/PropertyEditor/PropertyEditor.Build.cs
--- a/Engine/Source/Editor/PropertyEditor/PropertyEditor.Build.cs
+++ b/Engine/Source/Editor/PropertyEditor/PropertyEditor.Build.cs
@@ -8,11 +8,11 @@ public class PropertyEditor : ModuleRules
 	{
 		PublicDependencyModuleNames.AddRange(
 			new string[] {
+				"EditorConfig",
 				"EditorFramework",
 				"UnrealEd",
-                "ActorPickerMode",
-                "SceneDepthPickerMode",
-				"EditorConfig",
+                	"ActorPickerMode",
+                	"SceneDepthPickerMode",
 			}
 		);
 		
@@ -25,8 +25,8 @@ public class PropertyEditor : ModuleRules
 		PrivateIncludePathModuleNames.AddRange(
 			new string[] {
 				"MainFrame",
-                "AssetRegistry",
-                "AssetTools",
+               	"AssetRegistry",
+                	"AssetTools",
 				"ClassViewer",
 				"StructViewer",
 				"ContentBrowser",
@@ -35,7 +35,7 @@ public class PropertyEditor : ModuleRules
 
 		PrivateDependencyModuleNames.AddRange(
 			new string[] {
-                "AppFramework",
+                	"AppFramework",
 				"Core",
 				"CoreUObject",
 				"ApplicationCore",
@@ -46,19 +46,20 @@ public class PropertyEditor : ModuleRules
 				"SlateCore",
 				"EditorWidgets",
 				"Documentation",
-                "RHI",
+                	"RHI",
 				"ConfigEditor",
-                "SceneOutliner",
+                	"SceneOutliner",
 				"DesktopPlatform",
 				"PropertyPath",
 				"ToolWidgets",
+				"Json"
 			}
         );
 
 		DynamicallyLoadedModuleNames.AddRange(
 			new string[] {
-                "AssetRegistry",
-                "AssetTools",
+                	"AssetRegistry",
+                	"AssetTools",
 				"ClassViewer",
 				"StructViewer",
 				"ContentBrowser",
diff --git a/Engine/Source/Editor/PropertyEditor/Public/DetailCategoryBuilder.h b/Engine/Source/Editor/PropertyEditor/Public/DetailCategoryBuilder.h
--- a/Engine/Source/Editor/PropertyEditor/Public/DetailCategoryBuilder.h
+++ b/Engine/Source/Editor/PropertyEditor/Public/DetailCategoryBuilder.h
@@ -9,6 +9,7 @@
 #include "PropertyHandle.h"
 #include "IDetailCustomNodeBuilder.h"
 #include "DetailBuilderTypes.h"
+#include "PropertyEditorCopyPaste.h"
 
 class FDetailWidgetRow;
 class IDetailGroup;
@@ -211,4 +212,7 @@ public:
 
 	/** Add a property, but force it to behave as a normal, peer reference regardless of CPF_InstancedReference */
 	virtual void AddPropertyDisableInstancedReference(TSharedPtr<IPropertyHandle> PropertyHandle) = 0;
+
+	/** Optional PasteFromText delegate for this category */
+	virtual TSharedPtr<FOnPasteFromText> OnPasteFromText() const { return nullptr; }
 };
diff --git a/Engine/Source/Editor/PropertyEditor/Public/DetailTreeNode.h b/Engine/Source/Editor/PropertyEditor/Public/DetailTreeNode.h
--- a/Engine/Source/Editor/PropertyEditor/Public/DetailTreeNode.h
+++ b/Engine/Source/Editor/PropertyEditor/Public/DetailTreeNode.h
@@ -38,7 +38,7 @@ public:
 
 	/** IDetailTreeNode interface */
 	virtual FNodeWidgets CreateNodeWidgets() const;
-	virtual void GetChildren(TArray<TSharedRef<IDetailTreeNode>>& OutChildren);
+	virtual void GetChildren(TArray<TSharedRef<IDetailTreeNode>>& OutChildren, const bool& bInIgnoreVisibility = false);
 	virtual TSharedPtr<IDetailPropertyRow> GetRow() const override { return nullptr; }
 	virtual void GetFilterStrings(TArray<FString>& OutFilterStrings) const override { };
 	virtual bool GetInitiallyCollapsed() const override { return false; }
@@ -64,9 +64,10 @@ public:
 	/**
 	 * Gets child tree nodes
 	 *
-	 * @param OutChildren	The array to add children to
+	 * @param OutChildren			The array to add children to
+	 * @param bInIgnoreVisibility	Ignore the child node visibility
 	 */
-	virtual void GetChildren(FDetailNodeList& OutChildren) = 0;
+	virtual void GetChildren(FDetailNodeList& OutChildren, const bool& bInIgnoreVisibility = false) = 0;
 
 	/**
 	 * @ The parent DetailTreeNode of this node, or nullptr if this is a root node (or not yet added to a tree)
diff --git a/Engine/Source/Editor/PropertyEditor/Public/DetailWidgetRow.h b/Engine/Source/Editor/PropertyEditor/Public/DetailWidgetRow.h
--- a/Engine/Source/Editor/PropertyEditor/Public/DetailWidgetRow.h
+++ b/Engine/Source/Editor/PropertyEditor/Public/DetailWidgetRow.h
@@ -12,6 +12,7 @@
 #include "Widgets/SWidget.h"
 #include "Widgets/Layout/SSpacer.h"
 #include "DetailCategoryBuilder.h"
+#include "PropertyEditorCopyPaste.h"
 
 class FDetailWidgetRow;
 class FResetToDefaultOverride;
@@ -119,6 +120,7 @@ public:
 		, RowTagName()
 	{
 	}
+	
 	FDetailWidgetRow& operator=(const FDetailWidgetRow& Other)
 	{
 		NameWidget = FDetailWidgetDecl(*this, Other.NameWidget);
@@ -288,6 +290,12 @@ public:
 		return CopyMenuAction.ExecuteAction.IsBound() && PasteMenuAction.ExecuteAction.IsBound();
 	}
 
+	/** @return true if a custom paste text is bound on this row */
+	bool IsPasteFromTextBound() const
+	{
+		return OnPasteFromTextDelegate.IsValid() && OnPasteFromTextDelegate.Pin()->IsBoundToObject(this);
+	}
+
 	/**
 	* Sets a tag which can be used to identify this row 
 	*/
@@ -367,10 +375,12 @@ public:
 	TAttribute<bool> IsValueEnabledAttr;
 	/** String to filter with */
 	FText FilterTextString;
-	/** Action for coping data on this row */
+	/** Action for copying data on this row */
 	FUIAction CopyMenuAction;
 	/** Action for pasting data on this row */
-	FUIAction PasteMenuAction;
+	FUIAction PasteMenuAction;	
+	/** Delegate for pasting (optionally tagged) text on this row */
+	TWeakPtr<FOnPasteFromText> OnPasteFromTextDelegate;
 
 	struct FCustomMenuData
 	{
@@ -403,4 +413,3 @@ public:
 	/** True to force auto-expansion */
 	TOptional<bool> ForceAutoExpansion;
 };
-
diff --git a/Engine/Source/Editor/PropertyEditor/Public/IDetailGroup.h b/Engine/Source/Editor/PropertyEditor/Public/IDetailGroup.h
--- a/Engine/Source/Editor/PropertyEditor/Public/IDetailGroup.h
+++ b/Engine/Source/Editor/PropertyEditor/Public/IDetailGroup.h
@@ -25,8 +25,7 @@ enum class EDetailGroupDisplayMode
 class IDetailGroup : public IDetailLayoutRow
 {
 public:
-
-	virtual ~IDetailGroup(){}
+	virtual ~IDetailGroup() = default;
 
 	/**
 	 * Makes a custom row for the groups header
@@ -102,4 +101,7 @@ public:
 
 	/** Specify whether this group should be formatted to display like a subcategory */
 	virtual void SetDisplayMode(EDetailGroupDisplayMode Mode) = 0;
+
+	/** Return an optional PasteFromText delegate */
+	virtual TSharedPtr<FOnPasteFromText> OnPasteFromText() const { return nullptr; }
 };
diff --git a/Engine/Source/Editor/PropertyEditor/Public/IDetailTreeNode.h b/Engine/Source/Editor/PropertyEditor/Public/IDetailTreeNode.h
--- a/Engine/Source/Editor/PropertyEditor/Public/IDetailTreeNode.h
+++ b/Engine/Source/Editor/PropertyEditor/Public/IDetailTreeNode.h
@@ -146,11 +146,12 @@ public:
 
 	/**
 	 * Gets the children of this tree node    
-	 * Note: Customizations can determine the visibility of children.  This will only return visible children
+	 * Note: Customizations can determine the visibility of children.  This will only return visible children if bInIgnoreVisibility if false
 	 *
-	 * @param OutChildren	The generated children
+	 * @param OutChildren			The generated children
+	 * @param bInIgnoreVisibility	Whether to return children regardless of visibility
 	 */
-	virtual void GetChildren(TArray<TSharedRef<IDetailTreeNode>>& OutChildren) = 0;
+	virtual void GetChildren(TArray<TSharedRef<IDetailTreeNode>>& OutChildren, const bool& bInIgnoreVisibility = false) = 0;
 
 	/**
 	 * Gets an identifier name for this node.  This is not a name formatted for display purposes, but can be useful for storing
diff --git a/Engine/Source/Editor/PropertyEditor/Public/MaterialList.h b/Engine/Source/Editor/PropertyEditor/Public/MaterialList.h
--- a/Engine/Source/Editor/PropertyEditor/Public/MaterialList.h
+++ b/Engine/Source/Editor/PropertyEditor/Public/MaterialList.h
@@ -67,6 +67,8 @@ struct FMaterialListDelegates
 	FOnCanCopyMaterialList OnCanCopyMaterialList;
 	/** Delegate called Pasting a material list */
 	FOnPasteMaterialList OnPasteMaterialList;
+	/** Delegate called Pasting an optionally tagged text snippet */
+	TWeakPtr<FOnPasteFromText> OnPasteFromText;
 
 	/** Delegate called Copying a material item */
 	FOnCopyMaterialItem OnCopyMaterialItem;
@@ -338,6 +340,9 @@ private:
 	void OnCopyMaterialList();
 	void OnPasteMaterialList();
 
+	/** Handle pasting an optionally tagged text snippet */
+	void OnPasteMaterialListFromText(const FString& InTag, const FString& InText, const TOptional<FGuid>& InOperationId);
+
 	bool OnCanCopyMaterialItem(int32 CurrentSlot) const;
 	void OnCopyMaterialItem(int32 CurrentSlot);
 	void OnPasteMaterialItem(int32 CurrentSlot);
@@ -370,4 +375,3 @@ private:
 	/** The mesh asset that owns these materials */
 	TArray<FAssetData> OwnerAssetDataArray;
 };
-
diff --git a/Engine/Source/Editor/PropertyEditor/Public/PropertyEditorClipboard.h b/Engine/Source/Editor/PropertyEditor/Public/PropertyEditorClipboard.h
new file mode 100644
--- /dev/null
+++ b/Engine/Source/Editor/PropertyEditor/Public/PropertyEditorClipboard.h
@@ -0,0 +1,46 @@
+﻿// Copyright Epic Games, Inc. All Rights Reserved.
+
+#pragma once
+
+#include "Containers/Map.h"
+#include "Containers/UnrealString.h"
+#include "CoreTypes.h"
+#include "Templates/Function.h"
+#include "Templates/SharedPointerFwd.h"
+#include "UObject/NameTypes.h"
+
+/**
+ * Specializes clipboard handling to allow tagged entries,
+ * where you can partition the clipboard contents and look them up by tag (name).
+ * The clipboard is persisted using both the Platform implementation, and as a stored virtual clipboard.
+ */
+class PROPERTYEDITOR_API FPropertyEditorClipboard
+{
+public:
+	/** Perform a copy. */
+	static void ClipboardCopy(const TCHAR* Str);
+
+	/** Perform a copy, optionally specifying the tag. */
+	static void ClipboardCopy(const TCHAR* Str, FName Tag);
+
+	/** Perform a copy, supplying a function that maps the result to multiple tags. */
+	static void ClipboardCopy(TUniqueFunction<void(TMap<FName, FString>&)>&& TagMappingFunc);
+
+	/** Append the given values to the current clipboard entry. If none exists, it will create it. */
+	static void ClipboardAppend(TUniqueFunction<void(TMap<FName, FString>&)>&& TagMappingFunc);
+
+	/** Gets the clipboard contents. */
+	static bool ClipboardPaste(FString& Dest);
+	
+	/**
+	 * Will return true if the tag is specified and exists. 
+	 * If the tag doesn't exist, Dest will be populated with the untagged clipboard contents.  
+	 */
+	static bool ClipboardPaste(FString& Dest, FName Tag);
+
+private:
+	static TSharedRef<FPropertyEditorClipboard> Get();
+
+	/** Item per-tag, cleared when a new copy operation occurs. */
+	TMap<FName, FString> ClipboardContents;
+};
diff --git a/Engine/Source/Editor/PropertyEditor/Public/PropertyEditorCopyPaste.h b/Engine/Source/Editor/PropertyEditor/Public/PropertyEditorCopyPaste.h
new file mode 100644
--- /dev/null
+++ b/Engine/Source/Editor/PropertyEditor/Public/PropertyEditorCopyPaste.h
@@ -0,0 +1,20 @@
+﻿// Copyright Epic Games, Inc. All Rights Reserved.
+
+#pragma once
+
+#include "CoreTypes.h"
+#include "PropertyHandle.h"
+
+/** Utilities related to copy/paste. */
+namespace UE::PropertyEditor
+{
+	/** Gets the property path from the provided PropertyHandle (if valid). */
+	[[nodiscard]] PROPERTYEDITOR_API FString GetPropertyPath(const TSharedPtr<IPropertyHandle>& InPropertyHandle);
+
+	/** Returns true if the Tag is empty, or the specified property matches it. */
+	[[nodiscard]] PROPERTYEDITOR_API bool TagMatchesProperty(const FString& InTag, const TSharedPtr<IPropertyHandle>& InPropertyHandle);
+}
+
+/** Subscribe to and parse the string contents. Note that the caller is responsible for providing valid UE text,
+ * and the subscriber is responsible for applying it if applicable.  */
+typedef TMulticastDelegate<void(const FString& InTag, const FString& InText, const TOptional<FGuid>& InOperationId)> FOnPasteFromText;
