PropertyEditor reverted erroneous PropertyHandfle evaluation change

diff --git a/Engine/Source/Editor/PropertyEditor/Private/SDetailSingleItemRow.cpp b/Engine/Source/Editor/PropertyEditor/Private/SDetailSingleItemRow.cpp
--- a/Engine/Source/Editor/PropertyEditor/Private/SDetailSingleItemRow.cpp
+++ b/Engine/Source/Editor/PropertyEditor/Private/SDetailSingleItemRow.cpp
@@ -1134,6 +1134,17 @@ bool SDetailSingleItemRow::OnContextMenuOpening(FMenuBuilder& MenuBuilder)
 	return true;
 }
 
+TArray<TSharedPtr<IPropertyHandle>> SDetailSingleItemRow::GetPropertyHandles(const bool& bRecursive) const
+{
+	if (TArray<TSharedPtr<IPropertyHandle>> PropertyHandles = SDetailTableRowBase::GetPropertyHandles(bRecursive);
+		!PropertyHandles.IsEmpty())
+	{
+		return PropertyHandles;
+	}
+
+	return WidgetRow.GetPropertyHandles();
+}
+
 void SDetailSingleItemRow::OnCopyProperty()
 {
 	if (OwnerTreeNode.IsValid())
diff --git a/Engine/Source/Editor/PropertyEditor/Private/SDetailSingleItemRow.h b/Engine/Source/Editor/PropertyEditor/Private/SDetailSingleItemRow.h
--- a/Engine/Source/Editor/PropertyEditor/Private/SDetailSingleItemRow.h
+++ b/Engine/Source/Editor/PropertyEditor/Private/SDetailSingleItemRow.h
@@ -70,6 +70,7 @@ public:
 
 protected:
 	virtual bool OnContextMenuOpening( FMenuBuilder& MenuBuilder ) override;
+	virtual TArray<TSharedPtr<IPropertyHandle>> GetPropertyHandles(const bool& bRecursive = false) const override;
 	
 private:
 	void OnCopyProperty();
diff --git a/Engine/Source/Editor/PropertyEditor/Private/SDetailTableRowBase.h b/Engine/Source/Editor/PropertyEditor/Private/SDetailTableRowBase.h
--- a/Engine/Source/Editor/PropertyEditor/Private/SDetailTableRowBase.h
+++ b/Engine/Source/Editor/PropertyEditor/Private/SDetailTableRowBase.h
@@ -77,7 +77,7 @@ protected:
 	TArray<TSharedPtr<FPropertyNode>> GetPropertyNodes(const bool& bRecursive = false) const;
 
 	/** Retrieve all property handles represented by this row, and it's children if Recursive specified. */
-	TArray<TSharedPtr<IPropertyHandle>> GetPropertyHandles(const bool& bRecursive = false) const;
+	virtual TArray<TSharedPtr<IPropertyHandle>> GetPropertyHandles(const bool& bRecursive = false) const;
 
 private:
 	void OnExpandAllClicked()
