Property Matrix: Change text color to white on cell selection and exclude the property matrix from the "restore asset tabs" prompt

diff --git a/Engine/Source/Editor/PropertyEditor/Private/PropertyEditorModule.cpp b/Engine/Source/Editor/PropertyEditor/Private/PropertyEditorModule.cpp
--- a/Engine/Source/Editor/PropertyEditor/Private/PropertyEditorModule.cpp
+++ b/Engine/Source/Editor/PropertyEditor/Private/PropertyEditorModule.cpp
@@ -360,7 +360,7 @@ TSharedRef< IPropertyTableWidgetHandle > FPropertyEditorModule::CreatePropertyTa
 }
 
 TSharedRef< IPropertyTableCellPresenter > FPropertyEditorModule::CreateTextPropertyCellPresenter(const TSharedRef< class FPropertyNode >& InPropertyNode, const TSharedRef< class IPropertyTableUtilities >& InPropertyUtilities, 
-																								 const FSlateFontInfo* InFontPtr /* = NULL */)
+																								 const FSlateFontInfo* InFontPtr /* = NULL */ , const TSharedPtr< IPropertyTableCell >& InCell /* = nullptr */)
 {
 	FSlateFontInfo InFont;
 
@@ -375,7 +375,7 @@ TSharedRef< IPropertyTableCellPresenter > FPropertyEditorModule::CreateTextPrope
 	}
 
 	TSharedRef< FPropertyEditor > PropertyEditor = FPropertyEditor::Create( InPropertyNode, InPropertyUtilities );
-	return MakeShareable( new FTextPropertyTableCellPresenter( PropertyEditor, InPropertyUtilities, InFont) );
+	return MakeShareable( new FTextPropertyTableCellPresenter( PropertyEditor, InPropertyUtilities, InFont, InCell) );
 }
 
 FStructProperty* FPropertyEditorModule::RegisterStructOnScopeProperty(TSharedRef<FStructOnScope> StructOnScope)
diff --git a/Engine/Source/Editor/PropertyEditor/Private/UserInterface/PropertyTable/ObjectNameTableCellPresenter.h b/Engine/Source/Editor/PropertyEditor/Private/UserInterface/PropertyTable/ObjectNameTableCellPresenter.h
--- a/Engine/Source/Editor/PropertyEditor/Private/UserInterface/PropertyTable/ObjectNameTableCellPresenter.h
+++ b/Engine/Source/Editor/PropertyEditor/Private/UserInterface/PropertyTable/ObjectNameTableCellPresenter.h
@@ -95,6 +95,7 @@ private:
 					SNew( STextBlock )
 					.Font( FAppStyle::GetFontStyle( TextFontStyle ) )
 					.Text( FText::FromString(DisplayNamePieces[ Index ]) )
+					.ColorAndOpacity(this, &FObjectNameTableCellPresenter::GetTextForegroundColor)
 				];
 
 			if ( Index < DisplayNamePieces.Num() - 1 )
@@ -114,6 +115,15 @@ private:
 		return NameBox;
 	}
 
+	FSlateColor GetTextForegroundColor() const
+	{
+		if(Cell->GetTable()->GetSelectedCells().Contains( Cell ))
+		{
+			return FStyleColors::White;
+		}
+
+		return FSlateColor::UseForeground();
+	}
 
 private:
 
diff --git a/Engine/Source/Editor/PropertyEditor/Private/UserInterface/PropertyTable/SObjectColumnHeader.h b/Engine/Source/Editor/PropertyEditor/Private/UserInterface/PropertyTable/SObjectColumnHeader.h
--- a/Engine/Source/Editor/PropertyEditor/Private/UserInterface/PropertyTable/SObjectColumnHeader.h
+++ b/Engine/Source/Editor/PropertyEditor/Private/UserInterface/PropertyTable/SObjectColumnHeader.h
@@ -64,7 +64,7 @@ class SObjectColumnHeader : public SColumnHeader
 			}
 			else
 			{
-				CellPresenter = MakeShareable( new FTextPropertyTableCellPresenter( PropertyEditor, Utilities.ToSharedRef() ) );
+				CellPresenter = MakeShareable( new FTextPropertyTableCellPresenter( PropertyEditor, Utilities.ToSharedRef(), FAppStyle::GetFontStyle( PropertyTableConstants::NormalFontStyle ), Cell ) );
 			}
 		}
 
diff --git a/Engine/Source/Editor/PropertyEditor/Private/UserInterface/PropertyTable/STextColumnHeader.h b/Engine/Source/Editor/PropertyEditor/Private/UserInterface/PropertyTable/STextColumnHeader.h
--- a/Engine/Source/Editor/PropertyEditor/Private/UserInterface/PropertyTable/STextColumnHeader.h
+++ b/Engine/Source/Editor/PropertyEditor/Private/UserInterface/PropertyTable/STextColumnHeader.h
@@ -52,7 +52,7 @@ public:
 		if ( !CellPresenter.IsValid() && Cell->IsBound() )
 		{
 			TSharedRef< FPropertyEditor > PropertyEditor = FPropertyEditor::Create( Cell->GetNode().ToSharedRef(), Utilities.ToSharedRef() );
-			CellPresenter = MakeShareable( new FTextPropertyTableCellPresenter( PropertyEditor, Utilities.ToSharedRef() ) );
+			CellPresenter = MakeShareable( new FTextPropertyTableCellPresenter( PropertyEditor, Utilities.ToSharedRef(), FAppStyle::GetFontStyle( PropertyTableConstants::NormalFontStyle ), Cell ) );
 		}
 
 		return SNew( SPropertyTableCell, Cell )
diff --git a/Engine/Source/Editor/PropertyEditor/Private/UserInterface/PropertyTable/TextPropertyTableCellPresenter.cpp b/Engine/Source/Editor/PropertyEditor/Private/UserInterface/PropertyTable/TextPropertyTableCellPresenter.cpp
--- a/Engine/Source/Editor/PropertyEditor/Private/UserInterface/PropertyTable/TextPropertyTableCellPresenter.cpp
+++ b/Engine/Source/Editor/PropertyEditor/Private/UserInterface/PropertyTable/TextPropertyTableCellPresenter.cpp
@@ -6,9 +6,12 @@
 #include "Widgets/SBoxPanel.h"
 #include "Widgets/Text/STextBlock.h"
 #include "Widgets/Layout/SBorder.h"
+#include "IPropertyTable.h"
+#include "IPropertyTableCell.h"
 #include "IPropertyTableUtilities.h"
 #include "PropertyEditorHelpers.h"
 #include "SResetToDefaultPropertyEditor.h"
+#include "Styling/StyleColors.h"
 
 #include "UserInterface/PropertyEditor/SPropertyEditor.h"
 #include "UserInterface/PropertyEditor/SPropertyEditorNumeric.h"
@@ -19,11 +22,11 @@
 #include "UserInterface/PropertyEditor/SPropertyEditorColor.h"
 #include "UserInterface/PropertyEditor/SPropertyEditorDateTime.h"
 
-
-FTextPropertyTableCellPresenter::FTextPropertyTableCellPresenter( const TSharedRef< class FPropertyEditor >& InPropertyEditor, const TSharedRef< class IPropertyTableUtilities >& InPropertyUtilities, FSlateFontInfo InFont /*= FAppStyle::GetFontStyle( PropertyTableConstants::NormalFontStyle ) */ )
+FTextPropertyTableCellPresenter::FTextPropertyTableCellPresenter( const TSharedRef< class FPropertyEditor >& InPropertyEditor, const TSharedRef< class IPropertyTableUtilities >& InPropertyUtilities, FSlateFontInfo InFont /*= FAppStyle::GetFontStyle( PropertyTableConstants::NormalFontStyle ) */, const TSharedPtr< IPropertyTableCell >& InCell /*= nullptr*/ )
 	: PropertyWidget( NULL )
 	, PropertyEditor( InPropertyEditor )
 	, PropertyUtilities( InPropertyUtilities )
+	, Cell(InCell)
 	, HasReadOnlyEditingWidget( false )
 	, Font(InFont)
 {
@@ -40,6 +43,7 @@ TSharedRef< class SWidget > FTextPropertyTableCellPresenter::ConstructDisplayWid
 		[
 			SNew( STextBlock )
 			.Text( PropertyEditor->GetValueAsText() )
+			.ColorAndOpacity(this, &FTextPropertyTableCellPresenter::GetTextForegroundColor)
 			.ToolTipText( PropertyEditor->GetToolTipText() )
 			.Font( Font )
 		];
@@ -236,3 +240,13 @@ FText FTextPropertyTableCellPresenter::GetValueAsText()
 {
 	return PropertyEditor->GetValueAsText();
 }
+
+FSlateColor FTextPropertyTableCellPresenter::GetTextForegroundColor() const
+{
+	if(Cell && Cell->GetTable()->GetSelectedCells().Contains( Cell.ToSharedRef() ))
+	{
+		return FStyleColors::White;
+	}
+
+	return FSlateColor::UseForeground();
+}
diff --git a/Engine/Source/Editor/PropertyEditor/Private/UserInterface/PropertyTable/TextPropertyTableCellPresenter.h b/Engine/Source/Editor/PropertyEditor/Private/UserInterface/PropertyTable/TextPropertyTableCellPresenter.h
--- a/Engine/Source/Editor/PropertyEditor/Private/UserInterface/PropertyTable/TextPropertyTableCellPresenter.h
+++ b/Engine/Source/Editor/PropertyEditor/Private/UserInterface/PropertyTable/TextPropertyTableCellPresenter.h
@@ -8,11 +8,13 @@
 #include "IPropertyTableCellPresenter.h"
 #include "UserInterface/PropertyTable/PropertyTableConstants.h"
 
+class IPropertyTableCell;
+
 class FTextPropertyTableCellPresenter : public TSharedFromThis< FTextPropertyTableCellPresenter >, public IPropertyTableCellPresenter
 {
 public:
 
-	FTextPropertyTableCellPresenter( const TSharedRef< class FPropertyEditor >& InPropertyEditor, const TSharedRef< class IPropertyTableUtilities >& InPropertyUtilities, FSlateFontInfo InFont = FAppStyle::GetFontStyle( PropertyTableConstants::NormalFontStyle ) );
+	FTextPropertyTableCellPresenter( const TSharedRef< class FPropertyEditor >& InPropertyEditor, const TSharedRef< class IPropertyTableUtilities >& InPropertyUtilities, FSlateFontInfo InFont = FAppStyle::GetFontStyle( PropertyTableConstants::NormalFontStyle ), const TSharedPtr< IPropertyTableCell >& InCell = nullptr );
 
 	virtual ~FTextPropertyTableCellPresenter() {}
 
@@ -39,6 +41,7 @@ private:
 
 	bool CalculateIfUsingReadOnlyEditingWidget() const;
 
+	FSlateColor GetTextForegroundColor() const;
 
 private:
 
@@ -46,6 +49,8 @@ private:
 
 	TSharedRef< class FPropertyEditor > PropertyEditor;
 	TSharedRef< class IPropertyTableUtilities > PropertyUtilities;
+	
+	TSharedPtr< IPropertyTableCell > Cell;
 
 	bool HasReadOnlyEditingWidget;
 	FSlateFontInfo Font;
diff --git a/Engine/Source/Editor/PropertyEditor/Public/PropertyEditorModule.h b/Engine/Source/Editor/PropertyEditor/Public/PropertyEditorModule.h
--- a/Engine/Source/Editor/PropertyEditor/Public/PropertyEditorModule.h
+++ b/Engine/Source/Editor/PropertyEditor/Public/PropertyEditorModule.h
@@ -61,6 +61,7 @@ namespace FPropertyAccess
 
 
 class IPropertyHandle;
+class IPropertyTableCell;
 class SPropertyTreeViewImpl;
 class SWindow;
 class IPropertyTableCellPresenter;
@@ -383,7 +384,7 @@ public:
 	virtual TSharedRef< class IPropertyTableWidgetHandle > CreatePropertyTableWidgetHandle( const TSharedRef< IPropertyTable >& PropertyTable, const TArray< TSharedRef< class IPropertyTableCustomColumn > >& Customizations );
 
 	virtual TSharedRef< IPropertyTableCellPresenter > CreateTextPropertyCellPresenter( const TSharedRef< class FPropertyNode >& InPropertyNode, const TSharedRef< class IPropertyTableUtilities >& InPropertyUtilities, 
-		const FSlateFontInfo* InFontPtr = NULL);
+		const FSlateFontInfo* InFontPtr = NULL, const TSharedPtr< IPropertyTableCell >& InCell = nullptr);
 
 	/**
 	 * Register a floating struct on scope so that the details panel may use it as a property
