Property Editor: Fix non unity error and missing file

diff --git a/Engine/Source/Editor/PropertyEditor/Private/PropertyEditorToolkit.h b/Engine/Source/Editor/PropertyEditor/Private/PropertyEditorToolkit.h
--- a/Engine/Source/Editor/PropertyEditor/Private/PropertyEditorToolkit.h
+++ b/Engine/Source/Editor/PropertyEditor/Private/PropertyEditorToolkit.h
@@ -38,6 +38,8 @@ public:
 
 	virtual FLinearColor GetWorldCentricTabColorScale() const override;
 
+	virtual bool IncludeAssetInRestoreOpenAssetsPrompt() const override { return false; }
+
 	bool IsExposedAsColumn( const TWeakPtr< IPropertyTreeRow >& Row ) const;
 
 	void ToggleColumnForProperty( const TSharedPtr< class FPropertyPath >& PropertyPath );
diff --git a/Engine/Source/Editor/PropertyEditor/Private/UserInterface/PropertyTable/ObjectNameTableCellPresenter.h b/Engine/Source/Editor/PropertyEditor/Private/UserInterface/PropertyTable/ObjectNameTableCellPresenter.h
--- a/Engine/Source/Editor/PropertyEditor/Private/UserInterface/PropertyTable/ObjectNameTableCellPresenter.h
+++ b/Engine/Source/Editor/PropertyEditor/Private/UserInterface/PropertyTable/ObjectNameTableCellPresenter.h
@@ -12,6 +12,7 @@
 #include "Widgets/Text/STextBlock.h"
 #include "Widgets/SToolTip.h"
 #include "Styling/AppStyle.h"
+#include "Styling/StyleColors.h"
 #include "IPropertyTableCell.h"
 #include "Widgets/Layout/SBox.h"
 #include "IPropertyTableCellPresenter.h"
