Add proper caching support for AllowListAllSubclassProperties when intermediate classes don't have explicit entries in the list

diff --git a/Engine/Source/Editor/PropertyEditor/Private/PropertyPermissionList.cpp b/Engine/Source/Editor/PropertyEditor/Private/PropertyPermissionList.cpp
--- a/Engine/Source/Editor/PropertyEditor/Private/PropertyPermissionList.cpp
+++ b/Engine/Source/Editor/PropertyEditor/Private/PropertyPermissionList.cpp
@@ -222,10 +222,10 @@ bool FPropertyPermissionList::DoesPropertyPassFilter(const UStruct* ObjectStruct
 
 const FNamePermissionList& FPropertyPermissionList::GetCachedPermissionListForStruct(const UStruct* Struct) const
 {
-	const FNamePermissionList* CachedPermissionList = CachedPropertyPermissionList.Find(Struct);
+	const FPropertyPermissionListEntry* CachedPermissionList = CachedPropertyPermissionList.Find(Struct);
 	if (CachedPermissionList)
 	{
-		return *CachedPermissionList;
+		return CachedPermissionList->PermissionList;
 	}
 
 	// Default value doesn't matter since it's a no-op until the first PermissionList is encountered, at which
@@ -244,14 +244,19 @@ const FNamePermissionList& FPropertyPermissionList::GetCachedPermissionListForSt
 
 	// Normally this case would be caught in GetCachedPermissionListForStruct, but when being called recursively from a subclass
 	// we still need to update bInOutShouldAllowListAllProperties so that new PermissionLists cache properly.
-	FNamePermissionList* PermissionList = CachedPropertyPermissionList.Find(Struct);
-	if (PermissionList)
+	FPropertyPermissionListEntry* PermissionListEntry = CachedPropertyPermissionList.Find(Struct);
+	if (PermissionListEntry)
 	{
 		// Same check as below, but it specifically has to come after the recursive call so we have to duplicate the code here
 		if (Entry && Entry->Rules == EPropertyPermissionListRules::AllowListAllProperties)
 		{
 			bInOutShouldAllowListAllProperties = true;
 		}
+		// This is to properly support propagation of AllowListAllSubclassProperties for classes that do not have a raw entry
+		if (PermissionListEntry->Rules == EPropertyPermissionListRules::AllowListAllProperties)
+		{
+			bInOutShouldAllowListAllProperties = true;
+		}
 	}
 	else
 	{
@@ -302,7 +307,10 @@ const FNamePermissionList& FPropertyPermissionList::GetCachedPermissionListForSt
 			}
 		}
 
-		PermissionList = &CachedPropertyPermissionList.Add(Struct, MoveTemp(NewPermissionList));
+		PermissionListEntry = &CachedPropertyPermissionList.Add(Struct, { MoveTemp(NewPermissionList),
+			// propagate the allow list all properties setting, this is to make AllowListAllSubclassProperties work properly 
+			// when a sub class was already previously cached from a base class that has AllowListAllSubclassProperties set on it and another class is then building its list off of the cached subclass list
+			bInOutShouldAllowListAllProperties ? EPropertyPermissionListRules::AllowListAllProperties : EPropertyPermissionListRules::UseExistingPermissionList });
 	}
 	
 	// If this Struct has no PermissionList, then the ShouldPermissionListAllProperties rule just forwards its current value on to the next subclass.
@@ -313,7 +321,7 @@ const FNamePermissionList& FPropertyPermissionList::GetCachedPermissionListForSt
 		bInOutShouldAllowListAllProperties = Entry->Rules == EPropertyPermissionListRules::AllowListAllSubclassProperties;
 	}
 
-	return *PermissionList;
+	return PermissionListEntry->PermissionList;
 }
 
 bool FPropertyPermissionList::IsSpecificPropertyAllowListed(const UStruct* ObjectStruct, FName PropertyName) const
diff --git a/Engine/Source/Editor/PropertyEditor/Public/PropertyPermissionList.h b/Engine/Source/Editor/PropertyEditor/Public/PropertyPermissionList.h
--- a/Engine/Source/Editor/PropertyEditor/Public/PropertyPermissionList.h
+++ b/Engine/Source/Editor/PropertyEditor/Public/PropertyPermissionList.h
@@ -133,7 +133,7 @@ private:
 	friend uint32 GetTypeHash(const FPermissionListUpdate& PermisisonList);
 	TSet< FPermissionListUpdate > PendingUpdates;
 	/** Lazily-constructed combined cache of both the flattened class PermissionList and struct PermissionList */
-	mutable TMap<TWeakObjectPtr<const UStruct>, FNamePermissionList> CachedPropertyPermissionList;
+	mutable TMap<TWeakObjectPtr<const UStruct>, FPropertyPermissionListEntry> CachedPropertyPermissionList;
 
 	/** Get or create the cached PermissionList for a specific UStruct */
 	const FNamePermissionList& GetCachedPermissionListForStruct(const UStruct* Struct) const;
