Memory leak in Array draghandles in PropertyEditorHelpers.cpp

diff --git a/Engine/Source/Editor/PropertyEditor/Private/PropertyEditorHelpers.cpp b/Engine/Source/Editor/PropertyEditor/Private/PropertyEditorHelpers.cpp
--- a/Engine/Source/Editor/PropertyEditor/Private/PropertyEditorHelpers.cpp
+++ b/Engine/Source/Editor/PropertyEditor/Private/PropertyEditorHelpers.cpp
@@ -833,6 +833,7 @@ namespace PropertyEditorHelpers
 
 	TSharedRef<SWidget> MakePropertyReorderHandle(TSharedPtr<SDetailSingleItemRow> InParentRow, TAttribute<bool> InEnabledAttr)
 	{
+		const TWeakPtr<SDetailSingleItemRow> RowPtr = InParentRow.ToWeakPtr();
 		TSharedRef<SArrayRowHandle> Handle = SNew(SArrayRowHandle)
 			.Content()
 			[
@@ -849,7 +850,14 @@ namespace PropertyEditorHelpers
 		.ParentRow(InParentRow)
 		.Cursor(EMouseCursor::GrabHand)
 		.IsEnabled(InEnabledAttr)
-		.Visibility_Lambda([InParentRow]() { return InParentRow->IsHovered() ? EVisibility::Visible : EVisibility::Hidden; });
+		.Visibility_Lambda([RowPtr]()
+		{
+			if( const TSharedPtr<SDetailSingleItemRow> Row = RowPtr.Pin())
+			{
+				return Row->IsHovered() ? EVisibility::Visible : EVisibility::Hidden;
+			}
+			return EVisibility::Hidden;
+		});
 		return Handle;
 	}
 
