DetailsView: Update expansion state twice, once for initial nodes, and another time for nodes added by property customization

diff --git a/Engine/Source/Editor/PropertyEditor/Private/SDetailsView.cpp b/Engine/Source/Editor/PropertyEditor/Private/SDetailsView.cpp
--- a/Engine/Source/Editor/PropertyEditor/Private/SDetailsView.cpp
+++ b/Engine/Source/Editor/PropertyEditor/Private/SDetailsView.cpp
@@ -1000,9 +1000,15 @@ void SDetailsView::PostSetObject(const TArray<FDetailsViewObjectRoot>& Roots)
 		RootPropertyNode->InitNode( InitParams );
 	}
 
+	// Restore expanded items is called twice. The first time handles the existing property nodes,
+	// and second time after handling property customization as it might add new items that need to be expanded too.
+
+	// Restore existing nodes
+	RestoreAllExpandedItems();
+
 	UpdatePropertyMaps();
 
-	// Restore expanded items after handling property customization as some items that require expansion come from the customizations.
+	// Restore nodes generated by property customization
 	RestoreAllExpandedItems();
 
 	UpdateFilteredDetails();
