Property Matrix: Fix a crash when you attempted to edit a blank cell due to it not having a presenter

diff --git a/Engine/Source/Editor/PropertyEditor/Private/UserInterface/PropertyTable/SPropertyTableCell.cpp b/Engine/Source/Editor/PropertyEditor/Private/UserInterface/PropertyTable/SPropertyTableCell.cpp
--- a/Engine/Source/Editor/PropertyEditor/Private/UserInterface/PropertyTable/SPropertyTableCell.cpp
+++ b/Engine/Source/Editor/PropertyEditor/Private/UserInterface/PropertyTable/SPropertyTableCell.cpp
@@ -172,7 +172,7 @@ int32 SPropertyTableCell::OnPaint( const FPaintArgs& Args, const FGeometry& Allo
 		}
 
 		// If the cell is in edit mode, we give the presenter a chance to specify the background
-		if(Cell->InEditMode())
+		if(Cell->InEditMode() && Presenter)
 		{
 			if(const FSlateBrush* PresenterBackground = Presenter->GetEditModeCellBrush())
 			{
