Property Editor:
- Fix properties under sub-categories not appearing at all when using Struct on Scope in Property Row Generator

diff --git a/Engine/Source/Editor/PropertyEditor/Private/DetailLayoutBuilderImpl.cpp b/Engine/Source/Editor/PropertyEditor/Private/DetailLayoutBuilderImpl.cpp
--- a/Engine/Source/Editor/PropertyEditor/Private/DetailLayoutBuilderImpl.cpp
+++ b/Engine/Source/Editor/PropertyEditor/Private/DetailLayoutBuilderImpl.cpp
@@ -343,18 +343,97 @@ void FDetailLayoutBuilderImpl::GenerateDetailLayout()
 	TArray<TSharedRef<FDetailCategoryImpl>> SimpleCategories;
 	TArray<TSharedRef<FDetailCategoryImpl>> AdvancedOnlyCategories;
 
-	SubCategoryMap.Empty();
-	for (FCategoryMap::TIterator It(DefaultCategoryMap); It; ++It)
-	{
-		// Remove all subcategories
-		TSharedPtr<FDetailCategoryImpl> DetailCategory = It.Value();
-		// Note: Sub-categories are added later
-		int32 Index = INDEX_NONE;
-		if (DetailCategory->GetCategoryName().ToString().FindChar(FPropertyNodeConstants::CategoryDelimiterChar, Index))
+	// Remove all subcategories
+	{
+		const TSharedRef<FDetailLayoutBuilderImpl> This = SharedThis(this);
+
+		// The map of Default Categories added after removing Sub categories
+		// Only used when there is no support for sub-categories
+		FCategoryMap DefaultCategoryMapToAppend;
+
+		// Store the Default Category Map Num as it will be decreasing as sub category entries are removed
+		const int32 DefaultCategoryMapCount = DefaultCategoryMap.Num();
+
+		FName ParentStructPropertyName = NAME_None;
+		bool bSupportsSubCategory = false;
+
+		if (TSharedPtr<FComplexPropertyNode> RootNodePtr = RootNode.Pin())
 		{
-			SubCategoryMap.Add(It.Key(), DetailCategory);
+			if (FStructProperty* const ParentStructProperty = CastField<FStructProperty>(RootNodePtr->GetProperty()))
+			{
+				ParentStructPropertyName = ParentStructProperty->GetFName();
+			}
+
+			// Currently only Object Nodes with Show Categories support Sub-categories
+			bSupportsSubCategory = RootNodePtr->AsObjectNode() && RootNodePtr->HasNodeFlags(EPropertyNodeFlags::ShowCategories);
+		}
+
+		SubCategoryMap.Empty();
+
+		for (FCategoryMap::TIterator It(DefaultCategoryMap); It; ++It)
+		{
+			TSharedPtr<FDetailCategoryImpl> DetailCategory = It.Value();
+
+			int32 Index = INDEX_NONE;
+			const FString CategoryNameStr = DetailCategory->GetCategoryName().ToString();
+			if (!CategoryNameStr.FindChar(FPropertyNodeConstants::CategoryDelimiterChar, Index))
+			{
+				// No category delimiter found
+				continue;
+			}
+
+			// Note: Sub-categories are added later if supported
+			if (bSupportsSubCategory)
+			{
+				SubCategoryMap.Add(It.Key(), DetailCategory);
+			}
+			// When Sub category isn't supported, generate properties and move them to parent category
+			else
+			{
+				const FName ParentCategoryName = *CategoryNameStr.Left(Index);
+
+				TSharedPtr<FDetailCategoryImpl>* const ExistingParentDetailCategory = DefaultCategoryMap.Find(ParentCategoryName);
+
+				TSharedPtr<FDetailCategoryImpl>& ParentDetailCategory = ExistingParentDetailCategory
+					? *ExistingParentDetailCategory
+					: DefaultCategoryMapToAppend.FindOrAdd(ParentCategoryName);
+
+				if (!ParentDetailCategory.IsValid())
+				{
+					ParentDetailCategory = MakeShared<FDetailCategoryImpl>(ParentCategoryName, This);
+					ParentDetailCategory->SetSortOrder(DetailCategory->GetSortOrder());
+					ParentDetailCategory->SetDisplayName(ParentCategoryName, FText::GetEmpty());
+				}
+
+				// Move the Property Nodes from the sub-category to the parent category
+				// To do this, generate a layout for the sub category here as they're unsupported and won't have an opportunity to do it later
+				FDetailNodeList ChildNodes;
+				DetailCategory->GenerateLayout();
+				DetailCategory->GetGeneratedChildren(ChildNodes, /*bIgnoreVisibility*/true, /*bIgnoreAdvancedDropdown*/true);
+
+				for (const TSharedRef<FDetailTreeNode>& ChildNode : ChildNodes)
+				{
+					TSharedPtr<FPropertyNode> PropertyNode = ChildNode->GetPropertyNode();
+					if (!PropertyNode.IsValid())
+					{
+						continue;
+					}
+
+					// If there is no outer object then the class is the object root and there is only one instance
+					FName InstanceName = ParentStructPropertyName;
+					FPropertyNode* const ParentNode = PropertyNode->GetParentNode();
+					if (ParentNode && ParentNode->GetProperty())
+					{
+						InstanceName = ParentNode->GetProperty()->GetFName();
+					}
+					ParentDetailCategory->AddPropertyNode(PropertyNode.ToSharedRef(), InstanceName);
+				}
+			}
+
 			It.RemoveCurrent();
 		}
+
+		DefaultCategoryMap.Append(MoveTemp(DefaultCategoryMapToAppend));
 	}
 
 	// Build default categories
