Property Matrix: Permission Lists are now checked for every object that is being edited instead of just the most common base class, and we don't create widgets for cells that don't pass the permission list

diff --git a/Engine/Source/Editor/PropertyEditor/Private/Presentation/PropertyTable/PropertyTableCell.cpp b/Engine/Source/Editor/PropertyEditor/Private/Presentation/PropertyTable/PropertyTableCell.cpp
--- a/Engine/Source/Editor/PropertyEditor/Private/Presentation/PropertyTable/PropertyTableCell.cpp
+++ b/Engine/Source/Editor/PropertyEditor/Private/Presentation/PropertyTable/PropertyTableCell.cpp
@@ -7,6 +7,7 @@
 #include "Presentation/PropertyEditor/PropertyEditor.h"
 #include "PropertyHandle.h"
 #include "PropertyNode.h"
+#include "PropertyPermissionList.h"
 #include "UObject/WeakObjectPtr.h"
 
 class FPropertyPath;
@@ -15,6 +16,7 @@ class UObject;
 FPropertyTableCell::FPropertyTableCell( const TSharedRef< class IPropertyTableColumn >& InColumn, const TSharedRef< class IPropertyTableRow >& InRow )
 	: bIsBound( false )
 	, bInEditMode( false )
+	, bIsPropertyAllowed(true)
 	, EnteredEditModeEvent()
 	, ExitedEditModeEvent()
 	, Column( InColumn )
@@ -102,16 +104,18 @@ void FPropertyTableCell::Refresh()
 		bIsBound = true;
 		PropertyEditor = FPropertyEditor::Create( PropertyNode.ToSharedRef(), GetTable() );
 	}
+
+	bIsPropertyAllowed = DoesPropertyPassPermissionList();
 }
 
 bool FPropertyTableCell::IsReadOnly() const 
 { 
-	return !IsBound() || PropertyEditor->IsEditConst() || ( PropertyEditor->HasEditCondition() && !PropertyEditor->IsEditConditionMet() );
+	return !bIsPropertyAllowed || !IsBound() || PropertyEditor->IsEditConst() || ( PropertyEditor->HasEditCondition() && !PropertyEditor->IsEditConditionMet() );
 }
 
 bool FPropertyTableCell::IsValid() const
 {
-	return !IsBound() || (PropertyEditor->GetPropertyHandle()->GetProperty() != NULL);
+	return (!IsBound() || (PropertyEditor->GetPropertyHandle()->GetProperty() != NULL));
 }
 
 FString FPropertyTableCell::GetValueAsString() const
@@ -159,8 +163,47 @@ TSharedPtr< class IPropertyHandle > FPropertyTableCell::GetPropertyHandle() cons
 	return NULL;
 }
 
+bool FPropertyTableCell::DoesPropertyPassPermissionList()
+{
+	// If we have a permission list, make sure the object for this cell allows the property to be displayed
+	if(FPropertyEditorPermissionList::Get().IsEnabled())
+	{
+		// Get the Object and the Property
+		UObject* Object = GetObject().Get();
+		TSharedPtr< class FPropertyNode > Node = GetNode();
+		
+		if(Object && Node)
+		{
+			// Extract the FProperty and UClass from them respectively
+			const FProperty* Property = Node->GetProperty();
+			const UClass* ObjClass = Cast<UClass>(Object);
+			if (ObjClass == nullptr)
+			{
+				ObjClass = Object->GetClass();
+			}
+			
+			if(ObjClass && Property)
+			{
+				if(!FPropertyEditorPermissionList::Get().DoesPropertyPassFilter(ObjClass, Property->GetFName()))
+				{
+					return false;
+				}
+			}
+		}
+	}
+
+	return true;
+}
+
 void FPropertyTableCell::EnterEditMode() 
 {
+	// We don't let the cell enter edit mode if it doesn't pass the permission list, but we do allow ExitEditMode()
+	// in case the permission list changes when you are in edit mode
+	if(!bIsPropertyAllowed)
+	{
+		return;
+	}
+	
 	if ( !bInEditMode )
 	{
 		TSharedRef< IPropertyTable > Table = GetTable();
diff --git a/Engine/Source/Editor/PropertyEditor/Private/Presentation/PropertyTable/PropertyTableCell.h b/Engine/Source/Editor/PropertyEditor/Private/Presentation/PropertyTable/PropertyTableCell.h
--- a/Engine/Source/Editor/PropertyEditor/Private/Presentation/PropertyTable/PropertyTableCell.h
+++ b/Engine/Source/Editor/PropertyEditor/Private/Presentation/PropertyTable/PropertyTableCell.h
@@ -27,6 +27,7 @@ public:
 	virtual bool IsBound() const override { return bIsBound; }
 	virtual bool InEditMode() const override { return bInEditMode; }
 	virtual bool IsValid() const override;
+	virtual bool PassesPermissionList() const override { return bIsPropertyAllowed; }
 
 	virtual FString GetValueAsString() const override;
 	virtual FText GetValueAsText() const override;
@@ -48,11 +49,15 @@ public:
 	DECLARE_DERIVED_EVENT( FPropertyTableCell, IPropertyTableCell::FExitedEditModeEvent, FExitedEditModeEvent  );
 	virtual FExitedEditModeEvent& OnExitedEditMode() override { return ExitedEditModeEvent; }
 
-
+private:
+	// Check if the property for this cell passes the permission list for the object we are viewing
+	bool DoesPropertyPassPermissionList();
+	
 private:
 
 	bool bIsBound;
 	bool bInEditMode;
+	bool bIsPropertyAllowed;
 
 	FEnteredEditModeEvent EnteredEditModeEvent;
 	FExitedEditModeEvent ExitedEditModeEvent;
diff --git a/Engine/Source/Editor/PropertyEditor/Private/SPropertyTreeViewImpl.cpp b/Engine/Source/Editor/PropertyEditor/Private/SPropertyTreeViewImpl.cpp
--- a/Engine/Source/Editor/PropertyEditor/Private/SPropertyTreeViewImpl.cpp
+++ b/Engine/Source/Editor/PropertyEditor/Private/SPropertyTreeViewImpl.cpp
@@ -703,7 +703,27 @@ bool SPropertyTreeViewImpl::IsPropertyNodeVisible(TSharedPtr<FPropertyNode> InPr
 	const FProperty* Property = InPropertyNode->GetProperty();
 	if(Property != NULL)
 	{
-		bPropertyVisible = FPropertyEditorPermissionList::Get().DoesPropertyPassFilter(FDetailTreeNode::GetPropertyNodeBaseStructure(InPropertyNode->GetParentNode()), Property->GetFName());
+		// If we have a permission list in use, make sure this property is in the list for the most common base class and the class of all objects we are viewing
+		if(FPropertyEditorPermissionList::Get().IsEnabled())
+		{
+			bPropertyVisible = FPropertyEditorPermissionList::Get().DoesPropertyPassFilter(FDetailTreeNode::GetPropertyNodeBaseStructure(InPropertyNode->GetParentNode()), Property->GetFName());
+
+			for (int32 ObjectIndex = 0; ObjectIndex < RootPropertyNode->GetNumObjects(); ++ObjectIndex)
+			{
+				UObject* ObjectToCheck = RootPropertyNode->GetUObject(ObjectIndex);
+
+				const UClass* ObjClass = Cast<UClass>(ObjectToCheck);
+				if (ObjClass == nullptr)
+				{
+					ObjClass = ObjectToCheck->GetClass();
+				}
+
+				if(ObjClass)
+				{
+					bPropertyVisible &= FPropertyEditorPermissionList::Get().DoesPropertyPassFilter(ObjClass, Property->GetFName());
+				}
+			}
+		}
 
 		if(IsPropertyVisible.IsBound())
 		{
diff --git a/Engine/Source/Editor/PropertyEditor/Private/UserInterface/PropertyTable/SPropertyTableCell.cpp b/Engine/Source/Editor/PropertyEditor/Private/UserInterface/PropertyTable/SPropertyTableCell.cpp
--- a/Engine/Source/Editor/PropertyEditor/Private/UserInterface/PropertyTable/SPropertyTableCell.cpp
+++ b/Engine/Source/Editor/PropertyEditor/Private/UserInterface/PropertyTable/SPropertyTableCell.cpp
@@ -87,6 +87,13 @@ void SPropertyTableCell::OnCellValueChanged( UObject* Object, FPropertyChangedEv
 TSharedRef< SWidget > SPropertyTableCell::ConstructCellContents()
 {
 	TSharedRef< SWidget > CellContents = SNullWidget::NullWidget;
+
+	// If the cell doesnt pass the permission list, don't show any widget for it
+	if(!Cell->PassesPermissionList())
+	{
+		return CellContents;
+	}
+	
 	if ( Presenter.IsValid() )
 	{
 		if ( Cell->InEditMode() )
diff --git a/Engine/Source/Editor/PropertyEditor/Public/IPropertyTableCell.h b/Engine/Source/Editor/PropertyEditor/Public/IPropertyTableCell.h
--- a/Engine/Source/Editor/PropertyEditor/Public/IPropertyTableCell.h
+++ b/Engine/Source/Editor/PropertyEditor/Public/IPropertyTableCell.h
@@ -16,6 +16,7 @@ public:
 	virtual bool IsBound() const = 0;
 	virtual bool InEditMode() const = 0;
 	virtual bool IsValid() const = 0;
+	virtual bool PassesPermissionList() const { return true; }
 
 	virtual FString GetValueAsString() const = 0;
 	virtual FText GetValueAsText() const = 0;
