[PropertyEditor] FDetailTreeNode::GetChildren bInIgnoreVisibility also ignores Advanced visibility (so all child nodes are retrieved regardless of advanced expansion state)

diff --git a/Engine/Source/Editor/PropertyEditor/Private/DetailCategoryBuilderImpl.cpp b/Engine/Source/Editor/PropertyEditor/Private/DetailCategoryBuilderImpl.cpp
--- a/Engine/Source/Editor/PropertyEditor/Private/DetailCategoryBuilderImpl.cpp
+++ b/Engine/Source/Editor/PropertyEditor/Private/DetailCategoryBuilderImpl.cpp
@@ -1069,7 +1069,8 @@ void FDetailCategoryImpl::GetGeneratedChildren(FDetailNodeList& OutChildren, boo
 			OutChildren.Add(AdvancedDropdownNode.ToSharedRef());
 		}
 
-		if (ShouldAdvancedBeExpanded())
+		// bIgnoreVisibility treats advanced as expanded
+		if (bIgnoreVisibility || ShouldAdvancedBeExpanded())
 		{
 			for (TSharedRef<FDetailTreeNode>& Child : AdvancedChildNodes)
 			{
@@ -1077,7 +1078,7 @@ void FDetailCategoryImpl::GetGeneratedChildren(FDetailNodeList& OutChildren, boo
 				{
 					if (Child->ShouldShowOnlyChildren())
 					{
-						Child->GetChildren(OutChildren);
+						Child->GetChildren(OutChildren, bIgnoreVisibility);
 					}
 					else
 					{
