[PropertyEditor] Category copy/paste fix for blueprint editor

diff --git a/Engine/Source/Editor/PropertyEditor/Private/SDetailSingleItemRow.cpp b/Engine/Source/Editor/PropertyEditor/Private/SDetailSingleItemRow.cpp
--- a/Engine/Source/Editor/PropertyEditor/Private/SDetailSingleItemRow.cpp
+++ b/Engine/Source/Editor/PropertyEditor/Private/SDetailSingleItemRow.cpp
@@ -1333,7 +1333,9 @@ bool SDetailSingleItemRow::PasteFromText(const FString& InTag, const FString& In
 	}
 	
 	const TSharedPtr<IPropertyHandle> PropertyHandle = GetPropertyHandle();
-	if (!InTag.IsEmpty())
+
+	const bool bIsTagged = !InTag.IsEmpty();
+	if (bIsTagged)
 	{
 		const FString PropertyPath = UE::PropertyEditor::Private::GetPropertyPath(
 			[&]() { return GetPropertyHandle(); },
@@ -1346,12 +1348,42 @@ bool SDetailSingleItemRow::PasteFromText(const FString& InTag, const FString& In
 		}
 	}
 
+	// The logic below is largely taken from SDisplayClusterColorGradingColorWheel::CommitColor,
+	// which avoids writing to trashed objects
+
 	FScopedTransaction Transaction(NSLOCTEXT("UnrealEd", "PasteProperty", "Paste Property"));
-	if (PropertyHandle->SetValueFromFormattedString(InText, EPropertyValueSetFlags::InstanceObjects) != FPropertyAccess::Success)
+	
+	EPropertyValueSetFlags::Type PropertyValueSetFlags = EPropertyValueSetFlags::InstanceObjects;
+
+	// If tagged, add the InteractiveChange flag so as not to run PECP
+	// @todo: would be better to indicate that this is a batched paste rather than checking for a tag
+	if (bIsTagged)
+	{
+		PropertyValueSetFlags |= EPropertyValueSetFlags::InteractiveChange;
+		PropertyValueSetFlags |= EPropertyValueSetFlags::NotTransactable;
+	}
+	
+	if (PropertyHandle->SetValueFromFormattedString(InText, PropertyValueSetFlags) != FPropertyAccess::Success)
 	{
 		return false;
 	}
 
+	if (bIsTagged)
+	{
+		TArray<UObject*> OuterObjects;
+		PropertyHandle->GetOuterObjects(OuterObjects);
+		for (UObject* Object : OuterObjects)
+		{
+			if (!Object->HasAnyFlags(RF_Transactional))
+			{
+				Object->SetFlags(RF_Transactional);
+			}
+
+			SaveToTransactionBuffer(Object, false);
+			SnapshotTransactionBuffer(Object);
+		}
+	}
+
 	return true;
 }
 
