Object Mixer:

-Fix: Blueprint Lights Parameters Switch To "Multiple Values" When Transformed

diff --git a/Engine/Source/Editor/PropertyEditor/Private/SSingleProperty.cpp b/Engine/Source/Editor/PropertyEditor/Private/SSingleProperty.cpp
--- a/Engine/Source/Editor/PropertyEditor/Private/SSingleProperty.cpp
+++ b/Engine/Source/Editor/PropertyEditor/Private/SSingleProperty.cpp
@@ -269,7 +269,7 @@ void SSingleProperty::ReplaceObjects( const TMap<UObject*, UObject*>& OldToNewOb
 		// Scan all objects and look for objects which need to be replaced
 		for ( TPropObjectIterator Itor( RootPropertyNode->ObjectIterator() ); Itor; ++Itor )
 		{
-			UObject* Replacement = OldToNewObjectMap.FindRef( Itor->Get() );
+			UObject* Replacement = OldToNewObjectMap.FindRef( Itor->Get(true) );
 			if( Replacement )
 			{
 				bObjectsReplaced = true;
