[PropertyEditor] Fix for recursive pastability check (ignore expansion states)

diff --git a/Engine/Source/Editor/PropertyEditor/Private/PropertyEditorClipboardPrivate.h b/Engine/Source/Editor/PropertyEditor/Private/PropertyEditorClipboardPrivate.h
--- a/Engine/Source/Editor/PropertyEditor/Private/PropertyEditorClipboardPrivate.h
+++ b/Engine/Source/Editor/PropertyEditor/Private/PropertyEditorClipboardPrivate.h
@@ -25,6 +25,9 @@ namespace UE::PropertyEditor::Internal
 		/** Applicability of the clipboard content to this row. */ 
 		bool bIsApplicable = false;
 
+		/** Previous PropertyHandle count, used to validate applicability. */
+		int32 PreviousPropertyHandleNum = INDEX_NONE;
+
 		/** Deserialized property values (as text). */
 		TMap<FName, FString> PropertyValues;
 
@@ -36,6 +39,7 @@ namespace UE::PropertyEditor::Internal
 		{
 			Content.Reset();
 			bIsApplicable = false;
+			PreviousPropertyHandleNum = INDEX_NONE;
 			PropertyValues.Reset();
 			PropertyNames.Reset();
 		}
diff --git a/Engine/Source/Editor/PropertyEditor/Private/SDetailCategoryTableRow.cpp b/Engine/Source/Editor/PropertyEditor/Private/SDetailCategoryTableRow.cpp
--- a/Engine/Source/Editor/PropertyEditor/Private/SDetailCategoryTableRow.cpp
+++ b/Engine/Source/Editor/PropertyEditor/Private/SDetailCategoryTableRow.cpp
@@ -391,8 +391,11 @@ bool SDetailCategoryTableRow::CanPasteCategory()
 	FString ClipboardContent;
 	FPropertyEditorClipboard::ClipboardPaste(ClipboardContent);
 
-	// If same as last, return previously resolved applicability
-	if (PreviousClipboardData.Content.Get({}).Equals(ClipboardContent))
+	const bool bChildrenHaveChanged = PreviousClipboardData.PreviousPropertyHandleNum != PropertyHandles.Num();
+
+	// If child property count hasn't changed, and content same as last, return previously resolved applicability
+	if (!bChildrenHaveChanged
+		&& PreviousClipboardData.Content.Get({}).Equals(ClipboardContent))
 	{
 		return PreviousClipboardData.bIsApplicable;
 	}
@@ -427,6 +430,7 @@ bool SDetailCategoryTableRow::CanPasteCategory()
 	// @note: properties must all match to be applicable
 	PreviousClipboardData.Content = ClipboardContent;
 
+	PreviousClipboardData.PreviousPropertyHandleNum = PropertyHandles.Num();
 	return PreviousClipboardData.bIsApplicable = Algo::Compare(PreviousClipboardData.PropertyNames, PropertyNames);
 }
 
