Fixed crash when a class picker was called in the Dataprep Editor
  - Replace the access to a function pointer through a raw pointer to access through a shared pointer: The object was deleted before the function was caled.

diff --git a/Engine/Source/Editor/PropertyEditor/Private/UserInterface/PropertyEditor/SPropertyEditorClass.cpp b/Engine/Source/Editor/PropertyEditor/Private/UserInterface/PropertyEditor/SPropertyEditorClass.cpp
--- a/Engine/Source/Editor/PropertyEditor/Private/UserInterface/PropertyEditor/SPropertyEditorClass.cpp
+++ b/Engine/Source/Editor/PropertyEditor/Private/UserInterface/PropertyEditor/SPropertyEditorClass.cpp
@@ -289,7 +289,7 @@ void SPropertyEditorClass::CreateClassFilter(const TArray<TSharedRef<IClassViewe
 
 TSharedRef<SWidget> SPropertyEditorClass::GenerateClassPicker()
 {
-	FOnClassPicked OnPicked(FOnClassPicked::CreateRaw(this, &SPropertyEditorClass::OnClassPicked));
+	FOnClassPicked OnPicked(FOnClassPicked::CreateSP(this, &SPropertyEditorClass::OnClassPicked));
 
 	return SNew(SBox)
 		.WidthOverride(280)
