[RegressionFix] ParentClass and ImplementedInterfaces properties aren't diffing in class settings panel

diff --git a/Engine/Source/Editor/PropertyEditor/Private/AsyncDetailViewDiff.cpp b/Engine/Source/Editor/PropertyEditor/Private/AsyncDetailViewDiff.cpp
--- a/Engine/Source/Editor/PropertyEditor/Private/AsyncDetailViewDiff.cpp
+++ b/Engine/Source/Editor/PropertyEditor/Private/AsyncDetailViewDiff.cpp
@@ -37,35 +37,49 @@ bool TreeDiffSpecification::AreValuesEqual<TWeakPtr<FDetailTreeNode>>(const TWea
 		return PinnedTreeNodeA == PinnedTreeNodeB;
 	}
 
-	const FResolvedProperty ResolvedA = GetResolvedProperty(PinnedTreeNodeA->GetPropertyNode(), GetObject(PinnedTreeNodeA));
-    const FResolvedProperty ResolvedB = GetResolvedProperty(PinnedTreeNodeB->GetPropertyNode(), GetObject(PinnedTreeNodeB));
-	if (ResolvedA.Property && ResolvedB.Property)
-	{
-		// property nodes
-		if (!ResolvedA.Property->SameType(ResolvedB.Property))
-		{
-			return false;
-		}
-	    const void* DataA = ResolvedA.Property->ContainerPtrToValuePtr<void*>(ResolvedA.Object);
-	    const void* DataB = ResolvedB.Property->ContainerPtrToValuePtr<void*>(ResolvedB.Object);
-		return ResolvedA.Property->Identical(DataA, DataB, PPF_DeepComparison);
-	}
-	if (!ResolvedA.Property && !ResolvedB.Property)
+	TArray<TSharedRef<FPropertyNode>> PropertyNodesA;
+	TArray<TSharedRef<FPropertyNode>> PropertyNodesB;
+	PinnedTreeNodeA->GetAllPropertyNodes(PropertyNodesA);
+	PinnedTreeNodeB->GetAllPropertyNodes(PropertyNodesB);
+	
+	ensure(PropertyNodesA.Num() == PropertyNodesB.Num()); // AreMatching(...) should've stopped this from happening
+	if (PropertyNodesA.IsEmpty() || PropertyNodesB.IsEmpty())
 	{
 		// category nodes
 		return PinnedTreeNodeA->GetNodeName() == PinnedTreeNodeB->GetNodeName();
 	}
-	return ensure(false); // AreMatching(...) should've stopped this from happening
-}
 
-static bool MapKeysMatch(const TSharedPtr<FDetailTreeNode>& TreeNodeA, const TSharedPtr<FDetailTreeNode>& TreeNodeB, int32 KeyIndexA, int32 KeyIndexB)
-{
-	const TSharedPtr<FPropertyNode> MapPropertyNodeA = TreeNodeA->GetPropertyNode()->GetParentNode()->AsShared();
-	const TSharedPtr<FPropertyNode> MapPropertyNodeB = TreeNodeB->GetPropertyNode()->GetParentNode()->AsShared();
-	if (!MapPropertyNodeA || !MapPropertyNodeB)
+	const UObject* OwningObjectA = GetObject(PinnedTreeNodeA);
+	const UObject* OwningObjectB = GetObject(PinnedTreeNodeB);
+
+	for (int32 PropNodeIndex = 0; PropNodeIndex < PropertyNodesA.Num(); ++PropNodeIndex)
 	{
-		return false;
+		const TSharedRef<FPropertyNode>& PropertyNodeA = PropertyNodesA[PropNodeIndex];
+		const TSharedRef<FPropertyNode>& PropertyNodeB = PropertyNodesB[PropNodeIndex];
+
+		const FResolvedProperty ResolvedA = GetResolvedProperty(PropertyNodeA, OwningObjectA);
+		const FResolvedProperty ResolvedB = GetResolvedProperty(PropertyNodeB, OwningObjectB);
+		if (ResolvedA.Property && ResolvedB.Property)
+		{
+			// property nodes
+			if (!ResolvedA.Property->SameType(ResolvedB.Property))
+			{
+				return false;
+			}
+			const void* DataA = ResolvedA.Property->ContainerPtrToValuePtr<void*>(ResolvedA.Object);
+			const void* DataB = ResolvedB.Property->ContainerPtrToValuePtr<void*>(ResolvedB.Object);
+			if(!ResolvedA.Property->Identical(DataA, DataB, PPF_DeepComparison))
+			{
+				return false;
+			}
+		}
 	}
+	return true;
+}
+
+static bool MapKeysMatch(const TSharedRef<FPropertyNode>& MapPropertyNodeA, const TSharedRef<FPropertyNode>& MapPropertyNodeB, int32 KeyIndexA, int32 KeyIndexB,
+	const UObject* OwningObjectA, const UObject* OwningObjectB)
+{
 	const FMapProperty* MapPropertyA = CastField<FMapProperty>(MapPropertyNodeA->GetProperty());
 	const FMapProperty* MapPropertyB = CastField<FMapProperty>(MapPropertyNodeB->GetProperty());
 	if (!MapPropertyA || !MapPropertyB)
@@ -73,8 +87,8 @@ static bool MapKeysMatch(const TSharedPtr<FDetailTreeNode>& TreeNodeA, const TSh
 		return false;
 	}
 	
-	const FResolvedProperty ResolvedMapA = GetResolvedProperty(MapPropertyNodeA, GetObject(TreeNodeA));
-	const FResolvedProperty ResolvedMapB = GetResolvedProperty(MapPropertyNodeB, GetObject(TreeNodeB));
+	const FResolvedProperty ResolvedMapA = GetResolvedProperty(MapPropertyNodeA, OwningObjectA);
+	const FResolvedProperty ResolvedMapB = GetResolvedProperty(MapPropertyNodeB, OwningObjectB);
 	FScriptMapHelper MapHelperA(MapPropertyA, MapPropertyA->ContainerPtrToValuePtr<UObject*>(ResolvedMapA.Object));
 	FScriptMapHelper MapHelperB(MapPropertyB, MapPropertyB->ContainerPtrToValuePtr<UObject*>(ResolvedMapB.Object));
 	
@@ -88,14 +102,9 @@ static bool MapKeysMatch(const TSharedPtr<FDetailTreeNode>& TreeNodeA, const TSh
 	return false;
 }
 
-static bool SetKeysMatch(const TSharedPtr<FDetailTreeNode>& TreeNodeA, const TSharedPtr<FDetailTreeNode>& TreeNodeB, int32 KeyIndexA, int32 KeyIndexB)
+static bool SetKeysMatch(const TSharedRef<FPropertyNode>& SetPropertyNodeA, const TSharedRef<FPropertyNode>& SetPropertyNodeB, int32 KeyIndexA, int32 KeyIndexB,
+	const UObject* OwningObjectA, const UObject* OwningObjectB)
 {
-	const TSharedPtr<FPropertyNode> SetPropertyNodeA = TreeNodeA->GetPropertyNode()->GetParentNode()->AsShared();
-	const TSharedPtr<FPropertyNode> SetPropertyNodeB = TreeNodeB->GetPropertyNode()->GetParentNode()->AsShared();
-	if (!SetPropertyNodeA || !SetPropertyNodeB)
-	{
-		return false;
-	}
 	const FSetProperty* SetPropertyA = CastField<FSetProperty>(SetPropertyNodeA->GetProperty());
 	const FSetProperty* SetPropertyB = CastField<FSetProperty>(SetPropertyNodeB->GetProperty());
 	if (!SetPropertyA || !SetPropertyB)
@@ -103,8 +112,8 @@ static bool SetKeysMatch(const TSharedPtr<FDetailTreeNode>& TreeNodeA, const TSh
 		return false;
 	}
 	
-	const FResolvedProperty ResolvedSetA = GetResolvedProperty(SetPropertyNodeA, GetObject(TreeNodeA));
-	const FResolvedProperty ResolvedSetB = GetResolvedProperty(SetPropertyNodeB, GetObject(TreeNodeB));
+	const FResolvedProperty ResolvedSetA = GetResolvedProperty(SetPropertyNodeA, OwningObjectA);
+	const FResolvedProperty ResolvedSetB = GetResolvedProperty(SetPropertyNodeB, OwningObjectB);
 	FScriptSetHelper SetHelperA(SetPropertyA, SetPropertyA->ContainerPtrToValuePtr<UObject*>(ResolvedSetA.Object));
 	FScriptSetHelper SetHelperB(SetPropertyB, SetPropertyB->ContainerPtrToValuePtr<UObject*>(ResolvedSetB.Object));
 	
@@ -127,11 +136,30 @@ bool TreeDiffSpecification::AreMatching<TWeakPtr<FDetailTreeNode>>(const TWeakPt
 	{
 		return PinnedTreeNodeA == PinnedTreeNodeB;
 	}
+
 	
-	const TSharedPtr<FPropertyNode> PropertyNodeA = PinnedTreeNodeA->GetPropertyNode();
-	const TSharedPtr<FPropertyNode> PropertyNodeB =PinnedTreeNodeB->GetPropertyNode();
-	if (PropertyNodeA && PropertyNodeB)
+	TArray<TSharedRef<FPropertyNode>> PropertyNodesA;
+	TArray<TSharedRef<FPropertyNode>> PropertyNodesB;
+	PinnedTreeNodeA->GetAllPropertyNodes(PropertyNodesA);
+	PinnedTreeNodeB->GetAllPropertyNodes(PropertyNodesB);
+	if (PropertyNodesA.Num() != PropertyNodesB.Num())
 	{
+		return false;
+	}
+	if (PropertyNodesA.IsEmpty())
+	{
+		// category nodes
+		return PinnedTreeNodeA->GetNodeName() == PinnedTreeNodeB->GetNodeName();
+	}
+	
+	const UObject* OwningObjectA = GetObject(PinnedTreeNodeA);
+	const UObject* OwningObjectB = GetObject(PinnedTreeNodeB);
+	
+	for (int32 PropNodeIndex = 0; PropNodeIndex < PropertyNodesA.Num(); ++PropNodeIndex)
+	{
+		const TSharedRef<FPropertyNode>& PropertyNodeA = PropertyNodesA[PropNodeIndex];
+		const TSharedRef<FPropertyNode>& PropertyNodeB = PropertyNodesB[PropNodeIndex];
+
 		// property nodes
 		const int32 ArrayIndexA = PropertyNodeA->GetArrayIndex();
 		const int32 ArrayIndexB = PropertyNodeB->GetArrayIndex();
@@ -140,28 +168,26 @@ bool TreeDiffSpecification::AreMatching<TWeakPtr<FDetailTreeNode>>(const TWeakPt
 		
 		if (ArrayIndexA != INDEX_NONE && ArrayIndexB != INDEX_NONE)
 		{
-			const FProperty* ParentPropertyA = PropertyNodeA->GetParentNode()->GetProperty();
-			const FProperty* ParentPropertyB = PropertyNodeB->GetParentNode()->GetProperty();
+			const TSharedRef<FPropertyNode> ParentPropertyNodeA = PropertyNodeA->GetParentNode()->AsShared();
+			const TSharedRef<FPropertyNode> ParentPropertyNodeB = PropertyNodeB->GetParentNode()->AsShared();
+			const FProperty* ParentPropertyA = ParentPropertyNodeA->GetProperty();
+			const FProperty* ParentPropertyB = ParentPropertyNodeB->GetProperty();
 			
 			// sets and maps are stored by index in the property tree so we need to dig their keys out of the data
 			// and compare those instead
-			if (CastField<FMapProperty>(ParentPropertyA) || CastField<FMapProperty>(ParentPropertyB))
+			if (ParentPropertyA->IsA<FMapProperty>() || ParentPropertyB->IsA<FMapProperty>())
 			{
-				return MapKeysMatch(
-					PinnedTreeNodeA,
-					PinnedTreeNodeB,
-					ArrayIndexA,
-					ArrayIndexB
-				);
+				if (!MapKeysMatch(ParentPropertyNodeA, ParentPropertyNodeB, ArrayIndexA, ArrayIndexB, OwningObjectA, OwningObjectB))
+				{
+					return false;
+				}
 			}
 			if (ParentPropertyA->IsA<FSetProperty>() || ParentPropertyB->IsA<FSetProperty>())
 			{
-				return SetKeysMatch(
-					PinnedTreeNodeA,
-					PinnedTreeNodeB,
-					ArrayIndexA,
-					ArrayIndexB
-				);
+				if (SetKeysMatch(ParentPropertyNodeA, ParentPropertyNodeB, ArrayIndexA, ArrayIndexB, OwningObjectA, OwningObjectB))
+				{
+					return false;
+				}
 			}
 		}
 		
@@ -171,16 +197,13 @@ bool TreeDiffSpecification::AreMatching<TWeakPtr<FDetailTreeNode>>(const TWeakPt
 		}
 		const FName PropertyNameA = PropertyA ? PropertyA->GetFName() : NAME_None;
 		const FName PropertyNameB = PropertyB ? PropertyB->GetFName() : NAME_None;
-		
-		return PropertyNameA == PropertyNameB;
-	}
-	if (!PropertyNodeA && !PropertyNodeB)
-	{
-		// category nodes
-		return PinnedTreeNodeA->GetNodeName() == PinnedTreeNodeB->GetNodeName();
+
+		if (PropertyNameA != PropertyNameB)
+		{
+			return false;
+		}
 	}
-	// node type mismatch
-	return false;
+	return true;
 }
 
 template<>
diff --git a/Engine/Source/Editor/PropertyEditor/Private/DetailItemNode.cpp b/Engine/Source/Editor/PropertyEditor/Private/DetailItemNode.cpp
--- a/Engine/Source/Editor/PropertyEditor/Private/DetailItemNode.cpp
+++ b/Engine/Source/Editor/PropertyEditor/Private/DetailItemNode.cpp
@@ -7,6 +7,7 @@
 #include "DetailPropertyRow.h"
 #include "IDetailKeyframeHandler.h"
 #include "ObjectPropertyNode.h"
+#include "PropertyHandleImpl.h"
 #include "SConstrainedBox.h"
 #include "PropertyPermissionList.h"
 #include "SDetailCategoryTableRow.h"
@@ -878,6 +879,29 @@ TSharedPtr<FPropertyNode> FDetailItemNode::GetPropertyNode() const
 	return Customization.GetPropertyNode();
 }
 
+void FDetailItemNode::GetAllPropertyNodes(TArray<TSharedRef<FPropertyNode>>& OutNodes) const
+{
+	TSet<TSharedRef<FPropertyNode>> SeenNodes; // make's sure there aren't duplicates
+	if (const TSharedPtr<FPropertyNode> Node = GetPropertyNode())
+	{
+		SeenNodes.Add(Node.ToSharedRef());
+		OutNodes.Add(Node.ToSharedRef());
+	}
+
+	for (const TSharedPtr<IPropertyHandle>& CurPropertyHandle : Customization.GetWidgetRow().PropertyHandles)
+	{
+		const TSharedPtr<FPropertyHandleBase>& Handle = StaticCastSharedPtr<FPropertyHandleBase>(CurPropertyHandle);
+		if (const TSharedPtr<FPropertyNode> Node = Handle->GetPropertyNode())
+		{
+			if (!SeenNodes.Contains(Node.ToSharedRef()))
+			{
+				SeenNodes.Add(Node.ToSharedRef());
+				OutNodes.Add(Node.ToSharedRef());
+			}
+		}
+	}
+}
+
 TSharedPtr<IDetailPropertyRow> FDetailItemNode::GetRow() const
 {
 	if (Customization.IsValidCustomization() && Customization.PropertyRow.IsValid())
diff --git a/Engine/Source/Editor/PropertyEditor/Private/DetailItemNode.h b/Engine/Source/Editor/PropertyEditor/Private/DetailItemNode.h
--- a/Engine/Source/Editor/PropertyEditor/Private/DetailItemNode.h
+++ b/Engine/Source/Editor/PropertyEditor/Private/DetailItemNode.h
@@ -84,6 +84,7 @@ public:
 	virtual bool IsLeaf() override { return Children.Num() == 0;  }
 	virtual TAttribute<bool> IsPropertyEditingEnabled() const override;
 	virtual TSharedPtr<FPropertyNode> GetPropertyNode() const override;
+	virtual void GetAllPropertyNodes(TArray<TSharedRef<FPropertyNode>>& OutNodes) const override;
 	virtual TSharedPtr<FComplexPropertyNode> GetExternalRootPropertyNode() const override;
 	virtual TSharedPtr<IDetailPropertyRow> GetRow() const override;
 
diff --git a/Engine/Source/Editor/PropertyEditor/Public/DetailTreeNode.h b/Engine/Source/Editor/PropertyEditor/Public/DetailTreeNode.h
--- a/Engine/Source/Editor/PropertyEditor/Public/DetailTreeNode.h
+++ b/Engine/Source/Editor/PropertyEditor/Public/DetailTreeNode.h
@@ -166,6 +166,11 @@ public:
 	 * Gets the property node associated with this node.  Not all nodes have properties so this will fail for anything other than a property row or for property rows that have complex customizations that ignore the property
 	 */
 	virtual TSharedPtr<FPropertyNode> GetPropertyNode() const { return nullptr; }
+	
+	/** 
+	 * In certain circumstances a single DetailTreeNode can have multiple property nodes. This will get all of them.
+	 */
+	virtual void GetAllPropertyNodes(TArray<TSharedRef<FPropertyNode>>& OutNodes) const {}
 
 	/**
 	 * Gets the external property node associated with this node.  This will return nullptr for all rows expect property rows which were generated from an external root.
