Property Matrix: Fix a bug that stopped search from working in the "Pinned Columns" tab

diff --git a/Engine/Source/Editor/PropertyEditor/Private/SPropertyTreeViewImpl.cpp b/Engine/Source/Editor/PropertyEditor/Private/SPropertyTreeViewImpl.cpp
--- a/Engine/Source/Editor/PropertyEditor/Private/SPropertyTreeViewImpl.cpp
+++ b/Engine/Source/Editor/PropertyEditor/Private/SPropertyTreeViewImpl.cpp
@@ -664,14 +664,18 @@ void SPropertyTreeViewImpl::MarkFavoritesInternal( TSharedPtr<FPropertyNode> InP
 
 bool SPropertyTreeViewImpl::IsPropertyNodeVisible(TSharedPtr<FPropertyNode> InPropertyNode)
 {
-	bool bPropertyVisible = true;
-
 	// Object nodes always mark themselves as visible (but are never actually shown)
 	if(InPropertyNode->AsObjectNode())
 	{
 		return true;
 	}
 
+	// The node is marked not visible due to filtering etc
+	if(!InPropertyNode->IsVisible())
+	{
+		return false;
+	}
+	
 	// Category nodes are visible if they have any visible children
 	if(InPropertyNode->AsCategoryNode())
 	{
@@ -699,6 +703,9 @@ bool SPropertyTreeViewImpl::IsPropertyNodeVisible(TSharedPtr<FPropertyNode> InPr
 
 		return bHasVisibleChild;
 	}
+
+	// Regular property nodes are visible if they pass the IsPropertyVisible delegate and any permission lists if valid
+	bool bPropertyVisible = true;
 	
 	const FProperty* Property = InPropertyNode->GetProperty();
 	if(Property != NULL)
@@ -706,7 +713,7 @@ bool SPropertyTreeViewImpl::IsPropertyNodeVisible(TSharedPtr<FPropertyNode> InPr
 		// If we have a permission list in use, make sure this property is in the list for the most common base class and the class of all objects we are viewing
 		if(FPropertyEditorPermissionList::Get().IsEnabled())
 		{
-			bPropertyVisible = FPropertyEditorPermissionList::Get().DoesPropertyPassFilter(FDetailTreeNode::GetPropertyNodeBaseStructure(InPropertyNode->GetParentNode()), Property->GetFName());
+			bPropertyVisible &= FPropertyEditorPermissionList::Get().DoesPropertyPassFilter(FDetailTreeNode::GetPropertyNodeBaseStructure(InPropertyNode->GetParentNode()), Property->GetFName());
 
 			for (int32 ObjectIndex = 0; ObjectIndex < RootPropertyNode->GetNumObjects(); ++ObjectIndex)
 			{
