[CrashFix] SDetailsSplitter::GetHighlightColor causes crash if property is a sub-member of a null object
repro: review 26434801 and click on first entry in defaults diff

diff --git a/Engine/Source/Editor/PropertyEditor/Private/SDetailsSplitter.cpp b/Engine/Source/Editor/PropertyEditor/Private/SDetailsSplitter.cpp
--- a/Engine/Source/Editor/PropertyEditor/Private/SDetailsSplitter.cpp
+++ b/Engine/Source/Editor/PropertyEditor/Private/SDetailsSplitter.cpp
@@ -849,17 +849,19 @@ FLinearColor SDetailsSplitter::GetHighlightColor(const TUniquePtr<TDiffNode<TWea
 	
 	FPropertyPath Path = DetailNode->GetPropertyPath();
 	const UObject* OwningObject = DetailNode->GetDetailsView()->GetSelectedObjects()[0].Get();
-	if (Path.IsValid())
+	if (OwningObject && Path.IsValid())
 	{
 		FPropertyPath PathFromSubObject;
 		ShortenToPathFromLastObject(OwningObject, Path);
-		
-		if (const TMap<FPropertySoftPath, FLinearColor>* Highlights = CustomHighlights.Find(OwningObject->GetPathName(OwningObject->GetPackage())))
+		if (OwningObject)
 		{
-			if (const FLinearColor* Highlight = Highlights->Find(FPropertySoftPath(Path)))
-			{
-				return *Highlight;
-			}
+			if (const TMap<FPropertySoftPath, FLinearColor>* Highlights = CustomHighlights.Find(OwningObject->GetPathName(OwningObject->GetPackage())))
+            {
+            	if (const FLinearColor* Highlight = Highlights->Find(FPropertySoftPath(Path)))
+            	{
+            		return *Highlight;
+            	}
+            }
 		}
 	}
 	
