[PerfFix] Diff algorithm was constructing widgets unintentionally when attempting to get property handles from the DetailCategoryBuilder

diff --git a/Engine/Source/Editor/PropertyEditor/Private/DetailCategoryBuilderImpl.cpp b/Engine/Source/Editor/PropertyEditor/Private/DetailCategoryBuilderImpl.cpp
--- a/Engine/Source/Editor/PropertyEditor/Private/DetailCategoryBuilderImpl.cpp
+++ b/Engine/Source/Editor/PropertyEditor/Private/DetailCategoryBuilderImpl.cpp
@@ -123,6 +123,27 @@ FDetailWidgetRow FDetailLayoutCustomization::GetWidgetRow() const
 	}
 }
 
+TArrayView<TSharedPtr<IPropertyHandle>> FDetailLayoutCustomization::GetPropertyHandles() const
+{
+	if (HasCustomWidget())
+	{
+		return WidgetDecl->PropertyHandles;
+	}
+	else if (HasCustomBuilder())
+	{
+		return CustomBuilderRow->GetWidgetRow()->PropertyHandles;
+	}
+	else if (HasPropertyNode())
+	{
+		return PropertyRow->GetPropertyHandles();
+	}
+	else if (DetailGroup->GetHeaderPropertyRow())
+	{
+		return DetailGroup->GetHeaderPropertyRow()->GetPropertyHandles();
+	}
+	return TArrayView<TSharedPtr<IPropertyHandle>>();
+}
+
 const IDetailLayoutRow* FDetailLayoutCustomization::GetDetailLayoutRow() const
 {
 	if (HasCustomWidget())
diff --git a/Engine/Source/Editor/PropertyEditor/Private/DetailCategoryBuilderImpl.h b/Engine/Source/Editor/PropertyEditor/Private/DetailCategoryBuilderImpl.h
--- a/Engine/Source/Editor/PropertyEditor/Private/DetailCategoryBuilderImpl.h
+++ b/Engine/Source/Editor/PropertyEditor/Private/DetailCategoryBuilderImpl.h
@@ -56,6 +56,8 @@ struct FDetailLayoutCustomization
 	TSharedPtr<FPropertyNode> GetPropertyNode() const;
 	/** @return The row to display from this customization */
 	FDetailWidgetRow GetWidgetRow() const;
+	/** @return properties being customized */
+	TArrayView<TSharedPtr<IPropertyHandle>> GetPropertyHandles() const;
 	/** Whether or not this customization is considered an advanced property. */
 	bool bAdvanced { false };
 	/** Whether or not this customization is custom or a default one. */
diff --git a/Engine/Source/Editor/PropertyEditor/Private/DetailItemNode.cpp b/Engine/Source/Editor/PropertyEditor/Private/DetailItemNode.cpp
--- a/Engine/Source/Editor/PropertyEditor/Private/DetailItemNode.cpp
+++ b/Engine/Source/Editor/PropertyEditor/Private/DetailItemNode.cpp
@@ -888,7 +888,7 @@ void FDetailItemNode::GetAllPropertyNodes(TArray<TSharedRef<FPropertyNode>>& Out
 		OutNodes.Add(Node.ToSharedRef());
 	}
 
-	for (const TSharedPtr<IPropertyHandle>& CurPropertyHandle : Customization.GetWidgetRow().PropertyHandles)
+	for (const TSharedPtr<IPropertyHandle>& CurPropertyHandle : Customization.GetPropertyHandles())
 	{
 		const TSharedPtr<FPropertyHandleBase>& Handle = StaticCastSharedPtr<FPropertyHandleBase>(CurPropertyHandle);
 		if (const TSharedPtr<FPropertyNode> Node = Handle->GetPropertyNode())
diff --git a/Engine/Source/Editor/PropertyEditor/Private/DetailPropertyRow.cpp b/Engine/Source/Editor/PropertyEditor/Private/DetailPropertyRow.cpp
--- a/Engine/Source/Editor/PropertyEditor/Private/DetailPropertyRow.cpp
+++ b/Engine/Source/Editor/PropertyEditor/Private/DetailPropertyRow.cpp
@@ -231,6 +231,17 @@ FDetailWidgetRow FDetailPropertyRow::GetWidgetRow()
 	}
 }
 
+TArrayView<TSharedPtr<IPropertyHandle>> FDetailPropertyRow::GetPropertyHandles() const
+{
+	if (CustomPropertyWidget)
+	{
+		return CustomPropertyWidget->PropertyHandles;
+	}
+
+	// view single item as a c-array of size one
+	return TArrayView<TSharedPtr<IPropertyHandle>>(const_cast<TSharedPtr<IPropertyHandle>*>(&PropertyHandle), 1);
+}
+
 static bool IsHeaderRowRequired(const TSharedPtr<IPropertyHandle>& PropertyHandle)
 {
 	TSharedPtr<IPropertyHandle> ParentHandle = PropertyHandle->GetParentHandle();
diff --git a/Engine/Source/Editor/PropertyEditor/Private/DetailPropertyRow.h b/Engine/Source/Editor/PropertyEditor/Private/DetailPropertyRow.h
--- a/Engine/Source/Editor/PropertyEditor/Private/DetailPropertyRow.h
+++ b/Engine/Source/Editor/PropertyEditor/Private/DetailPropertyRow.h
@@ -86,6 +86,9 @@ public:
 	/** @return The widget row that should be displayed for this property row */
 	FDetailWidgetRow GetWidgetRow();
 
+	/** returns only the widget for editing this property. Use GetWidgetRow to get the full row with property name */
+	TArrayView<TSharedPtr<IPropertyHandle>> GetPropertyHandles() const;
+
 	/**
 	 * @return The property node for this row
 	 */
