[BugFix] False Positive Diffs for subobject pointers

diff --git a/Engine/Source/Editor/PropertyEditor/Private/AsyncDetailViewDiff.cpp b/Engine/Source/Editor/PropertyEditor/Private/AsyncDetailViewDiff.cpp
--- a/Engine/Source/Editor/PropertyEditor/Private/AsyncDetailViewDiff.cpp
+++ b/Engine/Source/Editor/PropertyEditor/Private/AsyncDetailViewDiff.cpp
@@ -61,14 +61,15 @@ bool TreeDiffSpecification::AreValuesEqual<TWeakPtr<FDetailTreeNode>>(const TWea
 		const FResolvedProperty ResolvedB = GetResolvedProperty(PropertyNodeB, OwningObjectB);
 		if (ResolvedA.Property && ResolvedB.Property)
 		{
-			// property nodes
 			if (!ResolvedA.Property->SameType(ResolvedB.Property))
 			{
 				return false;
 			}
-			const void* DataA = ResolvedA.Property->ContainerPtrToValuePtr<void*>(ResolvedA.Object);
-			const void* DataB = ResolvedB.Property->ContainerPtrToValuePtr<void*>(ResolvedB.Object);
-			if(!ResolvedA.Property->Identical(DataA, DataB, PPF_DeepComparison))
+
+			// use DiffUtils::Identical instead of FProperty::Identical so that sub-object pointers are considered identical
+			// if they're isomorphic rather than only if they have the same pointer values
+			TArray<FPropertySoftPath> Differing;
+			if(!DiffUtils::Identical(ResolvedA, ResolvedB, {}, Differing))
 			{
 				return false;
 			}
