Remove ensure as it is firing in -game and -server editor builds
#rnx

diff --git a/Engine/Source/Editor/PropertyEditor/Private/PropertyPermissionList.cpp b/Engine/Source/Editor/PropertyEditor/Private/PropertyPermissionList.cpp
--- a/Engine/Source/Editor/PropertyEditor/Private/PropertyPermissionList.cpp
+++ b/Engine/Source/Editor/PropertyEditor/Private/PropertyPermissionList.cpp
@@ -65,7 +65,7 @@ bool FPropertyPermissionList::Tick(float DeltaTime)
 
 void FPropertyPermissionList::RegisterOnBlueprintCompiled()
 {
-	if (ensure(GEditor))
+	if (GEditor)
 	{
 		GEditor->OnBlueprintCompiled().AddRaw(this, &FPropertyPermissionList::ClearCache);
 	}
