The object path used to set the properties via an asset data in FPropertyHandleObject::SetValue will now use the full path of the class instead of only the name.

Since the killpendingkill initiave we can have mutiple UClass objects that share the same name but have diferent paths. There is no value to only pass the name of class since it cannot be use to find class since it might yield multiple classes.

This also remove an issue when parsing the string since the full path of an class cannot start with an '_' while a class object can. This was an issue since the parser is not able to recognize a token that start with an '_'.

diff --git a/Engine/Source/Editor/PropertyEditor/Private/PropertyHandleImpl.cpp b/Engine/Source/Editor/PropertyEditor/Private/PropertyHandleImpl.cpp
--- a/Engine/Source/Editor/PropertyEditor/Private/PropertyHandleImpl.cpp
+++ b/Engine/Source/Editor/PropertyEditor/Private/PropertyHandleImpl.cpp
@@ -4065,7 +4065,7 @@ FPropertyAccess::Result FPropertyHandleObject::SetValue(const FAssetData& NewVal
 			NewValue.GetAsset();
 		}
 
-		FString ObjectPathName = NewValue.IsValid() ? FString::Printf(TEXT("%s'%s'"), *NewValue.GetClass()->GetName(), *NewValue.GetObjectPathString()) : TEXT("None");
+		FString ObjectPathName = NewValue.IsValid() ? FString::Printf(TEXT("%s'%s'"), *NewValue.AssetClassPath.ToString(), *NewValue.GetObjectPathString()) : TEXT("None");
 		return SetValueFromFormattedString(ObjectPathName, Flags, bSkipResolve);
 	}
 
