Changed SObjectPropertyEntryBox to take a const UClass* instead of UClass*.

diff --git a/Engine/Source/Editor/PropertyEditor/Private/UserInterface/PropertyEditor/SPropertyEditorAsset.cpp b/Engine/Source/Editor/PropertyEditor/Private/UserInterface/PropertyEditor/SPropertyEditorAsset.cpp
--- a/Engine/Source/Editor/PropertyEditor/Private/UserInterface/PropertyEditor/SPropertyEditorAsset.cpp
+++ b/Engine/Source/Editor/PropertyEditor/Private/UserInterface/PropertyEditor/SPropertyEditorAsset.cpp
@@ -830,7 +830,7 @@ FText SPropertyEditorAsset::OnGetAssetName() const
 
 FText SPropertyEditorAsset::OnGetAssetClassName() const
 {
-	UClass* Class = GetDisplayedClass();
+	const UClass* Class = GetDisplayedClass();
 	if(Class)
 	{
 		return FText::AsCultureInvariant(Class->GetName());
@@ -1073,7 +1073,7 @@ FPropertyAccess::Result SPropertyEditorAsset::GetValue( FObjectOrAssetData& OutV
 	return Result;
 }
 
-UClass* SPropertyEditorAsset::GetDisplayedClass() const
+const UClass* SPropertyEditorAsset::GetDisplayedClass() const
 {
 	FObjectOrAssetData Value;
 	GetValue( Value );
@@ -1217,7 +1217,7 @@ void SPropertyEditorAsset::OnUse()
 		FEditorDelegates::LoadSelectedAssetsIfNeeded.Broadcast();
 
 		// try to get a selected object of our class
-		UObject* Selection = nullptr;
+		const UObject* Selection = nullptr;
 		if( ObjectClass && ObjectClass->IsChildOf( AActor::StaticClass() ) )
 		{
 			Selection = GEditor->GetSelectedActors()->GetTop( ObjectClass );
diff --git a/Engine/Source/Editor/PropertyEditor/Private/UserInterface/PropertyEditor/SPropertyEditorAsset.h b/Engine/Source/Editor/PropertyEditor/Private/UserInterface/PropertyEditor/SPropertyEditorAsset.h
--- a/Engine/Source/Editor/PropertyEditor/Private/UserInterface/PropertyEditor/SPropertyEditorAsset.h
+++ b/Engine/Source/Editor/PropertyEditor/Private/UserInterface/PropertyEditor/SPropertyEditorAsset.h
@@ -74,7 +74,7 @@ public:
 		SLATE_ARGUMENT(TSharedPtr<FAssetThumbnailPool>, ThumbnailPool)
 		SLATE_ARGUMENT(FIntPoint, ThumbnailSize)
 		SLATE_ATTRIBUTE(FString, ObjectPath)
-		SLATE_ARGUMENT(UClass*, Class)
+		SLATE_ARGUMENT(const UClass*, Class)
 		SLATE_ARGUMENT(TOptional<TArray<UFactory*>>, NewAssetFactories)
 		SLATE_EVENT(FOnSetObject, OnSetObject)
 		SLATE_EVENT(FOnShouldFilterAsset, OnShouldFilterAsset)
@@ -221,7 +221,7 @@ private:
 	 * or the Class value this widget was constructed with.
 	 * @returns the UClass to display
 	 */
-	UClass* GetDisplayedClass() const;
+	const UClass* GetDisplayedClass() const;
 
 	/** 
 	 * Delegate for handling selection in the asset browser.
@@ -377,7 +377,7 @@ private:
 	mutable FAssetData CachedAssetData;
 
 	/** The class of the object we are editing */
-	UClass* ObjectClass;
+	const UClass* ObjectClass;
 
 	/** Classes that can be used with this property */
 	TArray<const UClass*> AllowedClassFilters;
diff --git a/Engine/Source/Editor/PropertyEditor/Public/PropertyCustomizationHelpers.h b/Engine/Source/Editor/PropertyEditor/Public/PropertyCustomizationHelpers.h
--- a/Engine/Source/Editor/PropertyEditor/Public/PropertyCustomizationHelpers.h
+++ b/Engine/Source/Editor/PropertyEditor/Public/PropertyCustomizationHelpers.h
@@ -191,7 +191,7 @@ public:
 		/** Thumbnail pool */
 		SLATE_ARGUMENT( TSharedPtr<FAssetThumbnailPool>, ThumbnailPool )
 		/** Class that is allowed in the asset picker */
-		SLATE_ARGUMENT( UClass*, AllowedClass )
+		SLATE_ARGUMENT( const UClass*, AllowedClass )
 		/** Optional list of factories which may be used to create new assets */
 		SLATE_ARGUMENT( TOptional<TArray<UFactory*>>, NewAssetFactories )
 		/** Called to check if an asset should be set */
diff --git a/Engine/Source/Editor/UnrealEd/Public/Selection.h b/Engine/Source/Editor/UnrealEd/Public/Selection.h
--- a/Engine/Source/Editor/UnrealEd/Public/Selection.h
+++ b/Engine/Source/Editor/UnrealEd/Public/Selection.h
@@ -166,7 +166,7 @@ public:
 	 * @param	bArchetypesOnly		[opt] true to only return archetype objects, false otherwise
 	 * @return						The first selected object of the specified class.
 	 */
-	UObject* GetTop(UClass* InClass, UClass* RequiredInterface=nullptr, bool bArchetypesOnly=false)
+	UObject* GetTop(const UClass* InClass, const UClass* RequiredInterface=nullptr, bool bArchetypesOnly=false)
 	{
 		check( InClass );
 		for( int32 i=0; i<Num(); ++i )
@@ -207,7 +207,7 @@ public:
 	* @param	InClass		The class of object to return.  Must be non-NULL.
 	* @return				The last selected object of the specified class.
 	*/
-	UObject* GetBottom(UClass* InClass)
+	UObject* GetBottom(const UClass* InClass)
 	{
 		check( InClass );
 		for( int32 i = Num()-1 ; i > -1 ; --i )
