Property Matrix: Fix a crash when you tried to edit a property on a light component because the table was being grabbed through the column which could be in the middle of reconstruction

diff --git a/Engine/Source/Editor/PropertyEditor/Private/UserInterface/PropertyTable/SPropertyTableCell.cpp b/Engine/Source/Editor/PropertyEditor/Private/UserInterface/PropertyTable/SPropertyTableCell.cpp
--- a/Engine/Source/Editor/PropertyEditor/Private/UserInterface/PropertyTable/SPropertyTableCell.cpp
+++ b/Engine/Source/Editor/PropertyEditor/Private/UserInterface/PropertyTable/SPropertyTableCell.cpp
@@ -11,6 +11,7 @@
 #include "IPropertyTable.h"
 #include "IPropertyTableCell.h"
 #include "IPropertyTableCellPresenter.h"
+#include "IPropertyTableRow.h"
 #include "Input/Events.h"
 #include "Internationalization/Internationalization.h"
 #include "Layout/Children.h"
@@ -167,12 +168,13 @@ int32 SPropertyTableCell::OnPaint( const FPaintArgs& Args, const FGeometry& Allo
 	{
 		const FSlateBrush* Background = CellBackground;
 
-		if ( Cell->GetTable()->GetCurrentCell() == Cell )
+		// We go through the row to grab the table because we are in OnPaint() and the column could be in the process of reconstruction
+		if (Cell->GetRow()->GetTable()->GetCurrentCell() == Cell )
 		{
 			// The current cell has the blue selection color and a white border around it
 			Background = FAppStyle::GetBrush( Style, ".CurrentCellBorder" );
 		}
-		else if ( Cell->GetTable()->GetSelectedCells().Contains( Cell.ToSharedRef() ) )
+		else if ( Cell->GetRow()->GetTable()->GetSelectedCells().Contains( Cell.ToSharedRef() ) )
 		{
 			// Other selected cells have the blue selection color along with the gridlines
 			Background = FAppStyle::GetBrush( Style, ".SelectedCellBorder" );
