Lazy refresh of the tree for a big speedup when filtering in the Details pane.

#rnx
#p4v-preflight-copy 26987960

(checking in directly as Horde / P4 is down - preflights completed successfully https://horde.devtools.epicgames.com/job/64d4ddff7f3f866457cf7959)

diff --git a/Engine/Source/Editor/PropertyEditor/Private/DetailCategoryBuilderImpl.cpp b/Engine/Source/Editor/PropertyEditor/Private/DetailCategoryBuilderImpl.cpp
--- a/Engine/Source/Editor/PropertyEditor/Private/DetailCategoryBuilderImpl.cpp
+++ b/Engine/Source/Editor/PropertyEditor/Private/DetailCategoryBuilderImpl.cpp
@@ -199,6 +199,8 @@ FDetailCategoryImpl::FDetailCategoryImpl(FName InCategoryName, TSharedRef<FDetai
 	, bFavoriteCategory(false)
 	, bShowOnlyChildren(false)
 	, bHasVisibleAdvanced(false)
+	, bPendingRefresh(false)
+	, bPendingRefreshNeedsRefilter(false)
 {
 	const UStruct* BaseStruct = InDetailLayout->GetRootNode()->GetBaseStructure();
 
@@ -657,7 +659,22 @@ void FDetailCategoryImpl::RequestItemExpanded(TSharedRef<FDetailTreeNode> TreeNo
 	}
 }
 
+void FDetailCategoryImpl::Tick(float DeltaTime)
+{
+	ensure(bPendingRefresh);
+	RefreshTreeInternal(bPendingRefreshNeedsRefilter);
+	bPendingRefresh = false;
+	RemoveTickableNode(*this);
+}
+
 void FDetailCategoryImpl::RefreshTree(bool bRefilterCategory)
+{
+	bPendingRefresh = true;
+	bPendingRefreshNeedsRefilter = bRefilterCategory;
+	AddTickableNode(*this);
+}
+
+void FDetailCategoryImpl::RefreshTreeInternal(bool bRefilterCategory)
 {
 	if (bRefilterCategory)
 	{
diff --git a/Engine/Source/Editor/PropertyEditor/Private/DetailCategoryBuilderImpl.h b/Engine/Source/Editor/PropertyEditor/Private/DetailCategoryBuilderImpl.h
--- a/Engine/Source/Editor/PropertyEditor/Private/DetailCategoryBuilderImpl.h
+++ b/Engine/Source/Editor/PropertyEditor/Private/DetailCategoryBuilderImpl.h
@@ -208,7 +208,7 @@ public:
 	virtual bool ShouldBeExpanded() const override;
 	virtual ENodeVisibility GetVisibility() const override;
 	virtual void FilterNode(const FDetailFilter& DetailFilter) override;
-	virtual void Tick(float DeltaTime) override {}
+	virtual void Tick(float DeltaTime) override;
 	virtual bool ShouldShowOnlyChildren() const override { return bShowOnlyChildren; }
 	virtual FName GetNodeName() const override { return GetCategoryName(); }
 
@@ -392,6 +392,13 @@ private:
 	 */
 	bool IsParentEnabled() const;
 
+	/**
+	 * Does the work of refreshing the tree (this is triggered from RefreshTree)
+	 *
+	 * @param bRefilterCategory True if the category should be refiltered
+	 */
+	void RefreshTreeInternal(bool bRefilterCategory);
+
 private:
 	/** Layouts that appear in this category category */
 	FDetailLayoutMap LayoutMap;
@@ -438,4 +445,6 @@ private:
 	bool bFavoriteCategory : 1;
 	bool bShowOnlyChildren : 1;
 	bool bHasVisibleAdvanced : 1;
+	bool bPendingRefresh : 1;
+	bool bPendingRefreshNeedsRefilter : 1;
 };
