Slate: SClassPropertyEntryBox: Add optional class filters

diff --git a/Engine/Source/Editor/PropertyEditor/Private/PropertyCustomizationHelpers.cpp b/Engine/Source/Editor/PropertyEditor/Private/PropertyCustomizationHelpers.cpp
--- a/Engine/Source/Editor/PropertyEditor/Private/PropertyCustomizationHelpers.cpp
+++ b/Engine/Source/Editor/PropertyEditor/Private/PropertyCustomizationHelpers.cpp
@@ -617,6 +617,7 @@ void SClassPropertyEntryBox::Construct(const FArguments& InArgs)
 				.ShowTree(InArgs._ShowTreeView)
 				.SelectedClass(InArgs._SelectedClass)
 				.OnSetClass(InArgs._OnSetClass)
+				.ClassViewerFilters(InArgs._ClassViewerFilters)
 		]
 	];
 }
diff --git a/Engine/Source/Editor/PropertyEditor/Public/PropertyCustomizationHelpers.h b/Engine/Source/Editor/PropertyEditor/Public/PropertyCustomizationHelpers.h
--- a/Engine/Source/Editor/PropertyEditor/Public/PropertyCustomizationHelpers.h
+++ b/Engine/Source/Editor/PropertyEditor/Public/PropertyCustomizationHelpers.h
@@ -23,20 +23,21 @@
 #include "IDetailPropertyRow.h"
 
 class AActor;
-struct FAssetData;
 class FAssetThumbnailPool;
 class FPropertyEditor;
+class IClassViewerFilter;
+class IDetailCategoryBuilder;
 class IDetailChildrenBuilder;
+class IDetailGroup;
 class IDetailLayoutBuilder;
+class IPropertyHandle;
 class SPropertyEditorAsset;
 class SPropertyEditorClass;
 class SPropertyEditorStruct;
+class SToolTip;
 class UActorComponent;
 class UFactory;
-class SToolTip;
-class IPropertyHandle;
-class IDetailGroup;
-class IDetailCategoryBuilder;
+struct FAssetData;
 struct FSceneOutlinerFilters;
 
 DECLARE_DELEGATE_OneParam(FOnAssetSelected, const FAssetData& /*AssetData*/);
@@ -301,6 +302,8 @@ public:
 		SLATE_ATTRIBUTE(const UClass*, SelectedClass)
 		/** Delegate used to set the currently selected class (required) */
 		SLATE_EVENT(FOnSetClass, OnSetClass)
+		/** Custom class filter(s) to be applied on the derived classes of the Metaclass (may be empty)*/
+		SLATE_ARGUMENT(TArray<TSharedRef<IClassViewerFilter>>, ClassViewerFilters)	
 	SLATE_END_ARGS()
 
 	PROPERTYEDITOR_API void Construct(const FArguments& InArgs);
