Feature:
* The details panel must be enhanced to show Components as cards
* This new Card style should be applicable to any details panel
* Subobject editor had to be hidden upon selection of entity

diff --git a/Engine/Source/Editor/EditorStyle/Private/StarshipStyle.cpp b/Engine/Source/Editor/EditorStyle/Private/StarshipStyle.cpp
--- a/Engine/Source/Editor/EditorStyle/Private/StarshipStyle.cpp
+++ b/Engine/Source/Editor/EditorStyle/Private/StarshipStyle.cpp
@@ -2995,6 +2995,18 @@ void FStarshipEditorStyle::FStyle::SetupPropertyEditorStyles()
 		Set( "DetailsView.CollapsedCategory",         new FSlateColorBrush(FStyleColors::Header));
 		Set( "DetailsView.CollapsedCategory_Hovered", new FSlateColorBrush(FStyleColors::Hover));
 		Set( "DetailsView.CategoryTop",               new FSlateColorBrush(FStyleColors::Header));
+		
+		/****** Styles for rounded corners for the Card style of a Details View ********/
+
+		Set( "DetailsView.CardHeaderTopLeftSideRounded",               new FSlateRoundedBoxBrush(FStyleColors::Header, FVector4(4.0f, 0.0f, 0.0f, 0.0f)));
+		Set( "DetailsView.CardHeaderLeftSideRounded",               new FSlateRoundedBoxBrush(FStyleColors::Header, FVector4(4.0f, 0.0f, 0.0f, 4.0f)));
+		Set( "DetailsView.CardHeaderTopRightSideRounded",               new FSlateRoundedBoxBrush(FStyleColors::Header, FVector4(0.0f, 4.0f, 0.0f, 0.0f)));
+		Set( "DetailsView.CardHeaderRightSideRounded",               new FSlateRoundedBoxBrush(FStyleColors::Header, FVector4(0.0f, 4.0f, 4.0f, 0.0f)));
+		Set( "DetailsView.CardHeaderTopRounded",               new FSlateRoundedBoxBrush(FStyleColors::Header, FVector4(4.0f, 4.0f, 0.0f, 0.0f)));
+		Set( "DetailsView.CardHeaderRounded",               new FSlateRoundedBoxBrush(FStyleColors::Header, FVector4(4.0f, 4.0f, 4.0f, 4.0f )));
+		
+		/*******************************************************************************/
+		
 		Set( "DetailsView.CategoryTop_Hovered",       new FSlateColorBrush(FStyleColors::Hover));
 		Set( "DetailsView.CategoryBottom",            new FSlateColorBrush(FStyleColors::Recessed));
 		
diff --git a/Engine/Source/Editor/LevelEditor/Private/SActorDetails.cpp b/Engine/Source/Editor/LevelEditor/Private/SActorDetails.cpp
--- a/Engine/Source/Editor/LevelEditor/Private/SActorDetails.cpp
+++ b/Engine/Source/Editor/LevelEditor/Private/SActorDetails.cpp
@@ -280,7 +280,7 @@ void SActorDetails::Construct(const FArguments& InArgs, UTypedElementSelectionSe
 	ComponentsBox->SetContent(SubobjectEditor.ToSharedRef());
 
 	TSharedRef<SWidget> SubobjectEditorButtonBox = SubobjectEditor->GetToolButtonsBox().ToSharedRef();
-	SubobjectEditorButtonBox->SetVisibility(MakeAttributeSP(this, &SActorDetails::GetComponentEditorVisibility));
+	SubobjectEditorButtonBox->SetVisibility(MakeAttributeSP(this, &SActorDetails::GetComponentEditorButtonsVisibility));
 
 
 	TFunction<TArray<TTypedElement<ITypedElementDetailsInterface>>()> GetDetailsHandles = [this]()
@@ -879,6 +879,12 @@ void SActorDetails::OnNativeComponentWarningHyperlinkClicked(const FSlateHyperli
 }
 
 EVisibility SActorDetails::GetComponentEditorVisibility() const
+{
+	const bool bHideEditorFromDetailsView = DetailsView.IsValid() &&  !DetailsView->IsDefaultStyle();
+	return GetActorContext() && !bHideEditorFromDetailsView  ? EVisibility::Visible : EVisibility::Collapsed;
+}
+
+EVisibility SActorDetails::GetComponentEditorButtonsVisibility() const
 {
 	return GetActorContext() ? EVisibility::Visible : EVisibility::Collapsed;
 }
diff --git a/Engine/Source/Editor/LevelEditor/Private/SActorDetails.h b/Engine/Source/Editor/LevelEditor/Private/SActorDetails.h
--- a/Engine/Source/Editor/LevelEditor/Private/SActorDetails.h
+++ b/Engine/Source/Editor/LevelEditor/Private/SActorDetails.h
@@ -88,6 +88,11 @@ private:
 	UObject* GetActorContextAsObject() const;
 	bool GetAllowComponentTreeEditing() const;
 
+	/**
+	 * Gets the Visibility of the top of the panel buttons in the details panel
+	 */
+	EVisibility GetComponentEditorButtonsVisibility() const;
+
 	void OnComponentsEditedInWorld();
 	void OnSubobjectEditorTreeViewSelectionChanged(const TArray<TSharedPtr<FSubobjectEditorTreeNode> >& SelectedNodes);
 	void OnSubobjectEditorTreeViewItemDoubleClicked(const TSharedPtr<FSubobjectEditorTreeNode> ClickedNode);
diff --git a/Engine/Source/Editor/PropertyEditor/Private/DetailAdvancedDropdownNode.cpp b/Engine/Source/Editor/PropertyEditor/Private/DetailAdvancedDropdownNode.cpp
--- a/Engine/Source/Editor/PropertyEditor/Private/DetailAdvancedDropdownNode.cpp
+++ b/Engine/Source/Editor/PropertyEditor/Private/DetailAdvancedDropdownNode.cpp
@@ -9,6 +9,8 @@
 #include "UserInterface/PropertyEditor/PropertyEditorConstants.h"
 #include "Widgets/Images/SImage.h"
 #include "Widgets/Input/SButton.h"
+#include "DetailsViewStyle.h"
+#include "SDetailsView.h"
 
 class SAdvancedDropdownRow : public SDetailTableRowBase
 {
@@ -85,26 +87,29 @@ public:
 			]
 		];
 
-		TWeakPtr<STableViewBase> OwnerTableViewWeak = InOwnerTableView;
-		auto GetScrollbarWellBrush = [this, OwnerTableViewWeak]()
+		OwnerTableViewWeak = InOwnerTableView;
+		auto GetScrollbarWellBrush = [this]()
 		{
 			return SDetailTableRowBase::IsScrollBarVisible(OwnerTableViewWeak) ?
 				FAppStyle::Get().GetBrush("DetailsView.GridLine") : 
 				FAppStyle::Get().GetBrush("DetailsView.CategoryMiddle");
 		};
 
-		auto GetScrollbarWellTint = [this, OwnerTableViewWeak]()
+		auto GetScrollbarWellTint = [this]()
 		{
-			return SDetailTableRowBase::IsScrollBarVisible(OwnerTableViewWeak) ?
+			return IsScrollBarVisible(OwnerTableViewWeak) ?
 				FStyleColors::White : 
 				this->GetRowBackgroundColor();
 		};
 
+		static FDetailsViewStyleKey PrimaryKey = SDetailsView::PrimaryDetailsViewStyleKey;
+		FDetailsViewStyle ViewStyle = DetailsView ? DetailsView->GetStyleKey() : PrimaryKey;
+
 		ChildSlot
 		[
 			SNew(SBorder)
 			.BorderImage(FAppStyle::Get().GetBrush("DetailsView.GridLine"))
-			.Padding(FMargin(0, 0, 0, 1))
+			.Padding(ViewStyle.GetRowPadding())
 			[
 				SNew(SHorizontalBox)
 				+ SHorizontalBox::Slot()
diff --git a/Engine/Source/Editor/PropertyEditor/Private/DetailsViewStyle.cpp b/Engine/Source/Editor/PropertyEditor/Private/DetailsViewStyle.cpp
new file mode 100644
--- /dev/null
+++ b/Engine/Source/Editor/PropertyEditor/Private/DetailsViewStyle.cpp
@@ -0,0 +1,186 @@
+//  Copyright Epic Games, Inc. All Rights Reserved.
+
+
+#include "DetailsViewStyle.h"
+
+FDetailsViewStyle::FDetailsViewStyle(
+	const FDetailsViewStyleKey& InKey, 
+	float InTopCategoryPadding,
+	float InHorizontalPadding) :
+		FSlateWidgetStyle(),
+		Key(InKey),
+		TopCategoryPadding(InTopCategoryPadding),
+		HorizontalPadding(InHorizontalPadding)
+{
+}
+
+FDetailsViewStyle::FDetailsViewStyle(FDetailsViewStyleKey& InKey) :
+		FSlateWidgetStyle(),
+		Key(InKey),
+		TopCategoryPadding(0.0f),
+		HorizontalPadding(0.0f)
+{
+	if (const FDetailsViewStyle* Style = GetStyle(InKey))
+	{
+		Initialize(InKey, Style->HorizontalPadding, Style->TopCategoryPadding);
+	}
+}
+
+FDetailsViewStyle::FDetailsViewStyle(FDetailsViewStyle& InStyle) :
+	FDetailsViewStyle(InStyle.Key, InStyle.TopCategoryPadding, InStyle.HorizontalPadding) 
+{
+}
+
+FDetailsViewStyle::FDetailsViewStyle(const FDetailsViewStyle& InStyle) :
+	FDetailsViewStyle(InStyle.Key, InStyle.TopCategoryPadding, InStyle.HorizontalPadding) 
+{
+}
+
+FMargin FDetailsViewStyle::GetOuterCategoryRowPadding() const
+{
+	return FMargin(HorizontalPadding, TopCategoryPadding, HorizontalPadding, 1);
+}
+
+FMargin FDetailsViewStyle::GetRowPadding() const
+{
+	return bIsOuterCategory ?
+		FMargin(HorizontalPadding, TopCategoryPadding, HorizontalPadding, 1) :
+		FMargin(HorizontalPadding, 0, HorizontalPadding, 1);
+}
+
+bool FDetailsViewStyle::operator==(FDetailsViewStyle& OtherLayoutType) const
+{
+	return Key == OtherLayoutType.Key;
+}
+
+FDetailsViewStyle& FDetailsViewStyle::operator=(FDetailsViewStyleKey& OtherLayoutTypeKey)
+{
+	if (const FDetailsViewStyle* Style = GetStyle(OtherLayoutTypeKey))
+	{
+		Initialize(OtherLayoutTypeKey, Style->HorizontalPadding, Style->TopCategoryPadding);
+	}
+	return *this;
+}
+
+const FName FDetailsViewStyle::GetTypeName() const
+{
+	return Key.GetName();
+}
+
+const FSlateBrush* FDetailsViewStyle::GetBackgroundImageForCategoryRow(
+	const bool bShowBorder, 
+	const bool bIsInnerCategory,
+	const bool bIsCategoryExpanded,
+	const bool bIsScrollBarVisible) const
+{
+	static const FSlateBrush* InnerCategoryRowBrush = FAppStyle::Get().GetBrush("DetailsView.CategoryMiddle");
+	static const FSlateBrush* ClassicStyleTopLevelCategoryRowBrush = FAppStyle::Get().GetBrush("DetailsView.CategoryTop");
+	static const FSlateBrush* CardStyleTopLevelCategoryCollapsedScrollbarVisibleRowBrush = FAppStyle::Get().GetBrush("DetailsView.CardHeaderRounded");
+	static const FSlateBrush* CardStyleTopLevelCategoryCollapsedScrollbarHiddenRowBrush = FAppStyle::Get().GetBrush("DetailsView.CardHeaderLeftSideRounded");
+	static const FSlateBrush* CardStyleTopLevelCategoryExpandedScrollbarVisibleRowBrush = FAppStyle::Get().GetBrush("DetailsView.CardHeaderTopRounded");
+	static const FSlateBrush* CardStyleTopLevelCategoryExpandedScrollbarHiddenRowBrush = FAppStyle::Get().GetBrush("DetailsView.CardHeaderTopLeftSideRounded");
+
+	if (bShowBorder)
+	{
+		if (bIsInnerCategory)
+		{
+			return InnerCategoryRowBrush;
+		}
+			
+		const bool bIsCardStyle = Key == FDetailsViewStyleKeys::Card;
+			
+		if (!bIsCardStyle)
+		{
+			return ClassicStyleTopLevelCategoryRowBrush;
+		}
+		if (!bIsCategoryExpanded)
+		{
+			if (bIsScrollBarVisible)
+			{
+				return CardStyleTopLevelCategoryCollapsedScrollbarVisibleRowBrush;
+			}
+			return CardStyleTopLevelCategoryCollapsedScrollbarHiddenRowBrush;
+		}
+		if (bIsScrollBarVisible)
+		{
+			return CardStyleTopLevelCategoryExpandedScrollbarVisibleRowBrush;
+		}
+		return CardStyleTopLevelCategoryExpandedScrollbarHiddenRowBrush;
+	}
+
+	return nullptr;
+}
+
+void FDetailsViewStyle::InitializeDetailsViewStyles()
+{
+	static FDetailsViewStyle CardStyle{FDetailsViewStyleKeys::Card, 8.0f, 8.0f};
+	StyleKeyToStyleTemplateMap.Add(FDetailsViewStyleKeys::Card.GetName(), &CardStyle);
+
+	static FDetailsViewStyle DefaultStyle{FDetailsViewStyleKeys::Default };
+	StyleKeyToStyleTemplateMap.Add(FDetailsViewStyleKeys::Default.GetName(), &DefaultStyle);
+	
+}
+
+const FSlateBrush* FDetailsViewStyle::GetBackgroundImageForScrollBarWell(
+	const bool bShowBorder,
+	const bool bIsInnerCategory,
+	const bool bIsCategoryExpanded,
+	const bool bIsScrollBarVisible) const
+{
+	static const FSlateBrush* InnerCategoryWellBrush = FAppStyle::Get().GetBrush("DetailsView.CategoryMiddle");
+	static const FSlateBrush* ClassicStyleTopLevelCategoryRowBrush = FAppStyle::Get().GetBrush("DetailsView.CategoryTop");
+	static const FSlateBrush* DetailsBackgroundWellBrush = FAppStyle::Get().GetBrush("DetailsView.GridLine");
+	static const FSlateBrush* CardStyleCollapsedScrollbarVisibleWellBrush = FAppStyle::Get().GetBrush("DetailsView.CardHeaderRightSideRounded");
+	static const FSlateBrush* CardStyleExpandedScrollbarVisibleWellBrush = FAppStyle::Get().GetBrush("DetailsView.CardHeaderTopRightSideRounded");
+
+	if (bShowBorder)
+	{
+		if (bIsInnerCategory)
+		{
+			return InnerCategoryWellBrush;
+		}
+		const bool bIsCardStyle = this->Key == FDetailsViewStyleKeys::Card;
+			
+		if (!bIsScrollBarVisible)
+		{
+			if (!bIsCardStyle)
+			{
+				return ClassicStyleTopLevelCategoryRowBrush;
+			}
+			if (!bIsCategoryExpanded){
+				return CardStyleCollapsedScrollbarVisibleWellBrush;
+			}
+			return CardStyleExpandedScrollbarVisibleWellBrush;
+		}
+		return DetailsBackgroundWellBrush;
+	}
+
+	return nullptr;
+}
+
+void FDetailsViewStyle::Initialize(
+	FDetailsViewStyleKey& InKey,
+	const float InHorizontalPadding,
+	const float InTopCategoryPadding)
+{
+	Key = InKey;
+	HorizontalPadding = InHorizontalPadding;
+	TopCategoryPadding = InTopCategoryPadding;
+}
+
+const FDetailsViewStyle* FDetailsViewStyle::GetStyle(FDetailsViewStyleKey& Key)
+{
+	const FDetailsViewStyle** StylePtr;
+		
+	StylePtr = StyleKeyToStyleTemplateMap.Find(Key.GetName());
+	StylePtr = StylePtr ?
+					StylePtr :
+					StyleKeyToStyleTemplateMap.Find(FDetailsViewStyleKeys::Default.GetName());
+		
+	return *StylePtr;
+}
+
+void FDetailsViewStyle::SetIsOuterCategory(bool bInIsOuterCategory)
+{
+	bIsOuterCategory = bInIsOuterCategory;
+}
\ No newline at end of file
diff --git a/Engine/Source/Editor/PropertyEditor/Private/DetailsViewStyle.h b/Engine/Source/Editor/PropertyEditor/Private/DetailsViewStyle.h
new file mode 100644
--- /dev/null
+++ b/Engine/Source/Editor/PropertyEditor/Private/DetailsViewStyle.h
@@ -0,0 +1,162 @@
+//  Copyright Epic Games, Inc. All Rights Reserved.
+
+#pragma once
+#include "DetailsViewStyleKey.h"
+#include "Layout/Margin.h"
+#include "Styling/AppStyle.h"
+#include "Styling/SlateWidgetStyle.h"
+#include "UObject/NameTypes.h"
+
+/**
+ * A Class which holds information regarding the style of a Details View
+ */
+class FDetailsViewStyle : public FSlateWidgetStyle
+{
+public:
+
+	/**
+	 * The default constructor of this @code FDetailsViewStyle @endcode.
+	 */
+	FDetailsViewStyle();
+	
+	/**
+	 * Constructs this @code FDetailsViewStyle @endcode with the specified parameters
+	 *
+	 * @param InKey the Key for  the @code FDetailsViewStyle @endcode
+	 * @param InTopCategoryPadding the top padding for the Category row
+	 * @param InHorizontalPadding the horizontal padding for the Category row
+	 */
+	FDetailsViewStyle(FDetailsViewStyleKey& InKey);
+
+	/**
+	 * Copy constructor for @code FDetailsViewStyle @endcode.
+	 *
+	 * @param InStyle the @code FDetailsViewStyle @endcode instance from which to initialize
+	 * this @code FDetailsViewStyle @endcode.
+	 */
+	FDetailsViewStyle(FDetailsViewStyle& InStyle);
+
+	/**
+	 * Const copy constructor for @code FDetailsViewStyle @endcode.
+	 *
+	 * @param InStyle the @code FDetailsViewStyle @endcode instance from which to initialize
+	 * this @code FDetailsViewStyle @endcode.
+	 */
+	FDetailsViewStyle(const FDetailsViewStyle& InStyle);
+
+	/**
+	 * Returns the padding for the outer Category row
+	 */
+	 FMargin GetOuterCategoryRowPadding() const;
+	
+	/**
+	* Sets a bool indicating whether or not the Category this is a Style for is an outer Category
+	*
+	* @param bIsOuterCategory a bool indicating whether or not the Category this is a Style for is an outer Category
+	*/
+	void SetIsOuterCategory(bool bIsOuterCategory);
+
+	/**
+	 * Returns the padding for details panel rows which are not outer Category rows
+	 */
+	FMargin GetRowPadding() const;
+
+	/**
+	 * The equality operator for @code FDetailsViewStule @endcode
+	 */
+	bool operator==(FDetailsViewStyle& OtherLayoutType) const;
+
+	/**
+	 * The assignment operator for @code FDetailsViewStule @endcode
+	 */
+	FDetailsViewStyle& operator=(FDetailsViewStyleKey& OtherLayoutTypeKey);
+	
+	/**
+   	 * Gets the @code FName @endcode Name of this style.
+   	 */
+	virtual const FName GetTypeName() const override;
+
+	/**
+	 * Returns the background image for the Category row, minus the scrollbar well
+	 *
+ 	 * @param bShowBorder a bool to indicate whether the border should be shown, at all
+	 * @param bIsCategoryExpanded a bool that indicates whether this Category is expanded 
+	 * @param bIsScrollBarVisible a bool that indicates whether the scrollbar for the details view is visible
+	 */	
+	const FSlateBrush* GetBackgroundImageForCategoryRow(
+		const bool bShowBorder,
+		const bool bIsInnerCategory,
+		const bool bIsCategoryExpanded,
+		const bool bIsScrollBarVisible) const;
+
+	/**
+	 * Returns the background image for the scroll bar well
+	 *
+	 * @param bShowBorder a bool to indicate whether the border should be shown, at all
+	 * @param bIsInnerCategory a bool that indicates whether this background is for an inner Category
+	 * @param bIsCategoryExpanded a bool that indicates whether this Category is expanded 
+	 * @param bIsScrollBarVisible a bool that indicates whether the scrollbar for the details view is visible
+	 */	
+	const FSlateBrush* GetBackgroundImageForScrollBarWell(
+			const bool bShowBorder,
+			const bool bIsInnerCategory,
+			const bool bIsCategoryExpanded, 
+			const bool bIsScrollBarVisible) const;
+
+	/**
+	 * Initializes all Details View Styles
+	 */
+	static void InitializeDetailsViewStyles();
+	
+private:
+
+	/**
+	 * Constructs this FDetailsViewStyle with the specified parameters
+	 *
+	 * @param InKey the name of the FDetailsViewStyle
+	 * @param InTopCategoryPadding the top padding for the Category row
+	 * @param InHorizontalPadding the horizontal padding for the Category row
+	 */
+	FDetailsViewStyle(
+		const FDetailsViewStyleKey& InKey, 
+		float InTopCategoryPadding = 0.f,
+		float InHorizontalPadding = 0.f);
+
+	/**
+   	 * Initializes this FDetailsViewStyle with the specified parameters
+   	 *
+   	 * @param InKey the name of the FDetailsViewStyle
+   	 * @param InTopCategoryPadding the top padding for the Category row
+   	 * @param InHorizontalPadding the horizontal padding for the Category row
+   	 */
+	void Initialize(
+		FDetailsViewStyleKey& InKey, 
+		const float InHorizontalPadding, 
+		const float InTopCategoryPadding);
+
+	/** the Name of the Style */
+	FDetailsViewStyleKey Key;
+
+	/**
+	* Returns the @code FDetailsViewStyle @endcode for which this is the @code FDetailsViewStyleKey @endcode .
+	*
+	* @param InKey the name of the FDetailsViewStyle
+	*/
+	static const FDetailsViewStyle* GetStyle(FDetailsViewStyleKey& InKey);
+
+	/**
+	 * A map of @code FDetailsViewStyleKey @endcode instances to the
+	 * @code const FDetailsViewStyle* @endcode instances for which they are the keys. 
+	 */
+	static inline TMap<FName, const FDetailsViewStyle*> StyleKeyToStyleTemplateMap;
+	
+	/** the Slate Units of the top padding for an outer Category row */
+	float TopCategoryPadding = 0.f;
+	
+	/** the Slate Units of the horizontal padding for all details rows */
+	float HorizontalPadding = 0.f;
+	
+	/** Whether the current Category is an outer versus an inner Category  */
+   	bool bIsOuterCategory = false;
+};
+
diff --git a/Engine/Source/Editor/PropertyEditor/Private/DetailsViewStyleKey.cpp b/Engine/Source/Editor/PropertyEditor/Private/DetailsViewStyleKey.cpp
new file mode 100644
--- /dev/null
+++ b/Engine/Source/Editor/PropertyEditor/Private/DetailsViewStyleKey.cpp
@@ -0,0 +1,46 @@
+//  Copyright Epic Games, Inc. All Rights Reserved.
+
+#include "DetailsViewStyleKey.h"
+
+FDetailsViewStyleKey::FDetailsViewStyleKey(): FDetailsViewStyleKey(NAME_None)
+{
+}
+
+FDetailsViewStyleKey::FDetailsViewStyleKey(FDetailsViewStyleKey& Key): FDetailsViewStyleKey(Key.GetName())
+{
+}
+
+FDetailsViewStyleKey::FDetailsViewStyleKey(const FDetailsViewStyleKey& Key): FDetailsViewStyleKey(Key.GetName())
+{
+}
+
+FDetailsViewStyleKey& FDetailsViewStyleKey::operator=(FDetailsViewStyleKey& OtherStyleKey)
+{
+	this->Name = OtherStyleKey.Name;
+	return *this;
+}
+
+FDetailsViewStyleKey& FDetailsViewStyleKey::operator=(const FDetailsViewStyleKey& OtherStyleKey)
+{
+	this->Name = OtherStyleKey.Name;
+	return *this;
+}
+
+bool FDetailsViewStyleKey::operator==(const FDetailsViewStyleKey& OtherStyleKey) const
+{
+	return this->Name == OtherStyleKey.Name;
+}
+
+FName FDetailsViewStyleKey::GetName() const
+{
+	return Name;
+}
+
+FDetailsViewStyleKey::FDetailsViewStyleKey(FName InName): Name(InName)
+{
+}
+
+bool FDetailsViewStyleKeys::IsDefault(const FDetailsViewStyleKey& StyleKey)
+{
+	return StyleKey.GetName() == Default.GetName();
+}
diff --git a/Engine/Source/Editor/PropertyEditor/Private/IDetailsViewPrivate.h b/Engine/Source/Editor/PropertyEditor/Private/IDetailsViewPrivate.h
--- a/Engine/Source/Editor/PropertyEditor/Private/IDetailsViewPrivate.h
+++ b/Engine/Source/Editor/PropertyEditor/Private/IDetailsViewPrivate.h
@@ -175,4 +175,29 @@ public:
 	
 	/** Retrieve a list of top-most detail tree nodes. */
 	virtual void GetHeadNodes(TArray<TWeakPtr<FDetailTreeNode>>& OutNodes) {}
+
+
+	/**
+	* Gets the @code FDetailsViewStyleKey& @endcode which provides a Key to the current style for a Details View
+	*/
+	virtual FDetailsViewStyleKey& GetStyleKey() override
+	{
+		static FDetailsViewStyleKey Style = FDetailsViewStyleKeys::Default; 
+		return Style;
+	}
+	
+	/**
+	* Updates @code FDetailsViewStyleKey& StyleKey @endcode for the current @code IDetailsViewPrivate @endcode state
+	*/
+	virtual void UpdateStyleKey() override
+	{
+	}
+	
+	/**
+	* Returns a bool indicating whether the given @code FDetailsViewStyleKey @endcode is the default Details View Style 
+	*/
+	virtual bool IsDefaultStyle() const override
+	{
+		return true;
+	}
 };
diff --git a/Engine/Source/Editor/PropertyEditor/Private/PropertyEditorModule.cpp b/Engine/Source/Editor/PropertyEditor/Private/PropertyEditorModule.cpp
--- a/Engine/Source/Editor/PropertyEditor/Private/PropertyEditorModule.cpp
+++ b/Engine/Source/Editor/PropertyEditor/Private/PropertyEditorModule.cpp
@@ -31,6 +31,7 @@
 #include "UserInterface/PropertyTable/TextPropertyTableCellPresenter.h"
 #include "Widgets/Colors/SColorPicker.h"
 #include "Widgets/Layout/SBorder.h"
+#include "DetailsViewStyle.h"
 
 
 IMPLEMENT_MODULE( FPropertyEditorModule, PropertyEditor );
@@ -91,6 +92,8 @@ void FPropertyEditorModule::StartupModule()
 	StructOnScopePropertyOwner = nullptr;
 
 	FCoreUObjectDelegates::OnObjectsReplaced.AddRaw(this, &FPropertyEditorModule::ReplaceViewedObjects);
+
+	FDetailsViewStyle::InitializeDetailsViewStyles();
 }
 
 void FPropertyEditorModule::ShutdownModule()
diff --git a/Engine/Source/Editor/PropertyEditor/Private/SDetailCategoryTableRow.cpp b/Engine/Source/Editor/PropertyEditor/Private/SDetailCategoryTableRow.cpp
--- a/Engine/Source/Editor/PropertyEditor/Private/SDetailCategoryTableRow.cpp
+++ b/Engine/Source/Editor/PropertyEditor/Private/SDetailCategoryTableRow.cpp
@@ -11,6 +11,8 @@
 #include "ScopedTransaction.h"
 #include "SDetailExpanderArrow.h"
 #include "SDetailRowIndent.h"
+#include "DetailsViewStyle.h"
+#include "SDetailsView.h"
 #include "Serialization/JsonSerializer.h"
 #include "Styling/StyleColors.h"
 #include "UserInterface/PropertyEditor/PropertyEditorConstants.h"
@@ -76,26 +78,23 @@ void SDetailCategoryTableRow::Construct(const FArguments& InArgs, TSharedRef<FDe
 			];
 	}
 
-	TWeakPtr<STableViewBase> OwnerTableViewWeak = InOwnerTableView;
-	auto GetScrollbarWellBrush = [this, OwnerTableViewWeak]()
-	{
-		return SDetailTableRowBase::IsScrollBarVisible(OwnerTableViewWeak) ?
-			FAppStyle::Get().GetBrush("DetailsView.GridLine") :
-			this->GetBackgroundImage();
-	};
-
-	auto GetScrollbarWellTint = [this, OwnerTableViewWeak]()
+	OwnerTableViewWeak = InOwnerTableView;
+	
+	auto GetScrollbarWellTint = [this]()
 	{
 		return SDetailTableRowBase::IsScrollBarVisible(OwnerTableViewWeak) ?
 			FSlateColor(EStyleColor::White) :
 			this->GetInnerBackgroundColor();
 	};
 
+	FDetailsViewStyle ViewStyle = DetailsView->GetStyleKey();
+	ViewStyle.SetIsOuterCategory(!bIsInnerCategory);
+
 	this->ChildSlot
 	[
 		SNew(SBorder)
 		.BorderImage(FAppStyle::Get().GetBrush("DetailsView.GridLine"))
-		.Padding(FMargin(0, 0, 0, 1))
+		.Padding(ViewStyle.GetRowPadding())
 		[
 			SNew(SHorizontalBox)
 			+ SHorizontalBox::Slot()
@@ -117,7 +116,7 @@ void SDetailCategoryTableRow::Construct(const FArguments& InArgs, TSharedRef<FDe
 			.AutoWidth()
 			[
 				SNew(SBorder)
-				.BorderImage_Lambda(GetScrollbarWellBrush)
+				.BorderImage(this, &SDetailCategoryTableRow::GetBackgroundImageForScrollBarWell)
 				.BorderBackgroundColor_Lambda(GetScrollbarWellTint)
 				.Padding(FMargin(0, 0, SDetailTableRowBase::ScrollBarPadding, 0))
 			]
@@ -210,19 +209,35 @@ EVisibility SDetailCategoryTableRow::IsSeparatorVisible() const
 	return bIsInnerCategory || IsItemExpanded() ? EVisibility::Collapsed : EVisibility::Visible;
 }
 
+/* TODO ~ refactor to not require styling of row and scrollbar well separately */
 const FSlateBrush* SDetailCategoryTableRow::GetBackgroundImage() const
 {
 	if (bShowBorder)
 	{
-		if (bIsInnerCategory)
-		{
-			return FAppStyle::Get().GetBrush("DetailsView.CategoryMiddle");
-		}
-
-		// intentionally no hover on outer categories
-		return FAppStyle::Get().GetBrush("DetailsView.CategoryTop");
+		IDetailsViewPrivate* View = OwnerTreeNode.IsValid() && OwnerTreeNode.Pin()->GetDetailsView() ? OwnerTreeNode.Pin()->GetDetailsView() : nullptr;
+		
+		static FDetailsViewStyleKey PrimaryKey = SDetailsView::PrimaryDetailsViewStyleKey;
+		FDetailsViewStyle ViewStyle = View ? View->GetStyleKey() : PrimaryKey;
+		const bool bIsCategoryExpanded = IsItemExpanded();
+		const bool bIsScrollBarVisible = IsScrollBarVisible(OwnerTableViewWeak);
+		
+		return ViewStyle.GetBackgroundImageForCategoryRow(bShowBorder, bIsInnerCategory, bIsCategoryExpanded, bIsScrollBarVisible);
 	}
+	return nullptr;
+}
 
+const FSlateBrush* SDetailCategoryTableRow::GetBackgroundImageForScrollBarWell() const
+{
+	if (bShowBorder)
+	{
+		IDetailsViewPrivate* View = OwnerTreeNode.IsValid() && OwnerTreeNode.Pin()->GetDetailsView() ? OwnerTreeNode.Pin()->GetDetailsView() : nullptr;
+		static FDetailsViewStyleKey PrimaryKey = SDetailsView::PrimaryDetailsViewStyleKey;
+		const FDetailsViewStyle ViewStyle = View ? View->GetStyleKey() : PrimaryKey;
+		const bool bIsCategoryExpanded = IsItemExpanded();
+		const bool bIsScrollBarVisible = IsScrollBarVisible(OwnerTableViewWeak);
+
+		return ViewStyle.GetBackgroundImageForScrollBarWell(bShowBorder, bIsInnerCategory, bIsCategoryExpanded, bIsScrollBarVisible);
+	}
 	return nullptr;
 }
 
diff --git a/Engine/Source/Editor/PropertyEditor/Private/SDetailCategoryTableRow.h b/Engine/Source/Editor/PropertyEditor/Private/SDetailCategoryTableRow.h
--- a/Engine/Source/Editor/PropertyEditor/Private/SDetailCategoryTableRow.h
+++ b/Engine/Source/Editor/PropertyEditor/Private/SDetailCategoryTableRow.h
@@ -37,7 +37,16 @@ protected:
 	
 private:
 	EVisibility IsSeparatorVisible() const;
+
+	/**
+	* Gets the @code const FSlateBrush* @endcode which holds the background image for the Category Row, minus the ScrollBar Well
+	*/
 	const FSlateBrush* GetBackgroundImage() const;
+
+	/**
+	* Gets the @code const FSlateBrush* @endcode which holds the background image for the ScrollBar Well
+	*/
+	const FSlateBrush* GetBackgroundImageForScrollBarWell() const;
 	FSlateColor GetInnerBackgroundColor() const;
 	FSlateColor GetOuterBackgroundColor() const;
 
diff --git a/Engine/Source/Editor/PropertyEditor/Private/SDetailSingleItemRow.cpp b/Engine/Source/Editor/PropertyEditor/Private/SDetailSingleItemRow.cpp
--- a/Engine/Source/Editor/PropertyEditor/Private/SDetailSingleItemRow.cpp
+++ b/Engine/Source/Editor/PropertyEditor/Private/SDetailSingleItemRow.cpp
@@ -26,6 +26,8 @@
 #include "Widgets/Images/SImage.h"
 #include "Widgets/Input/SButton.h"
 #include "Widgets/Input/SComboButton.h"
+#include "DetailsViewStyle.h"
+#include "SDetailsView.h"
 
 namespace DetailWidgetConstants
 {
@@ -651,15 +653,15 @@ void SDetailSingleItemRow::Construct( const FArguments& InArgs, FDetailLayoutCus
 		Widget = SNew(SSpacer);
 	}
 
-	TWeakPtr<STableViewBase> OwnerTableViewWeak = InOwnerTableView;
-	auto GetScrollbarWellBrush = [this, OwnerTableViewWeak]()
+	OwnerTableViewWeak = InOwnerTableView;
+	auto GetScrollbarWellBrush = [this]()
 	{
 		return SDetailTableRowBase::IsScrollBarVisible(OwnerTableViewWeak) ?
 			FAppStyle::Get().GetBrush("DetailsView.GridLine") : 
 			FAppStyle::Get().GetBrush("DetailsView.CategoryMiddle");
 	};
 
-	auto GetScrollbarWellTint = [this, OwnerTableViewWeak]()
+	auto GetScrollbarWellTint = [this]()
 	{
 		return SDetailTableRowBase::IsScrollBarVisible(OwnerTableViewWeak) ?
 			FStyleColors::White : 
@@ -671,11 +673,14 @@ void SDetailSingleItemRow::Construct( const FArguments& InArgs, FDetailLayoutCus
 		return this->IsHighlighted() ? FMargin(1) : FMargin(0);
 	};
 
+	static FDetailsViewStyleKey PrimaryKey = SDetailsView::PrimaryDetailsViewStyleKey;
+	FDetailsViewStyle ViewStyle = DetailsView ? DetailsView->GetStyleKey() : PrimaryKey;
+	
 	this->ChildSlot
 	[
 		SNew( SBorder )
 		.BorderImage(FAppStyle::Get().GetBrush("DetailsView.GridLine"))
-		.Padding(FMargin(0,0,0,1))
+		.Padding(ViewStyle.GetRowPadding())
 		.Clipping(EWidgetClipping::ClipToBounds)
 		[
 			SNew(SBox)
diff --git a/Engine/Source/Editor/PropertyEditor/Private/SDetailTableRowBase.h b/Engine/Source/Editor/PropertyEditor/Private/SDetailTableRowBase.h
--- a/Engine/Source/Editor/PropertyEditor/Private/SDetailTableRowBase.h
+++ b/Engine/Source/Editor/PropertyEditor/Private/SDetailTableRowBase.h
@@ -103,5 +103,9 @@ private:
 	}
 	
 protected:
+	/**
+	 * A weak pointer to the STableViewBase from which we can get information such as if the scrollbar is showing. 
+	 */
+	TWeakPtr<STableViewBase> OwnerTableViewWeak;
 	TWeakPtr<FDetailTreeNode> OwnerTreeNode;
 };
diff --git a/Engine/Source/Editor/PropertyEditor/Private/SDetailsView.cpp b/Engine/Source/Editor/PropertyEditor/Private/SDetailsView.cpp
--- a/Engine/Source/Editor/PropertyEditor/Private/SDetailsView.cpp
+++ b/Engine/Source/Editor/PropertyEditor/Private/SDetailsView.cpp
@@ -719,6 +719,8 @@ void SDetailsView::SetObjectArrayPrivate(const TArray<UObject*>& InObjects)
 	double StartTime = FPlatformTime::Seconds();
 
 	const TArray<FDetailsViewObjectRoot> Roots = ObjectFilter->FilterObjects(InObjects);
+	
+	UpdateStyleKey();
 
 	PreSetObject(Roots.Num());
 
@@ -1471,4 +1473,32 @@ const FSlateBrush* SDetailsView::GetViewOptionsBadgeIcon() const
 	return bHasBadge ? FAppStyle::Get().GetBrush("Icons.BadgeModified") : nullptr;
 }
 
+bool SDetailsView::IsDefaultStyle() const
+{
+	return StyleKey == PrimaryDetailsViewStyleKey;
+}
+
+void SDetailsView::UpdateStyleKey()
+{
+	const bool IsDefault = FDetailsViewStyleKeys::IsDefault(ObjectFilter->GetObjectsDetailsViewStyleKey());
+
+	if (IsDefault && DetailsViewArgs.StyleKey.IsValid()) 
+	{
+		StyleKey = *DetailsViewArgs.StyleKey.Get();
+	}
+	else if (!IsDefault)
+	{
+		StyleKey = ObjectFilter->GetObjectsDetailsViewStyleKey();
+	}
+	else
+	{
+		StyleKey = PrimaryDetailsViewStyleKey;
+	}
+}
+
+FDetailsViewStyleKey& SDetailsView::GetStyleKey()
+{
+	return StyleKey;
+}
+
 #undef LOCTEXT_NAMESPACE
diff --git a/Engine/Source/Editor/PropertyEditor/Private/SDetailsView.h b/Engine/Source/Editor/PropertyEditor/Private/SDetailsView.h
--- a/Engine/Source/Editor/PropertyEditor/Private/SDetailsView.h
+++ b/Engine/Source/Editor/PropertyEditor/Private/SDetailsView.h
@@ -117,6 +117,19 @@ public:
 	{
 		return RootObjectCustomization;
 	}
+	
+	/**
+	 * The Primary @code FDetailsViewStyleKey @endcode which is used for @code SDetailsView @endcode instances for which no other
+	 * @code FDetailsViewStyleKey @endcode has been provided, for example through @code FDetailsViewArgs @endcode or an
+	 * @code FDetailsViewObjectFilter @endcode.
+	 */
+	static inline const FDetailsViewStyleKey PrimaryDetailsViewStyleKey = FDetailsViewStyleKeys::Classic;
+
+	/**
+	* Returns a bool indicating whether the given @code FDetailsViewStyleKey @endcode is the default Details View Style 
+	*/
+	bool IsDefaultStyle() const;
+	
 private:
 
 	void SetObjectArrayPrivate(const TArray<UObject*>& InObjects);
@@ -196,7 +209,22 @@ private:
 	/** Get the badge for the view options icon (if there are any options active */
 	const FSlateBrush* GetViewOptionsBadgeIcon() const;
 
+	/**
+	* Updates @code FDetailsViewStyleKey& StyleKey @endcode for the current @code SDetailsView @endcode state
+	*/
+	virtual void UpdateStyleKey() override;
+
+	/**
+	* Gets the @code FDetailsViewStyleKey& @endcode which provides a Key to the current style for a Details View
+	*/
+	virtual FDetailsViewStyleKey& GetStyleKey() override;
+
 private:
+	/**
+	* The @code FDetailsViewStyle& @endcode which provides style information for this details view.
+	*/
+	FDetailsViewStyleKey StyleKey = PrimaryDetailsViewStyleKey;
+	
 	/** The filter for objects viewed by this details panel */
 	TSharedPtr<FDetailsViewObjectFilter> ObjectFilter;
 
diff --git a/Engine/Source/Editor/PropertyEditor/Public/DetailsViewArgs.h b/Engine/Source/Editor/PropertyEditor/Public/DetailsViewArgs.h
--- a/Engine/Source/Editor/PropertyEditor/Public/DetailsViewArgs.h
+++ b/Engine/Source/Editor/PropertyEditor/Public/DetailsViewArgs.h
@@ -3,6 +3,7 @@
 #pragma once
 
 #include "CoreMinimal.h"
+#include "DetailsViewStyleKey.h"
 #include "Widgets/Layout/SScrollBar.h"
 
 class FDetailsViewObjectFilter;
@@ -10,6 +11,8 @@ class FNotifyHook;
 class FTabManager;
 class FUICommandList;
 class IClassViewerFilter;
+class FDetailsViewStyleKey;
+class FDetailsViewStyleKeys;
 
 enum class EEditDefaultsOnlyNodeVisibility : uint8
 {
@@ -52,7 +55,8 @@ struct FDetailsViewArgs
 	TSharedPtr<SScrollBar> ExternalScrollbar;
 	/** whether to display scrollbar on the left or the right of the panel */
 	EHorizontalAlignment ScrollbarAlignment;
-
+	/** A @code TSharedPtr<FDetailsViewStyleKey> @endcode which provides the key to the Details View's style */
+	TSharedPtr<FDetailsViewStyleKey> StyleKey;
 	/** Identifier for this details view; NAME_None if this view is anonymous */
 	FName ViewIdentifier;
 	/** Notify hook to call when properties are changed */
diff --git a/Engine/Source/Editor/PropertyEditor/Public/DetailsViewObjectFilter.h b/Engine/Source/Editor/PropertyEditor/Public/DetailsViewObjectFilter.h
--- a/Engine/Source/Editor/PropertyEditor/Public/DetailsViewObjectFilter.h
+++ b/Engine/Source/Editor/PropertyEditor/Public/DetailsViewObjectFilter.h
@@ -3,6 +3,7 @@
 #pragma once
 
 #include "CoreMinimal.h"
+#include "DetailsViewStyleKey.h"
 
 /**
  * An object root is a collection of UObjects that represent a top level set of properties in a details panel
@@ -33,13 +34,31 @@ struct FDetailsViewObjectRoot
 };
 
 /**
- * An object filter determines the root objects that should be displayed from a set of given source objects passed to the details panel
+ * An object filter determines the root objects that should be displayed from a set of given source objects passed to the details panel.
+ * It can also be used to convey characteristics of a details view which may alter depending upon the type of objects
+ * that have been filtered.
  */
 class FDetailsViewObjectFilter
 {
 public:
+
 	virtual ~FDetailsViewObjectFilter() {}
 
+	/**
+	 * Given a const TArray<UObject*>& SourceObjects, it fills a TArray<FDetailsViewObjectRoot> with the objects
+	 * that we need as details objects roots. This may be the same as the original SourceObjects, or it may be
+	 * some subset of them, or some of their Sub-objects
+	 *
+	 * @param SourceObjects the array of objects acting as the Source of the root objects for the details view
+	 * @return the TArray<FDetailsViewObjectRoot> with the objects that will act as root objects for the details view 
+	 */
 	virtual TArray<FDetailsViewObjectRoot> FilterObjects(const TArray<UObject*>& SourceObjects) = 0;
 
+	/**
+	 * Returns the @code FDetailsViewStyleKey& @endcode that is the Key to the current objects' style
+	 */
+	virtual const FDetailsViewStyleKey& GetObjectsDetailsViewStyleKey()
+	{
+		return FDetailsViewStyleKeys::Default;
+	}
 };
diff --git a/Engine/Source/Editor/PropertyEditor/Public/DetailsViewStyleKey.h b/Engine/Source/Editor/PropertyEditor/Public/DetailsViewStyleKey.h
new file mode 100644
--- /dev/null
+++ b/Engine/Source/Editor/PropertyEditor/Public/DetailsViewStyleKey.h
@@ -0,0 +1,93 @@
+//  Copyright Epic Games, Inc. All Rights Reserved.
+
+
+#pragma once
+#include "UObject/NameTypes.h"
+
+
+/**
+ * A class which provides a strongly typed Key object for @code FDetailsViewStyleKey @endcode instances
+ */
+class FDetailsViewStyleKey
+{
+public:
+	/**
+	 * Default constructor for @code FDetailsViewStyleKey @endcode instances
+	 */ 
+	PROPERTYEDITOR_API FDetailsViewStyleKey();
+
+	/**
+	 * Copy Constructor for @code FDetailsViewStyleKey @endcode instances
+	 */
+	PROPERTYEDITOR_API FDetailsViewStyleKey(FDetailsViewStyleKey& Key);
+
+	/**
+	 * const Copy Constructor for @code FDetailsViewStyleKey @endcode instances
+	 */
+	PROPERTYEDITOR_API FDetailsViewStyleKey(const FDetailsViewStyleKey& Key);
+
+	/**
+	 * Assignment operator for @code FDetailsViewStyleKey @endcode instances.
+	 *
+	 * @param OtherStyleKey the key for the @code FDetailsViewStyleKey @endcode with which to initialize this @code FDetailsViewStyleKey @endcode
+	 */
+	PROPERTYEDITOR_API FDetailsViewStyleKey& operator=(FDetailsViewStyleKey& OtherStyleKey);
+
+	/**
+	 * const assignment operator @code FDetailsViewStyleKey @endcode instances.
+	 *
+	 * @param OtherStyleKey the key for the @code FDetailsViewStyleKey @endcode with which to initialize this @code FDetailsViewStyleKey @endcode
+	 */
+	PROPERTYEDITOR_API FDetailsViewStyleKey& operator=(const FDetailsViewStyleKey& OtherStyleKey);
+
+	PROPERTYEDITOR_API bool operator==(const FDetailsViewStyleKey& OtherStyleKey) const;
+
+	/**
+	 * Returns the FName Name of this @code FDetailsViewStyleKey @endcode 
+	 */
+	PROPERTYEDITOR_API FName GetName() const;
+
+	/**
+	 * A constructor for @code FDetailsViewStyleKey @endcode instances which takes an FName InName to initialize the leu
+	 */
+	PROPERTYEDITOR_API FDetailsViewStyleKey(FName InName);
+
+private:
+	FName Name;
+
+};
+
+/**
+ * A Class which provides keys for @code FDetailsViewStyleKey @endcode instances
+ */
+class FDetailsViewStyleKeys 
+{
+public:
+
+	/**
+	 * The Classic Details View style used in Unreal Engine 5.0 and prior, which has Categories displayed with no
+	 * padding between them and no built in horizontal space 
+	 */
+	static inline const PROPERTYEDITOR_API FDetailsViewStyleKey Classic{FName{"Classic"}};
+	
+	/**
+	* A Details View style which presents with the top level Categories in a card-like fashion, with space between
+	* and around them.
+	*/
+	static inline const PROPERTYEDITOR_API FDetailsViewStyleKey Card{FName{"Card"}};
+
+	/**
+	 * The @code Default @endcode Details View style should be used when an object has no interest in displaying its own custom
+	 * Details View and wants to default to whatever is the Primary Details View.
+	 */
+	static inline const PROPERTYEDITOR_API FDetailsViewStyleKey Default{FName{"Default"}};
+
+	/**
+	* Returns true if the @code const FDetailsViewStyleKey& @endcode is the @code Default @endcode Key, else it returns false.
+	*
+	* @param StyleKey a FDetailsViewStyleKey& to check whether it is the default 
+	*/
+	static bool PROPERTYEDITOR_API IsDefault(const FDetailsViewStyleKey& StyleKey);
+};
+
+
diff --git a/Engine/Source/Editor/PropertyEditor/Public/IDetailsView.h b/Engine/Source/Editor/PropertyEditor/Public/IDetailsView.h
--- a/Engine/Source/Editor/PropertyEditor/Public/IDetailsView.h
+++ b/Engine/Source/Editor/PropertyEditor/Public/IDetailsView.h
@@ -291,4 +291,19 @@ public:
 
 	/** Sets the minimum width of the detail panel's right column. */
 	virtual void SetRightColumnMinWidth(float InMinWidth) = 0;
+
+	/**
+	* Updates @code FDetailsViewStyleKey& StyleKey @endcode for the current IDetailsView state
+	*/
+	virtual void UpdateStyleKey() = 0;
+
+	/**
+	* Gets the @code FDetailsViewStyleKey& @endcode which provides a Key to the current style for a Details View
+	*/
+	virtual FDetailsViewStyleKey& GetStyleKey() = 0;
+	
+	/**
+   	* Returns a bool indicating whether the given @code FDetailsViewStyleKey @endcode is the default Details View Style 
+   	*/
+   	virtual bool IsDefaultStyle() const = 0;
 };
