Revert: Updates for Modular Class Display to add a widget registration system and get it working for BaseToolkit and toolbars....

diff --git a/Engine/Source/Editor/EditorStyle/Private/StarshipStyle.cpp b/Engine/Source/Editor/EditorStyle/Private/StarshipStyle.cpp
--- a/Engine/Source/Editor/EditorStyle/Private/StarshipStyle.cpp
+++ b/Engine/Source/Editor/EditorStyle/Private/StarshipStyle.cpp
@@ -2995,18 +2995,6 @@ void FStarshipEditorStyle::FStyle::SetupPropertyEditorStyles()
 		Set( "DetailsView.CollapsedCategory",         new FSlateColorBrush(FStyleColors::Header));
 		Set( "DetailsView.CollapsedCategory_Hovered", new FSlateColorBrush(FStyleColors::Hover));
 		Set( "DetailsView.CategoryTop",               new FSlateColorBrush(FStyleColors::Header));
-		
-		/****** Styles for rounded corners for the Card style of a Details View ********/
-
-		Set( "DetailsView.CardHeaderTopLeftSideRounded",               new FSlateRoundedBoxBrush(FStyleColors::Header, FVector4(4.0f, 0.0f, 0.0f, 0.0f)));
-		Set( "DetailsView.CardHeaderLeftSideRounded",               new FSlateRoundedBoxBrush(FStyleColors::Header, FVector4(4.0f, 0.0f, 0.0f, 4.0f)));
-		Set( "DetailsView.CardHeaderTopRightSideRounded",               new FSlateRoundedBoxBrush(FStyleColors::Header, FVector4(0.0f, 4.0f, 0.0f, 0.0f)));
-		Set( "DetailsView.CardHeaderRightSideRounded",               new FSlateRoundedBoxBrush(FStyleColors::Header, FVector4(0.0f, 4.0f, 4.0f, 0.0f)));
-		Set( "DetailsView.CardHeaderTopRounded",               new FSlateRoundedBoxBrush(FStyleColors::Header, FVector4(4.0f, 4.0f, 0.0f, 0.0f)));
-		Set( "DetailsView.CardHeaderRounded",               new FSlateRoundedBoxBrush(FStyleColors::Header, FVector4(4.0f, 4.0f, 4.0f, 4.0f )));
-		
-		/*******************************************************************************/
-		
 		Set( "DetailsView.CategoryTop_Hovered",       new FSlateColorBrush(FStyleColors::Hover));
 		Set( "DetailsView.CategoryBottom",            new FSlateColorBrush(FStyleColors::Recessed));
 		
diff --git a/Engine/Source/Editor/LevelEditor/Private/SActorDetails.cpp b/Engine/Source/Editor/LevelEditor/Private/SActorDetails.cpp
--- a/Engine/Source/Editor/LevelEditor/Private/SActorDetails.cpp
+++ b/Engine/Source/Editor/LevelEditor/Private/SActorDetails.cpp
@@ -280,7 +280,7 @@ void SActorDetails::Construct(const FArguments& InArgs, UTypedElementSelectionSe
 	ComponentsBox->SetContent(SubobjectEditor.ToSharedRef());
 
 	TSharedRef<SWidget> SubobjectEditorButtonBox = SubobjectEditor->GetToolButtonsBox().ToSharedRef();
-	SubobjectEditorButtonBox->SetVisibility(MakeAttributeSP(this, &SActorDetails::GetComponentEditorButtonsVisibility));
+	SubobjectEditorButtonBox->SetVisibility(MakeAttributeSP(this, &SActorDetails::GetComponentEditorVisibility));
 
 
 	TFunction<TArray<TTypedElement<ITypedElementDetailsInterface>>()> GetDetailsHandles = [this]()
@@ -879,12 +879,6 @@ void SActorDetails::OnNativeComponentWarningHyperlinkClicked(const FSlateHyperli
 }
 
 EVisibility SActorDetails::GetComponentEditorVisibility() const
-{
-	const bool bHideEditorFromDetailsView = DetailsView.IsValid() &&  !DetailsView->IsDefaultStyle();
-	return GetActorContext() && !bHideEditorFromDetailsView  ? EVisibility::Visible : EVisibility::Collapsed;
-}
-
-EVisibility SActorDetails::GetComponentEditorButtonsVisibility() const
 {
 	return GetActorContext() ? EVisibility::Visible : EVisibility::Collapsed;
 }
diff --git a/Engine/Source/Editor/LevelEditor/Private/SActorDetails.h b/Engine/Source/Editor/LevelEditor/Private/SActorDetails.h
--- a/Engine/Source/Editor/LevelEditor/Private/SActorDetails.h
+++ b/Engine/Source/Editor/LevelEditor/Private/SActorDetails.h
@@ -88,11 +88,6 @@ private:
 	UObject* GetActorContextAsObject() const;
 	bool GetAllowComponentTreeEditing() const;
 
-	/**
-	 * Gets the Visibility of the top of the panel buttons in the details panel
-	 */
-	EVisibility GetComponentEditorButtonsVisibility() const;
-
 	void OnComponentsEditedInWorld();
 	void OnSubobjectEditorTreeViewSelectionChanged(const TArray<TSharedPtr<FSubobjectEditorTreeNode> >& SelectedNodes);
 	void OnSubobjectEditorTreeViewItemDoubleClicked(const TSharedPtr<FSubobjectEditorTreeNode> ClickedNode);
diff --git a/Engine/Source/Editor/PropertyEditor/Private/DetailAdvancedDropdownNode.cpp b/Engine/Source/Editor/PropertyEditor/Private/DetailAdvancedDropdownNode.cpp
--- a/Engine/Source/Editor/PropertyEditor/Private/DetailAdvancedDropdownNode.cpp
+++ b/Engine/Source/Editor/PropertyEditor/Private/DetailAdvancedDropdownNode.cpp
@@ -9,8 +9,6 @@
 #include "UserInterface/PropertyEditor/PropertyEditorConstants.h"
 #include "Widgets/Images/SImage.h"
 #include "Widgets/Input/SButton.h"
-#include "DetailsViewStyle.h"
-#include "SDetailsView.h"
 
 class SAdvancedDropdownRow : public SDetailTableRowBase
 {
@@ -87,29 +85,26 @@ public:
 			]
 		];
 
-		OwnerTableViewWeak = InOwnerTableView;
-		auto GetScrollbarWellBrush = [this]()
+		TWeakPtr<STableViewBase> OwnerTableViewWeak = InOwnerTableView;
+		auto GetScrollbarWellBrush = [this, OwnerTableViewWeak]()
 		{
 			return SDetailTableRowBase::IsScrollBarVisible(OwnerTableViewWeak) ?
 				FAppStyle::Get().GetBrush("DetailsView.GridLine") : 
 				FAppStyle::Get().GetBrush("DetailsView.CategoryMiddle");
 		};
 
-		auto GetScrollbarWellTint = [this]()
+		auto GetScrollbarWellTint = [this, OwnerTableViewWeak]()
 		{
-			return IsScrollBarVisible(OwnerTableViewWeak) ?
+			return SDetailTableRowBase::IsScrollBarVisible(OwnerTableViewWeak) ?
 				FStyleColors::White : 
 				this->GetRowBackgroundColor();
 		};
 
-		static FDetailsViewStyleKey PrimaryKey = SDetailsView::PrimaryDetailsViewStyleKey;
-		FDetailsViewStyle ViewStyle = DetailsView ? DetailsView->GetStyleKey() : PrimaryKey;
-
 		ChildSlot
 		[
 			SNew(SBorder)
 			.BorderImage(FAppStyle::Get().GetBrush("DetailsView.GridLine"))
-			.Padding(ViewStyle.GetRowPadding())
+			.Padding(FMargin(0, 0, 0, 1))
 			[
 				SNew(SHorizontalBox)
 				+ SHorizontalBox::Slot()
diff --git a/Engine/Source/Editor/PropertyEditor/Private/DetailsViewStyle.cpp b/Engine/Source/Editor/PropertyEditor/Private/DetailsViewStyle.cpp
deleted file mode 100644
diff --git a/Engine/Source/Editor/PropertyEditor/Private/DetailsViewStyle.h b/Engine/Source/Editor/PropertyEditor/Private/DetailsViewStyle.h
deleted file mode 100644
diff --git a/Engine/Source/Editor/PropertyEditor/Private/DetailsViewStyleKey.cpp b/Engine/Source/Editor/PropertyEditor/Private/DetailsViewStyleKey.cpp
deleted file mode 100644
diff --git a/Engine/Source/Editor/PropertyEditor/Private/IDetailsViewPrivate.h b/Engine/Source/Editor/PropertyEditor/Private/IDetailsViewPrivate.h
--- a/Engine/Source/Editor/PropertyEditor/Private/IDetailsViewPrivate.h
+++ b/Engine/Source/Editor/PropertyEditor/Private/IDetailsViewPrivate.h
@@ -175,29 +175,4 @@ public:
 	
 	/** Retrieve a list of top-most detail tree nodes. */
 	virtual void GetHeadNodes(TArray<TWeakPtr<FDetailTreeNode>>& OutNodes) {}
-
-
-	/**
-	* Gets the @code FDetailsViewStyleKey& @endcode which provides a Key to the current style for a Details View
-	*/
-	virtual FDetailsViewStyleKey& GetStyleKey() override
-	{
-		static FDetailsViewStyleKey Style = FDetailsViewStyleKeys::Default; 
-		return Style;
-	}
-	
-	/**
-	* Updates @code FDetailsViewStyleKey& StyleKey @endcode for the current @code IDetailsViewPrivate @endcode state
-	*/
-	virtual void UpdateStyleKey() override
-	{
-	}
-	
-	/**
-	* Returns a bool indicating whether the given @code FDetailsViewStyleKey @endcode is the default Details View Style 
-	*/
-	virtual bool IsDefaultStyle() const override
-	{
-		return true;
-	}
 };
diff --git a/Engine/Source/Editor/PropertyEditor/Private/PropertyEditorModule.cpp b/Engine/Source/Editor/PropertyEditor/Private/PropertyEditorModule.cpp
--- a/Engine/Source/Editor/PropertyEditor/Private/PropertyEditorModule.cpp
+++ b/Engine/Source/Editor/PropertyEditor/Private/PropertyEditorModule.cpp
@@ -31,7 +31,6 @@
 #include "UserInterface/PropertyTable/TextPropertyTableCellPresenter.h"
 #include "Widgets/Colors/SColorPicker.h"
 #include "Widgets/Layout/SBorder.h"
-#include "DetailsViewStyle.h"
 
 
 IMPLEMENT_MODULE( FPropertyEditorModule, PropertyEditor );
@@ -92,8 +91,6 @@ void FPropertyEditorModule::StartupModule()
 	StructOnScopePropertyOwner = nullptr;
 
 	FCoreUObjectDelegates::OnObjectsReplaced.AddRaw(this, &FPropertyEditorModule::ReplaceViewedObjects);
-
-	FDetailsViewStyle::InitializeDetailsViewStyles();
 }
 
 void FPropertyEditorModule::ShutdownModule()
diff --git a/Engine/Source/Editor/PropertyEditor/Private/SDetailCategoryTableRow.cpp b/Engine/Source/Editor/PropertyEditor/Private/SDetailCategoryTableRow.cpp
--- a/Engine/Source/Editor/PropertyEditor/Private/SDetailCategoryTableRow.cpp
+++ b/Engine/Source/Editor/PropertyEditor/Private/SDetailCategoryTableRow.cpp
@@ -11,8 +11,6 @@
 #include "ScopedTransaction.h"
 #include "SDetailExpanderArrow.h"
 #include "SDetailRowIndent.h"
-#include "DetailsViewStyle.h"
-#include "SDetailsView.h"
 #include "Serialization/JsonSerializer.h"
 #include "Styling/StyleColors.h"
 #include "UserInterface/PropertyEditor/PropertyEditorConstants.h"
@@ -78,23 +76,26 @@ void SDetailCategoryTableRow::Construct(const FArguments& InArgs, TSharedRef<FDe
 			];
 	}
 
-	OwnerTableViewWeak = InOwnerTableView;
-	
-	auto GetScrollbarWellTint = [this]()
+	TWeakPtr<STableViewBase> OwnerTableViewWeak = InOwnerTableView;
+	auto GetScrollbarWellBrush = [this, OwnerTableViewWeak]()
+	{
+		return SDetailTableRowBase::IsScrollBarVisible(OwnerTableViewWeak) ?
+			FAppStyle::Get().GetBrush("DetailsView.GridLine") :
+			this->GetBackgroundImage();
+	};
+
+	auto GetScrollbarWellTint = [this, OwnerTableViewWeak]()
 	{
 		return SDetailTableRowBase::IsScrollBarVisible(OwnerTableViewWeak) ?
 			FSlateColor(EStyleColor::White) :
 			this->GetInnerBackgroundColor();
 	};
 
-	FDetailsViewStyle ViewStyle = DetailsView->GetStyleKey();
-	ViewStyle.SetIsOuterCategory(!bIsInnerCategory);
-
 	this->ChildSlot
 	[
 		SNew(SBorder)
 		.BorderImage(FAppStyle::Get().GetBrush("DetailsView.GridLine"))
-		.Padding(ViewStyle.GetRowPadding())
+		.Padding(FMargin(0, 0, 0, 1))
 		[
 			SNew(SHorizontalBox)
 			+ SHorizontalBox::Slot()
@@ -116,7 +117,7 @@ void SDetailCategoryTableRow::Construct(const FArguments& InArgs, TSharedRef<FDe
 			.AutoWidth()
 			[
 				SNew(SBorder)
-				.BorderImage(this, &SDetailCategoryTableRow::GetBackgroundImageForScrollBarWell)
+				.BorderImage_Lambda(GetScrollbarWellBrush)
 				.BorderBackgroundColor_Lambda(GetScrollbarWellTint)
 				.Padding(FMargin(0, 0, SDetailTableRowBase::ScrollBarPadding, 0))
 			]
@@ -209,35 +210,19 @@ EVisibility SDetailCategoryTableRow::IsSeparatorVisible() const
 	return bIsInnerCategory || IsItemExpanded() ? EVisibility::Collapsed : EVisibility::Visible;
 }
 
-/* TODO ~ refactor to not require styling of row and scrollbar well separately */
 const FSlateBrush* SDetailCategoryTableRow::GetBackgroundImage() const
 {
 	if (bShowBorder)
 	{
-		IDetailsViewPrivate* View = OwnerTreeNode.IsValid() && OwnerTreeNode.Pin()->GetDetailsView() ? OwnerTreeNode.Pin()->GetDetailsView() : nullptr;
-		
-		static FDetailsViewStyleKey PrimaryKey = SDetailsView::PrimaryDetailsViewStyleKey;
-		FDetailsViewStyle ViewStyle = View ? View->GetStyleKey() : PrimaryKey;
-		const bool bIsCategoryExpanded = IsItemExpanded();
-		const bool bIsScrollBarVisible = IsScrollBarVisible(OwnerTableViewWeak);
-		
-		return ViewStyle.GetBackgroundImageForCategoryRow(bShowBorder, bIsInnerCategory, bIsCategoryExpanded, bIsScrollBarVisible);
-	}
-	return nullptr;
-}
-
-const FSlateBrush* SDetailCategoryTableRow::GetBackgroundImageForScrollBarWell() const
-{
-	if (bShowBorder)
-	{
-		IDetailsViewPrivate* View = OwnerTreeNode.IsValid() && OwnerTreeNode.Pin()->GetDetailsView() ? OwnerTreeNode.Pin()->GetDetailsView() : nullptr;
-		static FDetailsViewStyleKey PrimaryKey = SDetailsView::PrimaryDetailsViewStyleKey;
-		const FDetailsViewStyle ViewStyle = View ? View->GetStyleKey() : PrimaryKey;
-		const bool bIsCategoryExpanded = IsItemExpanded();
-		const bool bIsScrollBarVisible = IsScrollBarVisible(OwnerTableViewWeak);
+		if (bIsInnerCategory)
+		{
+			return FAppStyle::Get().GetBrush("DetailsView.CategoryMiddle");
+		}
 
-		return ViewStyle.GetBackgroundImageForScrollBarWell(bShowBorder, bIsInnerCategory, bIsCategoryExpanded, bIsScrollBarVisible);
+		// intentionally no hover on outer categories
+		return FAppStyle::Get().GetBrush("DetailsView.CategoryTop");
 	}
+
 	return nullptr;
 }
 
diff --git a/Engine/Source/Editor/PropertyEditor/Private/SDetailCategoryTableRow.h b/Engine/Source/Editor/PropertyEditor/Private/SDetailCategoryTableRow.h
--- a/Engine/Source/Editor/PropertyEditor/Private/SDetailCategoryTableRow.h
+++ b/Engine/Source/Editor/PropertyEditor/Private/SDetailCategoryTableRow.h
@@ -37,16 +37,7 @@ protected:
 	
 private:
 	EVisibility IsSeparatorVisible() const;
-
-	/**
-	* Gets the @code const FSlateBrush* @endcode which holds the background image for the Category Row, minus the ScrollBar Well
-	*/
 	const FSlateBrush* GetBackgroundImage() const;
-
-	/**
-	* Gets the @code const FSlateBrush* @endcode which holds the background image for the ScrollBar Well
-	*/
-	const FSlateBrush* GetBackgroundImageForScrollBarWell() const;
 	FSlateColor GetInnerBackgroundColor() const;
 	FSlateColor GetOuterBackgroundColor() const;
 
diff --git a/Engine/Source/Editor/PropertyEditor/Private/SDetailSingleItemRow.cpp b/Engine/Source/Editor/PropertyEditor/Private/SDetailSingleItemRow.cpp
--- a/Engine/Source/Editor/PropertyEditor/Private/SDetailSingleItemRow.cpp
+++ b/Engine/Source/Editor/PropertyEditor/Private/SDetailSingleItemRow.cpp
@@ -26,8 +26,6 @@
 #include "Widgets/Images/SImage.h"
 #include "Widgets/Input/SButton.h"
 #include "Widgets/Input/SComboButton.h"
-#include "DetailsViewStyle.h"
-#include "SDetailsView.h"
 
 namespace DetailWidgetConstants
 {
@@ -653,15 +651,15 @@ void SDetailSingleItemRow::Construct( const FArguments& InArgs, FDetailLayoutCus
 		Widget = SNew(SSpacer);
 	}
 
-	OwnerTableViewWeak = InOwnerTableView;
-	auto GetScrollbarWellBrush = [this]()
+	TWeakPtr<STableViewBase> OwnerTableViewWeak = InOwnerTableView;
+	auto GetScrollbarWellBrush = [this, OwnerTableViewWeak]()
 	{
 		return SDetailTableRowBase::IsScrollBarVisible(OwnerTableViewWeak) ?
 			FAppStyle::Get().GetBrush("DetailsView.GridLine") : 
 			FAppStyle::Get().GetBrush("DetailsView.CategoryMiddle");
 	};
 
-	auto GetScrollbarWellTint = [this]()
+	auto GetScrollbarWellTint = [this, OwnerTableViewWeak]()
 	{
 		return SDetailTableRowBase::IsScrollBarVisible(OwnerTableViewWeak) ?
 			FStyleColors::White : 
@@ -673,14 +671,11 @@ void SDetailSingleItemRow::Construct( const FArguments& InArgs, FDetailLayoutCus
 		return this->IsHighlighted() ? FMargin(1) : FMargin(0);
 	};
 
-	static FDetailsViewStyleKey PrimaryKey = SDetailsView::PrimaryDetailsViewStyleKey;
-	FDetailsViewStyle ViewStyle = DetailsView ? DetailsView->GetStyleKey() : PrimaryKey;
-	
 	this->ChildSlot
 	[
 		SNew( SBorder )
 		.BorderImage(FAppStyle::Get().GetBrush("DetailsView.GridLine"))
-		.Padding(ViewStyle.GetRowPadding())
+		.Padding(FMargin(0,0,0,1))
 		.Clipping(EWidgetClipping::ClipToBounds)
 		[
 			SNew(SBox)
diff --git a/Engine/Source/Editor/PropertyEditor/Private/SDetailTableRowBase.h b/Engine/Source/Editor/PropertyEditor/Private/SDetailTableRowBase.h
--- a/Engine/Source/Editor/PropertyEditor/Private/SDetailTableRowBase.h
+++ b/Engine/Source/Editor/PropertyEditor/Private/SDetailTableRowBase.h
@@ -103,9 +103,5 @@ private:
 	}
 	
 protected:
-	/**
-	 * A weak pointer to the STableViewBase from which we can get information such as if the scrollbar is showing. 
-	 */
-	TWeakPtr<STableViewBase> OwnerTableViewWeak;
 	TWeakPtr<FDetailTreeNode> OwnerTreeNode;
 };
diff --git a/Engine/Source/Editor/PropertyEditor/Private/SDetailsView.cpp b/Engine/Source/Editor/PropertyEditor/Private/SDetailsView.cpp
--- a/Engine/Source/Editor/PropertyEditor/Private/SDetailsView.cpp
+++ b/Engine/Source/Editor/PropertyEditor/Private/SDetailsView.cpp
@@ -719,8 +719,6 @@ void SDetailsView::SetObjectArrayPrivate(const TArray<UObject*>& InObjects)
 	double StartTime = FPlatformTime::Seconds();
 
 	const TArray<FDetailsViewObjectRoot> Roots = ObjectFilter->FilterObjects(InObjects);
-	
-	UpdateStyleKey();
 
 	PreSetObject(Roots.Num());
 
@@ -1473,32 +1471,4 @@ const FSlateBrush* SDetailsView::GetViewOptionsBadgeIcon() const
 	return bHasBadge ? FAppStyle::Get().GetBrush("Icons.BadgeModified") : nullptr;
 }
 
-bool SDetailsView::IsDefaultStyle() const
-{
-	return StyleKey == PrimaryDetailsViewStyleKey;
-}
-
-void SDetailsView::UpdateStyleKey()
-{
-	const bool IsDefault = FDetailsViewStyleKeys::IsDefault(ObjectFilter->GetObjectsDetailsViewStyleKey());
-
-	if (IsDefault && DetailsViewArgs.StyleKey.IsValid()) 
-	{
-		StyleKey = *DetailsViewArgs.StyleKey.Get();
-	}
-	else if (!IsDefault)
-	{
-		StyleKey = ObjectFilter->GetObjectsDetailsViewStyleKey();
-	}
-	else
-	{
-		StyleKey = PrimaryDetailsViewStyleKey;
-	}
-}
-
-FDetailsViewStyleKey& SDetailsView::GetStyleKey()
-{
-	return StyleKey;
-}
-
 #undef LOCTEXT_NAMESPACE
diff --git a/Engine/Source/Editor/PropertyEditor/Private/SDetailsView.h b/Engine/Source/Editor/PropertyEditor/Private/SDetailsView.h
--- a/Engine/Source/Editor/PropertyEditor/Private/SDetailsView.h
+++ b/Engine/Source/Editor/PropertyEditor/Private/SDetailsView.h
@@ -117,19 +117,6 @@ public:
 	{
 		return RootObjectCustomization;
 	}
-	
-	/**
-	 * The Primary @code FDetailsViewStyleKey @endcode which is used for @code SDetailsView @endcode instances for which no other
-	 * @code FDetailsViewStyleKey @endcode has been provided, for example through @code FDetailsViewArgs @endcode or an
-	 * @code FDetailsViewObjectFilter @endcode.
-	 */
-	static inline const FDetailsViewStyleKey PrimaryDetailsViewStyleKey = FDetailsViewStyleKeys::Classic;
-
-	/**
-	* Returns a bool indicating whether the given @code FDetailsViewStyleKey @endcode is the default Details View Style 
-	*/
-	bool IsDefaultStyle() const;
-	
 private:
 
 	void SetObjectArrayPrivate(const TArray<UObject*>& InObjects);
@@ -209,22 +196,7 @@ private:
 	/** Get the badge for the view options icon (if there are any options active */
 	const FSlateBrush* GetViewOptionsBadgeIcon() const;
 
-	/**
-	* Updates @code FDetailsViewStyleKey& StyleKey @endcode for the current @code SDetailsView @endcode state
-	*/
-	virtual void UpdateStyleKey() override;
-
-	/**
-	* Gets the @code FDetailsViewStyleKey& @endcode which provides a Key to the current style for a Details View
-	*/
-	virtual FDetailsViewStyleKey& GetStyleKey() override;
-
 private:
-	/**
-	* The @code FDetailsViewStyle& @endcode which provides style information for this details view.
-	*/
-	FDetailsViewStyleKey StyleKey = PrimaryDetailsViewStyleKey;
-	
 	/** The filter for objects viewed by this details panel */
 	TSharedPtr<FDetailsViewObjectFilter> ObjectFilter;
 
diff --git a/Engine/Source/Editor/PropertyEditor/Public/DetailsViewArgs.h b/Engine/Source/Editor/PropertyEditor/Public/DetailsViewArgs.h
--- a/Engine/Source/Editor/PropertyEditor/Public/DetailsViewArgs.h
+++ b/Engine/Source/Editor/PropertyEditor/Public/DetailsViewArgs.h
@@ -3,7 +3,6 @@
 #pragma once
 
 #include "CoreMinimal.h"
-#include "DetailsViewStyleKey.h"
 #include "Widgets/Layout/SScrollBar.h"
 
 class FDetailsViewObjectFilter;
@@ -11,8 +10,6 @@ class FNotifyHook;
 class FTabManager;
 class FUICommandList;
 class IClassViewerFilter;
-class FDetailsViewStyleKey;
-class FDetailsViewStyleKeys;
 
 enum class EEditDefaultsOnlyNodeVisibility : uint8
 {
@@ -24,7 +21,7 @@ enum class EEditDefaultsOnlyNodeVisibility : uint8
 	Automatic,
 };
 
-/**éé
+/**Ã©Ã©
  * Init params for a details view widget
  */
 struct FDetailsViewArgs
@@ -55,8 +52,7 @@ struct FDetailsViewArgs
 	TSharedPtr<SScrollBar> ExternalScrollbar;
 	/** whether to display scrollbar on the left or the right of the panel */
 	EHorizontalAlignment ScrollbarAlignment;
-	/** A @code TSharedPtr<FDetailsViewStyleKey> @endcode which provides the key to the Details View's style */
-	TSharedPtr<FDetailsViewStyleKey> StyleKey;
+
 	/** Identifier for this details view; NAME_None if this view is anonymous */
 	FName ViewIdentifier;
 	/** Notify hook to call when properties are changed */
diff --git a/Engine/Source/Editor/PropertyEditor/Public/DetailsViewObjectFilter.h b/Engine/Source/Editor/PropertyEditor/Public/DetailsViewObjectFilter.h
--- a/Engine/Source/Editor/PropertyEditor/Public/DetailsViewObjectFilter.h
+++ b/Engine/Source/Editor/PropertyEditor/Public/DetailsViewObjectFilter.h
@@ -3,7 +3,6 @@
 #pragma once
 
 #include "CoreMinimal.h"
-#include "DetailsViewStyleKey.h"
 
 /**
  * An object root is a collection of UObjects that represent a top level set of properties in a details panel
@@ -34,31 +33,13 @@ struct FDetailsViewObjectRoot
 };
 
 /**
- * An object filter determines the root objects that should be displayed from a set of given source objects passed to the details panel.
- * It can also be used to convey characteristics of a details view which may alter depending upon the type of objects
- * that have been filtered.
+ * An object filter determines the root objects that should be displayed from a set of given source objects passed to the details panel
  */
 class FDetailsViewObjectFilter
 {
 public:
-
 	virtual ~FDetailsViewObjectFilter() {}
 
-	/**
-	 * Given a const TArray<UObject*>& SourceObjects, it fills a TArray<FDetailsViewObjectRoot> with the objects
-	 * that we need as details objects roots. This may be the same as the original SourceObjects, or it may be
-	 * some subset of them, or some of their Sub-objects
-	 *
-	 * @param SourceObjects the array of objects acting as the Source of the root objects for the details view
-	 * @return the TArray<FDetailsViewObjectRoot> with the objects that will act as root objects for the details view 
-	 */
 	virtual TArray<FDetailsViewObjectRoot> FilterObjects(const TArray<UObject*>& SourceObjects) = 0;
 
-	/**
-	 * Returns the @code FDetailsViewStyleKey& @endcode that is the Key to the current objects' style
-	 */
-	virtual const FDetailsViewStyleKey& GetObjectsDetailsViewStyleKey()
-	{
-		return FDetailsViewStyleKeys::Default;
-	}
 };
diff --git a/Engine/Source/Editor/PropertyEditor/Public/DetailsViewStyleKey.h b/Engine/Source/Editor/PropertyEditor/Public/DetailsViewStyleKey.h
deleted file mode 100644
diff --git a/Engine/Source/Editor/PropertyEditor/Public/IDetailsView.h b/Engine/Source/Editor/PropertyEditor/Public/IDetailsView.h
--- a/Engine/Source/Editor/PropertyEditor/Public/IDetailsView.h
+++ b/Engine/Source/Editor/PropertyEditor/Public/IDetailsView.h
@@ -291,19 +291,4 @@ public:
 
 	/** Sets the minimum width of the detail panel's right column. */
 	virtual void SetRightColumnMinWidth(float InMinWidth) = 0;
-
-	/**
-	* Updates @code FDetailsViewStyleKey& StyleKey @endcode for the current IDetailsView state
-	*/
-	virtual void UpdateStyleKey() = 0;
-
-	/**
-	* Gets the @code FDetailsViewStyleKey& @endcode which provides a Key to the current style for a Details View
-	*/
-	virtual FDetailsViewStyleKey& GetStyleKey() = 0;
-	
-	/**
-   	* Returns a bool indicating whether the given @code FDetailsViewStyleKey @endcode is the default Details View Style 
-   	*/
-   	virtual bool IsDefaultStyle() const = 0;
 };
