Slate: Fixed FDeprecateSlateVector2D properties not being treated as vectors in the property editor

diff --git a/Engine/Source/Editor/PropertyEditor/Private/PropertyHandleImpl.cpp b/Engine/Source/Editor/PropertyEditor/Private/PropertyHandleImpl.cpp
--- a/Engine/Source/Editor/PropertyEditor/Private/PropertyHandleImpl.cpp
+++ b/Engine/Source/Editor/PropertyEditor/Private/PropertyHandleImpl.cpp
@@ -22,6 +22,7 @@
 
 #include "Framework/Notifications/NotificationManager.h"
 #include "Widgets/Notifications/SNotificationList.h"
+#include "Types/SlateVector2.h"
 #include "UObject/EnumProperty.h"
 #include "UObject/FieldPathProperty.h"
 #include "IDetailPropertyRow.h"
@@ -4479,7 +4480,8 @@ bool FPropertyHandleVector::Supports( TSharedRef<FPropertyNode> PropertyNode )
 		bSupported = StructName == NAME_Vector ||
 			StructName == NAME_Vector2D ||
 			StructName == NAME_Vector4 ||
-			StructName == NAME_Quat;
+			StructName == NAME_Quat ||
+			StructName == FDeprecateSlateVector2D::StaticStruct()->GetFName();
 	}
 
 	return bSupported;
