Fix PropertHandleRotator crash when no children not yet populated

diff --git a/Engine/Source/Editor/PropertyEditor/Private/PropertyHandleImpl.cpp b/Engine/Source/Editor/PropertyEditor/Private/PropertyHandleImpl.cpp
--- a/Engine/Source/Editor/PropertyEditor/Private/PropertyHandleImpl.cpp
+++ b/Engine/Source/Editor/PropertyEditor/Private/PropertyHandleImpl.cpp
@@ -4740,18 +4740,36 @@ bool FPropertyHandleRotator::Supports( TSharedRef<FPropertyNode> PropertyNode )
 FPropertyHandleRotator::FPropertyHandleRotator( TSharedRef<class FPropertyNode> PropertyNode, FNotifyHook* NotifyHook, TSharedPtr<IPropertyUtilities> PropertyUtilities )
 	: FPropertyHandleStruct( PropertyNode, NotifyHook, PropertyUtilities ) 
 {
-	const bool bRecurse = false;
-	// A vector is a struct property that has 3 children.  We get/set the values from the children
-	RollValue = MakeShareable( new FPropertyHandleMixed( Implementation->GetChildNode("Roll", bRecurse).ToSharedRef(), NotifyHook, PropertyUtilities ) );
+	if (Implementation->GetNumChildren() > 0)
+	{
+		const bool bRecurse = false;
+
+		// A rotator is a struct property that has 3 children.  We get/set the values from the children
+		if (TSharedPtr<FPropertyNode> RollNode = Implementation->GetChildNode("Roll", bRecurse))
+		{
+			RollValue = MakeShareable(new FPropertyHandleMixed(RollNode.ToSharedRef(), NotifyHook, PropertyUtilities));
+		}
 
-	PitchValue = MakeShareable( new FPropertyHandleMixed( Implementation->GetChildNode("Pitch", bRecurse).ToSharedRef(), NotifyHook, PropertyUtilities ) );
+		if (TSharedPtr<FPropertyNode> PitchNode = Implementation->GetChildNode("Pitch", bRecurse))
+		{
+			PitchValue = MakeShareable(new FPropertyHandleMixed(PitchNode.ToSharedRef(), NotifyHook, PropertyUtilities));
+		}
 
-	YawValue = MakeShareable( new FPropertyHandleMixed( Implementation->GetChildNode("Yaw", bRecurse).ToSharedRef(), NotifyHook, PropertyUtilities ) );
+		if (TSharedPtr<FPropertyNode> YawNode = Implementation->GetChildNode("Yaw", bRecurse))
+		{
+			YawValue = MakeShareable(new FPropertyHandleMixed(YawNode.ToSharedRef(), NotifyHook, PropertyUtilities));
+		}
+	}
 }
 
 
 FPropertyAccess::Result FPropertyHandleRotator::GetValue( FRotator& OutValue ) const
 {
+	if (!RollValue || !PitchValue || !YawValue)
+	{
+		return FPropertyAccess::Fail;
+	}
+
 	// To get the value from the rotator we read each child.  If reading a child fails, the value for that component is not set
 	FPropertyAccess::Result ResR = RollValue->GetValue( OutValue.Roll );
 	FPropertyAccess::Result ResP = PitchValue->GetValue( OutValue.Pitch );
@@ -4773,6 +4791,11 @@ FPropertyAccess::Result FPropertyHandleRotator::GetValue( FRotator& OutValue ) c
 
 FPropertyAccess::Result FPropertyHandleRotator::SetValue( const FRotator& NewValue, EPropertyValueSetFlags::Type Flags )
 {
+	if (!RollValue || !PitchValue || !YawValue)
+	{
+		return FPropertyAccess::Fail;
+	}
+
 	// To set the value from the rotator we set each child. 
 	FPropertyAccess::Result ResR = RollValue->SetValue( NewValue.Roll, Flags );
 	FPropertyAccess::Result ResP = PitchValue->SetValue( NewValue.Pitch, Flags );
@@ -4790,18 +4813,33 @@ FPropertyAccess::Result FPropertyHandleRotator::SetValue( const FRotator& NewVal
 
 FPropertyAccess::Result FPropertyHandleRotator::SetRoll( double InRoll, EPropertyValueSetFlags::Type Flags )
 {
+	if (!RollValue)
+	{
+		return FPropertyAccess::Fail;
+	}
+
 	FPropertyAccess::Result Res = RollValue->SetValue( InRoll, Flags );
 	return Res;
 }
 
 FPropertyAccess::Result FPropertyHandleRotator::SetPitch( double InPitch, EPropertyValueSetFlags::Type Flags )
 {
+	if (!PitchValue)
+	{
+		return FPropertyAccess::Fail;
+	}
+
 	FPropertyAccess::Result Res = PitchValue->SetValue( InPitch, Flags );
 	return Res;
 }
 
 FPropertyAccess::Result FPropertyHandleRotator::SetYaw( double InYaw, EPropertyValueSetFlags::Type Flags )
 {
+	if (!YawValue)
+	{
+		return FPropertyAccess::Fail;
+	}
+
 	FPropertyAccess::Result Res = YawValue->SetValue( InYaw, Flags );
 	return Res;
 }
