Added ScrollPropertyIntoView() to details view and fixed default property button padding
- Fixed padding of default buttons contaier (or else the widget created with CreateDefaultPropertyButtonWidgets() will unintentionally make a row taller)
- Added IDetailsView.ScrollPropertyIntoView(), similar to HighlightProperty() sans does not set highlight

diff --git a/Engine/Source/Editor/PropertyEditor/Private/PropertyHandleImpl.cpp b/Engine/Source/Editor/PropertyEditor/Private/PropertyHandleImpl.cpp
--- a/Engine/Source/Editor/PropertyEditor/Private/PropertyHandleImpl.cpp
+++ b/Engine/Source/Editor/PropertyEditor/Private/PropertyHandleImpl.cpp
@@ -2432,7 +2432,7 @@ class SDefaultPropertyButtonWidgets : public SCompoundWidget
 			.AutoWidth()
 			.HAlign(HAlign_Center)
 			.VAlign(VAlign_Center)
-			.Padding(2.0f, 1.0f)
+			.Padding(2.0f, 0.0f)
 			[
 				RequiredButton
 			];
diff --git a/Engine/Source/Editor/PropertyEditor/Private/SDetailsViewBase.cpp b/Engine/Source/Editor/PropertyEditor/Private/SDetailsViewBase.cpp
--- a/Engine/Source/Editor/PropertyEditor/Private/SDetailsViewBase.cpp
+++ b/Engine/Source/Editor/PropertyEditor/Private/SDetailsViewBase.cpp
@@ -324,6 +324,34 @@ void SDetailsViewBase::HighlightProperty(const FPropertyPath& Property)
 	CurrentlyHighlightedNode = TreeNode;
 }
 
+void SDetailsViewBase::ScrollPropertyIntoView(const FPropertyPath& Property, const bool bExpandProperty)
+{
+	if (!Property.IsValid())
+	{
+		return;
+	}
+
+	const TSharedPtr<FDetailTreeNode> TreeNode = FindBestFitTreeNodeFromProperty(RootTreeNodes, Property);
+	if (TreeNode.IsValid())
+	{
+		// make sure all ancestors are expanded so we can see the found node
+		TSharedPtr<FDetailTreeNode> Ancestor = TreeNode;
+		while(Ancestor->GetParentNode().IsValid())
+		{
+			Ancestor = Ancestor->GetParentNode().Pin();
+			DetailTree->SetItemExpansion(Ancestor.ToSharedRef(), true);
+		}
+
+		if (bExpandProperty)
+		{
+			DetailTree->SetItemExpansion(TreeNode.ToSharedRef(), true);
+		}
+		
+		// scroll to the found node
+		DetailTree->RequestScrollIntoView(TreeNode.ToSharedRef());
+	}
+}
+
 static void ExpandPaintSpacePropertyBoundsRecursive(const TArray<TSharedRef<FDetailTreeNode>>& TreeNodes, TSharedPtr<SDetailTree> DetailTree, FSlateRect& InOutRect)
 {
 	for(const TSharedRef<FDetailTreeNode>& TreeNode : TreeNodes)
diff --git a/Engine/Source/Editor/PropertyEditor/Private/SDetailsViewBase.h b/Engine/Source/Editor/PropertyEditor/Private/SDetailsViewBase.h
--- a/Engine/Source/Editor/PropertyEditor/Private/SDetailsViewBase.h
+++ b/Engine/Source/Editor/PropertyEditor/Private/SDetailsViewBase.h
@@ -86,6 +86,7 @@ public:
 	virtual TArray<TPair<int32, FPropertyPath>> GetPropertyRowNumbers() const override;
 	virtual int32 CountRows() const override;
 	virtual void HighlightProperty(const FPropertyPath& Property) override;
+	virtual void ScrollPropertyIntoView(const FPropertyPath& Property, bool bExpandProperty = false) override;
 	virtual FSlateRect GetPaintSpacePropertyBounds(const TSharedRef<FDetailTreeNode>& InDetailTreeNode, bool bIncludeChildren = true) const override;
 	virtual FSlateRect GetTickSpacePropertyBounds(const TSharedRef<FDetailTreeNode>& InDetailTreeNode, bool bIncludeChildren = true) const override;
 	virtual bool IsAncestorCollapsed(const TSharedRef<IDetailTreeNode>& Node) const override;
diff --git a/Engine/Source/Editor/PropertyEditor/Public/IDetailsView.h b/Engine/Source/Editor/PropertyEditor/Public/IDetailsView.h
--- a/Engine/Source/Editor/PropertyEditor/Public/IDetailsView.h
+++ b/Engine/Source/Editor/PropertyEditor/Public/IDetailsView.h
@@ -220,7 +220,12 @@ public:
 	 * Creates a box around the treenode corresponding to Property and scrolls the treenode into view
 	 */
 	virtual void HighlightProperty(const FPropertyPath& Property) = 0;
-	
+
+	/**
+	 * Scrolls the treenode for the specific property into view
+	 */
+	virtual void ScrollPropertyIntoView(const FPropertyPath& Property, bool bExpandProperty = false) = 0;
+
 	/**
 	 * Forces all advanced property sections to be in expanded state:
 	 */
