Support MustImplement, AllowedClasses and DisallowedClasses filtering for EditInline properties.

diff --git a/Engine/Source/Editor/PropertyEditor/Private/UserInterface/PropertyEditor/SPropertyEditorEditInline.cpp b/Engine/Source/Editor/PropertyEditor/Private/UserInterface/PropertyEditor/SPropertyEditorEditInline.cpp
--- a/Engine/Source/Editor/PropertyEditor/Private/UserInterface/PropertyEditor/SPropertyEditorEditInline.cpp
+++ b/Engine/Source/Editor/PropertyEditor/Private/UserInterface/PropertyEditor/SPropertyEditorEditInline.cpp
@@ -1,9 +1,11 @@
 // Copyright Epic Games, Inc. All Rights Reserved.
 
 #include "UserInterface/PropertyEditor/SPropertyEditorEditInline.h"
+#include "Algo/AnyOf.h"
 #include "Modules/ModuleManager.h"
 #include "Widgets/Images/SImage.h"
 #include "PropertyEditorHelpers.h"
+#include "PropertyRestriction.h"
 #include "ObjectPropertyNode.h"
 #include "PropertyHandleImpl.h"
 #include "ClassViewerModule.h"
@@ -26,46 +28,111 @@ public:
 
 	/** Hierarchy of objects that own this property. Used to check against ClassWithin. */
 	TSet< const UObject* > OwningObjects;
+	
+	/** The interface that must be implemented. */
+	const UClass* InterfaceThatMustBeImplemented;
+
+	/** Classes that can be picked */
+	TArray<const UClass*> AllowedClassFilters;
 
-	bool IsClassAllowed(const FClassViewerInitializationOptions& InInitOptions, const UClass* InClass, TSharedRef< FClassViewerFilterFuncs > InFilterFuncs ) override
+	/** Classes that can't be picked */
+	TArray<const UClass*> DisallowedClassFilters;
+
+	virtual bool IsClassAllowed(const FClassViewerInitializationOptions& InInitOptions, const UClass* InClass, TSharedRef< FClassViewerFilterFuncs > InFilterFuncs ) override
 	{
-		const bool bChildOfObjectClass = ObjProperty && InClass->IsChildOf(ObjProperty->PropertyClass);
-		const bool bDerivedInterfaceClass = IntProperty && InClass->ImplementsInterface(IntProperty->InterfaceClass);
+		return IsClassAllowedHelper(InClass, InFilterFuncs);
+	}
+	
+	virtual bool IsUnloadedClassAllowed(const FClassViewerInitializationOptions& InInitOptions, const TSharedRef< const IUnloadedBlueprintData > InBlueprint, TSharedRef< FClassViewerFilterFuncs > InFilterFuncs) override
+	{
+		return IsClassAllowedHelper(InBlueprint, InFilterFuncs);
+	}
+
+private:
 
-		const bool bMatchesFlags = InClass->HasAnyClassFlags(CLASS_EditInlineNew) && 
+	template <typename TClass>
+	bool IsClassAllowedHelper(TClass InClass, TSharedRef< FClassViewerFilterFuncs > InFilterFuncs)
+	{
+		const bool bMatchesFlags = InClass->HasAnyClassFlags(CLASS_EditInlineNew) &&
 			!InClass->HasAnyClassFlags(CLASS_Hidden | CLASS_HideDropDown | CLASS_Deprecated) &&
 			(bAllowAbstract || !InClass->HasAnyClassFlags(CLASS_Abstract));
 
-		if( (bChildOfObjectClass || bDerivedInterfaceClass) && bMatchesFlags )
+		if (!bMatchesFlags)
 		{
-			// Verify that the Owners of the property satisfy the ClassWithin constraint of the given class.
-			// When ClassWithin is null, assume it can be owned by anything.
-			return InClass->ClassWithin == nullptr || InFilterFuncs->IfMatchesAll_ObjectsSetIsAClass(OwningObjects, InClass->ClassWithin) != EFilterReturn::Failed;
+			return false;
 		}
 
-		return false;
-	}
+		// If it's neither an object property that is a child of the object class or an interface property that doesn't implement the interface, skip
+		const bool bChildOfObjectClass = ObjProperty && InClass->IsChildOf(ObjProperty->PropertyClass);
+		const bool bDerivedInterfaceClass = IntProperty && InClass->ImplementsInterface(IntProperty->InterfaceClass);
+		if (!bChildOfObjectClass && !bDerivedInterfaceClass)
+		{
+			return false;
+		}
 
-	virtual bool IsUnloadedClassAllowed(const FClassViewerInitializationOptions& InInitOptions, const TSharedRef< const IUnloadedBlueprintData > InUnloadedClassData, TSharedRef< FClassViewerFilterFuncs > InFilterFuncs) override
-	{
-		const bool bChildOfObjectClass = InUnloadedClassData->IsChildOf(ObjProperty->PropertyClass);
+		// If the class doesn't satisfy an interface requirement, skip
+		if ((InterfaceThatMustBeImplemented != nullptr) && !InClass->ImplementsInterface(InterfaceThatMustBeImplemented))
+		{
+			return false;
+		}
 
-		const bool bMatchesFlags = InUnloadedClassData->HasAnyClassFlags(CLASS_EditInlineNew) && 
-			!InUnloadedClassData->HasAnyClassFlags(CLASS_Hidden | CLASS_HideDropDown | CLASS_Deprecated) &&
-			(bAllowAbstract || !InUnloadedClassData->HasAnyClassFlags((CLASS_Abstract)));
+		auto IsClassChildOf = [InClass](const UClass* FilterClass)
+		{
+			return InClass->IsChildOf(FilterClass);
+		};
 
-		if (bChildOfObjectClass && bMatchesFlags)
+		// If the the class is explicitly present in the disallowed class filter set, we can't allow it
+		if (Algo::AnyOf(DisallowedClassFilters, IsClassChildOf))
 		{
-			const UClass* ClassWithin = InUnloadedClassData->GetClassWithin();
+			return false;
+		}
 
-			// Verify that the Owners of the property satisfy the ClassWithin constraint of the given class.
-			// When ClassWithin is null, assume it can be owned by anything.
-			return ClassWithin == nullptr || InFilterFuncs->IfMatchesAll_ObjectsSetIsAClass(OwningObjects, ClassWithin) != EFilterReturn::Failed;
+		// If the class is explicitly not present in the allowed class filter set, we can't allow it
+		if (!AllowedClassFilters.IsEmpty() && Algo::NoneOf(AllowedClassFilters, IsClassChildOf))
+		{
+			return false;
 		}
-		return false;
+
+		// Verify that the Owners of the property satisfy the ClassWithin constraint of the given class.
+		// When ClassWithin is null, assume it can be owned by anything.
+		const UClass* ClassWithin = GetClassWithin(InClass);
+		if (ClassWithin != nullptr && InFilterFuncs->IfMatchesAll_ObjectsSetIsAClass(OwningObjects, ClassWithin) ==
+			EFilterReturn::Failed)
+		{
+			return false;
+		}
+
+		return true;
+	}
+
+	static const UClass* GetClassWithin(const UClass* InClass)
+	{
+		return InClass->ClassWithin;
+	}
+
+	static const UClass* GetClassWithin(const TSharedRef< const IUnloadedBlueprintData > InBlueprint)
+	{
+		return InBlueprint->GetClassWithin();
 	}
 };
 
+namespace UE::PropertyEditor::EditInline::Private
+{
+
+UClass* FindOrLoadClass(const FString& ClassName)
+{
+	UClass* Class = UClass::TryFindTypeSlow<UClass>(ClassName, EFindFirstObjectOptions::EnsureIfAmbiguous);
+
+	if (!Class)
+	{
+		Class = LoadObject<UClass>(nullptr, *ClassName);
+	}
+
+	return Class;
+}
+
+} // namespace UE::PropertyEditor::EditInline::Private
+
 void SPropertyEditorEditInline::Construct( const FArguments& InArgs, const TSharedRef< class FPropertyEditor >& InPropertyEditor )
 {
 	PropertyEditor = InPropertyEditor;
@@ -178,6 +245,43 @@ TSharedRef<SWidget> SPropertyEditorEditInline::GenerateClassPicker()
 	FProperty* Property = PropertyNode->GetProperty();
 	ClassFilter->ObjProperty = CastField<FObjectPropertyBase>( Property );
 	ClassFilter->IntProperty = CastField<FInterfaceProperty>( Property );
+
+	// Filter based on UPROPERTY meta data
+	const FProperty* MetadataProperty = Property->GetOwnerProperty();
+
+	static const FName NAME_MustImplement(ANSITEXTVIEW("MustImplement"));
+	ClassFilter->InterfaceThatMustBeImplemented = MetadataProperty->GetClassMetaData(NAME_MustImplement);
+
+	static const FName NAME_AllowedClasses(ANSITEXTVIEW("AllowedClasses"));
+	TArray<const UClass*> AllowedClassFilters = PropertyCustomizationHelpers::GetClassesFromMetadataString(MetadataProperty->GetMetaData(NAME_AllowedClasses));
+
+	static const FName NAME_DisallowedClasses(ANSITEXTVIEW("DisallowedClasses"));
+	TArray<const UClass*> DisallowedClassFilters = PropertyCustomizationHelpers::GetClassesFromMetadataString(MetadataProperty->GetMetaData(NAME_DisallowedClasses));
+
+	using namespace UE::PropertyEditor::EditInline::Private;
+
+	// Filter based on restrictions
+	for (const TSharedRef<const FPropertyRestriction>& ClassRestriction : PropertyNode->GetRestrictions())
+	{
+		for (TArray<FString>::TConstIterator Iter = ClassRestriction.Get().GetHiddenValuesIterator(); Iter; ++Iter)
+		{
+			if (const UClass* HiddenClass = FindOrLoadClass(*Iter))
+			{
+				DisallowedClassFilters.Add(HiddenClass);
+			}
+		}
+
+		for (TArray<FString>::TConstIterator Iter = ClassRestriction.Get().GetDisabledValuesIterator(); Iter; ++Iter)
+		{
+			if (const UClass* DisabledClass = FindOrLoadClass(*Iter))
+			{
+				DisallowedClassFilters.Add(DisabledClass);
+			}
+		}
+	}
+
+	ClassFilter->AllowedClassFilters = MoveTemp(AllowedClassFilters);
+	ClassFilter->DisallowedClassFilters = MoveTemp(DisallowedClassFilters);
 	
 	bool bContainerHasNoClear = false;
 	if(PropertyNode->GetArrayIndex() != INDEX_NONE)
