Move SPropertyEditorClass's FindOrLoadClass out of the global namespace into UE::PropertyEditor::Class::Private to address build issue.

diff --git a/Engine/Source/Editor/PropertyEditor/Private/UserInterface/PropertyEditor/SPropertyEditorClass.cpp b/Engine/Source/Editor/PropertyEditor/Private/UserInterface/PropertyEditor/SPropertyEditorClass.cpp
--- a/Engine/Source/Editor/PropertyEditor/Private/UserInterface/PropertyEditor/SPropertyEditorClass.cpp
+++ b/Engine/Source/Editor/PropertyEditor/Private/UserInterface/PropertyEditor/SPropertyEditorClass.cpp
@@ -70,7 +70,10 @@ private:
 	}
 };
 
-static UClass* FindOrLoadClass(const FString& ClassName)
+namespace UE::PropertyEditor::Class::Private
+{
+
+UClass* FindOrLoadClass(const FString& ClassName)
 {
 	UClass* Class = UClass::TryFindTypeSlow<UClass>(ClassName, EFindFirstObjectOptions::EnsureIfAmbiguous);
 
@@ -82,6 +85,8 @@ static UClass* FindOrLoadClass(const FString& ClassName)
 	return Class;
 }
 
+} // namespace UE::PropertyEditor::Class::Private
+
 void SPropertyEditorClass::GetDesiredWidth(float& OutMinDesiredWidth, float& OutMaxDesiredWidth)
 {
 	OutMinDesiredWidth = 200.0f;
@@ -145,6 +150,8 @@ void SPropertyEditorClass::Construct(const FArguments& InArgs, const TSharedPtr<
 		AllowedClassFilters = PropertyCustomizationHelpers::GetClassesFromMetadataString(Property->GetOwnerProperty()->GetMetaData("AllowedClasses"));
 		DisallowedClassFilters = PropertyCustomizationHelpers::GetClassesFromMetadataString(Property->GetOwnerProperty()->GetMetaData("DisallowedClasses"));
 
+		using namespace UE::PropertyEditor::Class::Private;
+
 		// Filter based on restrictions
 		for (const TSharedRef<const FPropertyRestriction>& ClassRestriction : PropertyNode->GetRestrictions())
 		{
@@ -327,7 +334,7 @@ void SPropertyEditorClass::SendToObjects(const FString& NewValue)
 	}
 	else if (!NewValue.IsEmpty() && NewValue != TEXT("None"))
 	{
-		const UClass* NewClass = FindOrLoadClass(NewValue);
+		const UClass* NewClass = UE::PropertyEditor::Class::Private::FindOrLoadClass(NewValue);
 		OnSetClass.Execute(NewClass);
 	}
 	else
