Restore internal linkage to FindOrLoadClass functions in SPropertyEditorClass.cpp and SPropertyEditorEditInline.cpp

diff --git a/Engine/Source/Editor/PropertyEditor/Private/UserInterface/PropertyEditor/SPropertyEditorClass.cpp b/Engine/Source/Editor/PropertyEditor/Private/UserInterface/PropertyEditor/SPropertyEditorClass.cpp
--- a/Engine/Source/Editor/PropertyEditor/Private/UserInterface/PropertyEditor/SPropertyEditorClass.cpp
+++ b/Engine/Source/Editor/PropertyEditor/Private/UserInterface/PropertyEditor/SPropertyEditorClass.cpp
@@ -73,7 +73,7 @@ private:
 namespace UE::PropertyEditor::Class::Private
 {
 
-UClass* FindOrLoadClass(const FString& ClassName)
+static UClass* FindOrLoadClass(const FString& ClassName)
 {
 	UClass* Class = UClass::TryFindTypeSlow<UClass>(ClassName, EFindFirstObjectOptions::EnsureIfAmbiguous);
 
diff --git a/Engine/Source/Editor/PropertyEditor/Private/UserInterface/PropertyEditor/SPropertyEditorEditInline.cpp b/Engine/Source/Editor/PropertyEditor/Private/UserInterface/PropertyEditor/SPropertyEditorEditInline.cpp
--- a/Engine/Source/Editor/PropertyEditor/Private/UserInterface/PropertyEditor/SPropertyEditorEditInline.cpp
+++ b/Engine/Source/Editor/PropertyEditor/Private/UserInterface/PropertyEditor/SPropertyEditorEditInline.cpp
@@ -119,7 +119,7 @@ private:
 namespace UE::PropertyEditor::EditInline::Private
 {
 
-UClass* FindOrLoadClass(const FString& ClassName)
+static UClass* FindOrLoadClass(const FString& ClassName)
 {
 	UClass* Class = UClass::TryFindTypeSlow<UClass>(ClassName, EFindFirstObjectOptions::EnsureIfAmbiguous);
 
