Introduce validation for MustImplement meta tags for SPropertyEditorClass and SPropertyEditorEditInline to prevent non-interface classes being specified, and to protect against UInterface from being specified.

diff --git a/Engine/Source/Editor/PropertyEditor/Private/UserInterface/PropertyEditor/SPropertyEditorClass.cpp b/Engine/Source/Editor/PropertyEditor/Private/UserInterface/PropertyEditor/SPropertyEditorClass.cpp
--- a/Engine/Source/Editor/PropertyEditor/Private/UserInterface/PropertyEditor/SPropertyEditorClass.cpp
+++ b/Engine/Source/Editor/PropertyEditor/Private/UserInterface/PropertyEditor/SPropertyEditorClass.cpp
@@ -146,6 +146,20 @@ void SPropertyEditorClass::Construct(const FArguments& InArgs, const TSharedPtr<
 		bShowTree = Property->GetOwnerProperty()->HasMetaData(TEXT("ShowTreeView"));
 		bShowDisplayNames = Property->GetOwnerProperty()->HasMetaData(TEXT("ShowDisplayNames"));
 
+		if (RequiredInterface != nullptr)
+		{
+			if (!RequiredInterface->HasAnyClassFlags(CLASS_Interface))
+			{
+				UE_LOG(LogPropertyNode, Warning, TEXT("Property (%s) specifies a MustImplement class which isn't an interface (%s), clearing filter."), *Property->GetFullName(), *RequiredInterface->GetPathName());
+				RequiredInterface = nullptr;
+			}
+			else if (RequiredInterface == UInterface::StaticClass())
+			{
+				UE_LOG(LogPropertyNode, Warning, TEXT("Property (%s) specifies a MustImplement class which isn't valid (UInterface), clearing filter."), *Property->GetFullName());
+				RequiredInterface = nullptr;
+			}
+		}
+
 		// Filter based on UPROPERTY meta data
 		AllowedClassFilters = PropertyCustomizationHelpers::GetClassesFromMetadataString(Property->GetOwnerProperty()->GetMetaData("AllowedClasses"));
 		DisallowedClassFilters = PropertyCustomizationHelpers::GetClassesFromMetadataString(Property->GetOwnerProperty()->GetMetaData("DisallowedClasses"));
diff --git a/Engine/Source/Editor/PropertyEditor/Private/UserInterface/PropertyEditor/SPropertyEditorEditInline.cpp b/Engine/Source/Editor/PropertyEditor/Private/UserInterface/PropertyEditor/SPropertyEditorEditInline.cpp
--- a/Engine/Source/Editor/PropertyEditor/Private/UserInterface/PropertyEditor/SPropertyEditorEditInline.cpp
+++ b/Engine/Source/Editor/PropertyEditor/Private/UserInterface/PropertyEditor/SPropertyEditorEditInline.cpp
@@ -252,6 +252,20 @@ TSharedRef<SWidget> SPropertyEditorEditInline::GenerateClassPicker()
 	static const FName NAME_MustImplement(ANSITEXTVIEW("MustImplement"));
 	ClassFilter->InterfaceThatMustBeImplemented = MetadataProperty->GetClassMetaData(NAME_MustImplement);
 
+	if (ClassFilter->InterfaceThatMustBeImplemented != nullptr)
+	{
+		if (!ClassFilter->InterfaceThatMustBeImplemented->HasAnyClassFlags(CLASS_Interface))
+		{
+			UE_LOG(LogPropertyNode, Warning, TEXT("Property (%s) specifies a MustImplement class which isn't an interface (%s), clearing filter."), *Property->GetFullName(), *ClassFilter->InterfaceThatMustBeImplemented->GetPathName());
+			ClassFilter->InterfaceThatMustBeImplemented = nullptr;
+		}
+		else if (ClassFilter->InterfaceThatMustBeImplemented == UInterface::StaticClass())
+		{
+			UE_LOG(LogPropertyNode, Warning, TEXT("Property (%s) specifies a MustImplement class which isn't valid (UInterface), clearing filter."), *Property->GetFullName());
+			ClassFilter->InterfaceThatMustBeImplemented = nullptr;
+		}
+	}
+
 	static const FName NAME_AllowedClasses(ANSITEXTVIEW("AllowedClasses"));
 	TArray<const UClass*> AllowedClassFilters = PropertyCustomizationHelpers::GetClassesFromMetadataString(MetadataProperty->GetMetaData(NAME_AllowedClasses));
 
