Details Panel - Component card layout actions

The categories on the components for core entities now have an action menu. It has many things in it that will log.

diff --git a/Engine/Source/Developer/WidgetRegistration/Private/ToolElementRegistry.cpp b/Engine/Source/Developer/WidgetRegistration/Private/ToolElementRegistry.cpp
--- a/Engine/Source/Developer/WidgetRegistration/Private/ToolElementRegistry.cpp
+++ b/Engine/Source/Developer/WidgetRegistration/Private/ToolElementRegistry.cpp
@@ -16,7 +16,14 @@ void FToolElementRegistrationArgs::ResetWidget()
 
 FToolElementRegistrationArgs::FToolElementRegistrationArgs(
 	EToolElement InToolElementType) :
-	ToolElementType(InToolElementType)
+	ToolElementType(InToolElementType),
+	StyleClassName("Section")
+{
+}
+
+FToolElementRegistrationArgs::FToolElementRegistrationArgs( FName InStyleClassName ) :
+	ToolElementType(EToolElement::Section),
+	StyleClassName(InStyleClassName)
 {
 }
 	
diff --git a/Engine/Source/Developer/WidgetRegistration/Public/Layout/SeparatorBuilder.h b/Engine/Source/Developer/WidgetRegistration/Public/Layout/SeparatorBuilder.h
--- a/Engine/Source/Developer/WidgetRegistration/Public/Layout/SeparatorBuilder.h
+++ b/Engine/Source/Developer/WidgetRegistration/Public/Layout/SeparatorBuilder.h
@@ -86,7 +86,7 @@ public:
 	 *
  	 * @return the TSharedRef<SSeparator> as a reference to the generated SSeparator
 	 */
-	TSharedRef<SSeparator> operator*();
+	WIDGETREGISTRATION_API TSharedRef<SSeparator> operator*();
 
 	/**
 	 * Initializes the size of the separator. If called after the widget creation (that is, a call to GenerateWidget)
@@ -95,7 +95,7 @@ public:
 	 * @param NewSize the FSeparatorSize& that will provide the size for the Separator
 	 * @return this FSeparatorBuilder to enable cascading.
 	 */
-	FSeparatorBuilder& InitializeSize(const FSeparatorSize& NewSize);
+	WIDGETREGISTRATION_API FSeparatorBuilder& InitializeSize(const FSeparatorSize& NewSize);
 
 	/** 
 	 * Sets the color of the Separator. This will update the color of the Separator even after widget creation.
@@ -105,7 +105,7 @@ public:
 	 * for that EStyleColor immediately.
 	 * @return this FSeparatorBuilder to enable cascading.
 	 */
-	FSeparatorBuilder& SetColor(const EStyleColor& NewColor);
+	WIDGETREGISTRATION_API FSeparatorBuilder& SetColor(const EStyleColor& NewColor);
 
 	/** 
 	 * Binds the Visibility of the Separator to a TAttribute for continuous update.
@@ -113,7 +113,7 @@ public:
 	 * @param NewVisibility the TAttribute<EVisibility> to bind for the SSeparator. 
 	 * @return this FSeparatorBuilder to enable cascading.
 	 */
-	FSeparatorBuilder& BindVisibility(TAttribute<EVisibility> NewVisibility)
+	WIDGETREGISTRATION_API FSeparatorBuilder& BindVisibility(TAttribute<EVisibility> NewVisibility)
 	{
 		Visibility = NewVisibility;
 		return *this;
diff --git a/Engine/Source/Developer/WidgetRegistration/Public/Layout/SeparatorTemplates.h b/Engine/Source/Developer/WidgetRegistration/Public/Layout/SeparatorTemplates.h
--- a/Engine/Source/Developer/WidgetRegistration/Public/Layout/SeparatorTemplates.h
+++ b/Engine/Source/Developer/WidgetRegistration/Public/Layout/SeparatorTemplates.h
@@ -17,7 +17,7 @@ struct FSeparatorTemplates
 	 * @return a small (2 Slate Unit), vertical FSeparatorBuilder with the color of the theme
 	 * specified color for Background
 	 */
-	static FSeparatorBuilder SmallVerticalBackgroundNoBorder();
+	WIDGETREGISTRATION_API static FSeparatorBuilder SmallVerticalBackgroundNoBorder();
 
 	/**
 	* Creates and returns a small (2 Slate Unit), horizontal FSeparatorBuilder set to the color of the theme
@@ -26,7 +26,7 @@ struct FSeparatorTemplates
 	* @return a small (2 Slate Unit), horizontal FSeparatorBuilder with the color of the theme
 	* specified color for Background
 	 */
-	static FSeparatorBuilder SmallHorizontalBackgroundNoBorder();
+	WIDGETREGISTRATION_API static FSeparatorBuilder SmallHorizontalBackgroundNoBorder();
 
 	/**
 	* Creates and returns a small (2 Slate Unit), horizontal FSeparatorBuilder set to the color of the theme
@@ -35,5 +35,5 @@ struct FSeparatorTemplates
 	 * @return a small (2 Slate Unit), horizontal FSeparatorBuilder with the color of the theme
 	 * specified color for Panels.
 	 */
-	static FSeparatorBuilder SmallHorizontalPanelNoBorder();
+	WIDGETREGISTRATION_API static FSeparatorBuilder SmallHorizontalPanelNoBorder();
 };
diff --git a/Engine/Source/Developer/WidgetRegistration/Public/ToolElementRegistry.h b/Engine/Source/Developer/WidgetRegistration/Public/ToolElementRegistry.h
--- a/Engine/Source/Developer/WidgetRegistration/Public/ToolElementRegistry.h
+++ b/Engine/Source/Developer/WidgetRegistration/Public/ToolElementRegistry.h
@@ -11,7 +11,8 @@ enum class EToolElement : uint32
 {
 	Toolbar = 0,
 	Toolkit,
-	Separator
+	Separator,
+	Section
 };
 
 /** A class which provides the data necessary to generate a widget. This class should be
@@ -26,6 +27,9 @@ public:
 	/** The constructor, which takes an EToolElement which will define the type of Tool Element*/
 	FToolElementRegistrationArgs(EToolElement InToolElementType);
 
+	/** The constructor, which takes an EToolElement which will define the type of Tool Element*/
+	FToolElementRegistrationArgs(FName InStyleClassName);
+
 	/** default destructor in case any subclasses need to provide a destructor*/
 	virtual ~FToolElementRegistrationArgs() = default;
 
@@ -40,6 +44,9 @@ public:
 
 	/** The type of tool element this is Registration args for */
 	const EToolElement ToolElementType;
+
+	/** The style class name */
+	const FName StyleClassName;
 };
 
 /** Serves as a key into the FToolElementRegistry of FToolElements */
diff --git a/Engine/Source/Editor/EditorStyle/Private/StarshipStyle.cpp b/Engine/Source/Editor/EditorStyle/Private/StarshipStyle.cpp
--- a/Engine/Source/Editor/EditorStyle/Private/StarshipStyle.cpp
+++ b/Engine/Source/Editor/EditorStyle/Private/StarshipStyle.cpp
@@ -3015,15 +3015,30 @@ void FStarshipEditorStyle::FStyle::SetupPropertyEditorStyles()
 		
 		/****** Styles for rounded corners for the Card style of a Details View ********/
 
-		Set( "DetailsView.CardHeaderTopLeftSideRounded",               new FSlateRoundedBoxBrush(FStyleColors::Header, FVector4(4.0f, 0.0f, 0.0f, 0.0f)));
-		Set( "DetailsView.CardHeaderLeftSideRounded",               new FSlateRoundedBoxBrush(FStyleColors::Header, FVector4(4.0f, 0.0f, 0.0f, 4.0f)));
-		Set( "DetailsView.CardHeaderTopRightSideRounded",               new FSlateRoundedBoxBrush(FStyleColors::Header, FVector4(0.0f, 4.0f, 0.0f, 0.0f)));
-		Set( "DetailsView.CardHeaderRightSideRounded",               new FSlateRoundedBoxBrush(FStyleColors::Header, FVector4(0.0f, 4.0f, 4.0f, 0.0f)));
-		Set( "DetailsView.CardHeaderTopRounded",               new FSlateRoundedBoxBrush(FStyleColors::Header, FVector4(4.0f, 4.0f, 0.0f, 0.0f)));
-		Set( "DetailsView.CardHeaderRounded",               new FSlateRoundedBoxBrush(FStyleColors::Header, FVector4(4.0f, 4.0f, 4.0f, 4.0f )));
+		Set( "DetailsView.CardHeaderTopLeftSideRounded", new FSlateRoundedBoxBrush(FStyleColors::Header, FVector4(4.0f, 0.0f, 0.0f, 0.0f)));
+		Set( "DetailsView.CardHeaderLeftSideRounded", new FSlateRoundedBoxBrush(FStyleColors::Header, FVector4(4.0f, 0.0f, 0.0f, 4.0f)));
+		Set( "DetailsView.CardHeaderTopRightSideRounded", new FSlateRoundedBoxBrush(FStyleColors::Header, FVector4(0.0f, 4.0f, 0.0f, 0.0f)));
+		Set( "DetailsView.CardHeaderRightSideRounded", new FSlateRoundedBoxBrush(FStyleColors::Header, FVector4(0.0f, 4.0f, 4.0f, 0.0f)));
+		Set( "DetailsView.CardHeaderTopRounded", new FSlateRoundedBoxBrush(FStyleColors::Header, FVector4(4.0f, 4.0f, 0.0f, 0.0f)));
+		Set( "DetailsView.CardHeaderRounded", new FSlateRoundedBoxBrush(FStyleColors::Header, FVector4(4.0f, 4.0f, 4.0f, 4.0f )));
 		
 		/*******************************************************************************/
-		
+
+		const FButtonStyle CategoryButton = FButtonStyle()
+			.SetNormal(FSlateRoundedBoxBrush(FStyleColors::Header, 0.f))
+			.SetHovered(FSlateRoundedBoxBrush(FStyleColors::Input, 0.f))
+			.SetPressed(FSlateRoundedBoxBrush(FStyleColors::Header, 0.f))
+			.SetDisabled(FSlateRoundedBoxBrush(FStyleColors::Header, 0.f))
+			.SetNormalPadding(FMargin(2.f, 0.f, 2.f, 0.f))
+			.SetPressedPadding(FMargin(2.f, 0.f, 2.f, 0.f));
+
+		FComboButtonStyle CategoryComboButton = FComboButtonStyle(FStarshipCoreStyle::GetCoreStyle().GetWidgetStyle<FComboButtonStyle>("ComboButton"))
+		.SetButtonStyle(CategoryButton)
+		.SetDownArrowPadding(FMargin(2.f, 5.f, 3.f, 5.f))
+		.SetDownArrowImage(CORE_IMAGE_BRUSH_SVG("Starship/Common/ellipsis-vertical-narrow", FVector2D(6, 15)));
+		CategoryComboButton.ButtonStyle = CategoryButton;
+		Set( "DetailsView.CategoryComboButton", CategoryComboButton);
+
 		Set( "DetailsView.CategoryTop_Hovered",       new FSlateColorBrush(FStyleColors::Hover));
 		Set( "DetailsView.CategoryBottom",            new FSlateColorBrush(FStyleColors::Recessed));
 		
diff --git a/Engine/Source/Editor/LevelEditor/Private/SActorDetails.cpp b/Engine/Source/Editor/LevelEditor/Private/SActorDetails.cpp
--- a/Engine/Source/Editor/LevelEditor/Private/SActorDetails.cpp
+++ b/Engine/Source/Editor/LevelEditor/Private/SActorDetails.cpp
@@ -45,6 +45,11 @@
 
 #define LOCTEXT_NAMESPACE "SActorDetails"
 
+static TAutoConsoleVariable<bool> CVarForceShowComponentEditor(
+	TEXT("CoreEntity.UI.ForceShowComponentEditor"),
+	false,
+	TEXT("Force the component editor to show in the main details tab."));
+
 namespace UE::LevelEditor::Private
 {
 	class SElementSelectionDetailsButtons : public SCompoundWidget
@@ -617,9 +622,13 @@ FString SActorDetails::GetReferencerName() const
 	return TEXT("SActorDetails");
 }
 
-void SActorDetails::SetActorDetailsRootCustomization(TSharedPtr<FDetailsViewObjectFilter> InActorDetailsObjectFilter, TSharedPtr<IDetailRootObjectCustomization> ActorDetailsRootCustomization)
+	void SActorDetails::SetActorDetailsRootCustomization(TSharedPtr<FDetailsViewObjectFilter> InActorDetailsObjectFilter, TSharedPtr<IDetailRootObjectCustomization> ActorDetailsRootCustomization)
 {
-	ActorDetailsObjectFilter = InActorDetailsObjectFilter.ToWeakPtr();
+	if (InActorDetailsObjectFilter.IsValid())
+	{
+		DisplayManager = InActorDetailsObjectFilter->GetDisplayManager();
+	}
+	
 	DetailsView->SetObjectFilter(InActorDetailsObjectFilter);
 	DetailsView->SetRootObjectCustomizationInstance(ActorDetailsRootCustomization);
 	DetailsView->ForceRefresh();
@@ -881,12 +890,10 @@ void SActorDetails::OnNativeComponentWarningHyperlinkClicked(const FSlateHyperli
 
 EVisibility SActorDetails::GetComponentEditorVisibility() const
 {
-	IConsoleVariable* ForceShow = IConsoleManager::Get().FindConsoleVariable(TEXT("CoreEntity.UI.ForceShowComponentEditor"));
-
 	// force hide it if the style is not default and the ForceShowComponentEditor CVar is not set to true
-	const bool bHideEditorFromDetailsView = (ActorDetailsObjectFilter.IsValid() &&
-							                 ActorDetailsObjectFilter.Pin()->ShouldHideComponentEditor()) &&
-		                    (!ForceShow || !ForceShow->GetBool());
+	const bool bHideEditorFromDetailsView = (DisplayManager.IsValid() &&
+							                 DisplayManager->ShouldHideComponentEditor()) &&
+											(!CVarForceShowComponentEditor.GetValueOnAnyThread());
 	return GetActorContext() && !bHideEditorFromDetailsView  ? EVisibility::Visible : EVisibility::Collapsed;
 }
 
diff --git a/Engine/Source/Editor/LevelEditor/Private/SActorDetails.h b/Engine/Source/Editor/LevelEditor/Private/SActorDetails.h
--- a/Engine/Source/Editor/LevelEditor/Private/SActorDetails.h
+++ b/Engine/Source/Editor/LevelEditor/Private/SActorDetails.h
@@ -25,6 +25,7 @@ class SSubobjectEditor;
 class ISCSEditorUICustomization;
 class FSubobjectEditorTreeNode;
 struct FPropertyChangedEvent;
+class FDetailsDisplayManager;
 
 namespace UE::LevelEditor::Private
 {
@@ -143,6 +144,7 @@ private:
 	// Used to prevent reentrant changes
 	bool bSelectionGuard = false;
 
-	/** The object filter which contains relevant information about the details objects being displayed */
-	TWeakPtr<FDetailsViewObjectFilter> ActorDetailsObjectFilter;
+	/** A @code TSharedPtr @endcode to the @code FDetailsDisplayManager @endcode to manage various aspects of the
+	 *details view */
+	TSharedPtr<FDetailsDisplayManager> DisplayManager;
 };
diff --git a/Engine/Source/Editor/PropertyEditor/Private/DetailAdvancedDropdownNode.cpp b/Engine/Source/Editor/PropertyEditor/Private/DetailAdvancedDropdownNode.cpp
--- a/Engine/Source/Editor/PropertyEditor/Private/DetailAdvancedDropdownNode.cpp
+++ b/Engine/Source/Editor/PropertyEditor/Private/DetailAdvancedDropdownNode.cpp
@@ -108,7 +108,7 @@ public:
 		[
 			SNew(SBorder)
 			.BorderImage(FAppStyle::Get().GetBrush("DetailsView.GridLine"))
-			.Padding(ViewStyle.GetRowPadding())
+			.Padding( 0, 0, 0, 1 )
 			[
 				SNew(SHorizontalBox)
 				+ SHorizontalBox::Slot()
@@ -122,16 +122,6 @@ public:
 						ContentWidget.ToSharedRef()
 					]
 				]
-				+ SHorizontalBox::Slot()
-				.HAlign(HAlign_Right)
-				.VAlign(VAlign_Fill)
-				.AutoWidth()
-				[
-					SNew(SBorder)
-					.BorderImage_Lambda(GetScrollbarWellBrush)
-					.BorderBackgroundColor_Lambda(GetScrollbarWellTint)
-					.Padding(FMargin(0, 0, SDetailTableRowBase::ScrollBarPadding, 0))
-				]
 			]
 		];
 
diff --git a/Engine/Source/Editor/PropertyEditor/Private/DetailCategoryBuilderImpl.cpp b/Engine/Source/Editor/PropertyEditor/Private/DetailCategoryBuilderImpl.cpp
--- a/Engine/Source/Editor/PropertyEditor/Private/DetailCategoryBuilderImpl.cpp
+++ b/Engine/Source/Editor/PropertyEditor/Private/DetailCategoryBuilderImpl.cpp
@@ -818,15 +818,39 @@ TSharedRef<ITableRow> FDetailCategoryImpl::GenerateWidgetForTableView(const TSha
 		HeaderContent = Row.ValueWidget.Widget;
 	}
 
+	InitializeObjectName();
 	TSharedPtr<FDetailLayoutBuilderImpl> ParentLayout = GetParentLayoutImpl();
 
 	return SNew(SDetailCategoryTableRow, AsShared(), OwnerTable)
 		.PasteFromText(OnPasteFromText())
+		.ObjectName( ObjectName )
 		.InnerCategory(ParentLayout.IsValid() ? ParentLayout->IsLayoutForExternalRoot() : false)
 		.DisplayName(GetDisplayName())
 		.HeaderContent(HeaderContent);
 }
 
+void FDetailCategoryImpl::InitializeObjectName()
+{
+	if (!DetailLayoutBuilder.IsValid())
+	{
+		return;
+	}
+	
+	const TSharedPtr<FComplexPropertyNode> Node = DetailLayoutBuilder.Pin()->GetRootNode();
+	
+	if (Node.IsValid())
+	{
+		if (const FObjectPropertyNode* Object = Node->AsObjectNode())
+		{
+			if (const UObject* CategoryObject = Object->GetUObject(0))
+			{
+				const FName Name{CategoryObject->GetName()};
+				ObjectName = Name;
+			}
+		}
+	}	
+}
+
 bool FDetailCategoryImpl::GenerateStandaloneWidget(FDetailWidgetRow& OutRow) const
 {
 	TSharedPtr<SWidget> HeaderContent = HeaderContentWidget;
diff --git a/Engine/Source/Editor/PropertyEditor/Private/DetailCategoryBuilderImpl.h b/Engine/Source/Editor/PropertyEditor/Private/DetailCategoryBuilderImpl.h
--- a/Engine/Source/Editor/PropertyEditor/Private/DetailCategoryBuilderImpl.h
+++ b/Engine/Source/Editor/PropertyEditor/Private/DetailCategoryBuilderImpl.h
@@ -399,6 +399,9 @@ private:
 	 */
 	void RefreshTreeInternal(bool bRefilterCategory);
 
+	/** Initialize the name of the Category  */
+	void InitializeObjectName();
+
 private:
 	/** Layouts that appear in this category category */
 	FDetailLayoutMap LayoutMap;
@@ -423,6 +426,8 @@ private:
 	/** Delegate handling pasting an optionally tagged text snippet */
 	TSharedPtr<FOnPasteFromText> PasteFromTextDelegate;
 
+	/** The name of the object that this category is being defined by */
+	FName ObjectName;
 	/** The category identifier */
 	FName CategoryName;
 	/** The sort order of this category (amongst all categories) */
diff --git a/Engine/Source/Editor/PropertyEditor/Private/DetailsDisplayManager.cpp b/Engine/Source/Editor/PropertyEditor/Private/DetailsDisplayManager.cpp
new file mode 100644
--- /dev/null
+++ b/Engine/Source/Editor/PropertyEditor/Private/DetailsDisplayManager.cpp
@@ -0,0 +1,83 @@
+//  Copyright Epic Games, Inc. All Rights Reserved.
+
+#include "DetailsDisplayManager.h"
+
+#include "DetailsViewStyle.h"
+
+
+FDetailsDisplayManager::~FDetailsDisplayManager()
+{
+	OnDetailsNeedsUpdate.Unbind();
+}
+
+bool FDetailsDisplayManager::ShouldHideComponentEditor()
+{
+	return false;
+}
+
+bool FDetailsDisplayManager::ShouldShowCategoryMenu()
+{
+	return false;
+}
+
+void FDetailsDisplayManager::SetCategoryObjectName(FName InCategoryObjectName)
+{
+	CategoryObjectName = InCategoryObjectName;
+}
+
+TSharedPtr<SWidget> FDetailsDisplayManager::GetCategoryMenu(FName InCategoryObjectName)
+{
+	return nullptr;
+}
+
+void FDetailsDisplayManager::UpdateView() const
+{
+	OnDetailsNeedsUpdate.ExecuteIfBound();
+}
+
+const FDetailsViewStyleKey& FDetailsDisplayManager::GetDetailsViewStyleKey() const
+{
+	return FDetailsViewStyleKeys::Default();
+}
+
+void FDetailsDisplayManager::SetIsOuterCategory(bool bInIsOuterCategory)
+{
+	bIsOuterCategory = bInIsOuterCategory;
+}
+
+FMargin FDetailsDisplayManager::GetTablePadding() const
+{
+	if (GetDetailsViewStyleKey() == FDetailsViewStyleKeys::Card())
+	{
+		return bIsScrollbarShowing ?
+					FDetailsViewStyle::TablePaddingWithScrollbarCard() :
+					FDetailsViewStyle::TablePaddingNoScrollbarCard();
+	}
+	return bIsScrollbarShowing ?
+					FDetailsViewStyle::TablePaddingWithScrollbarClassic() :
+					FDetailsViewStyle::TablePaddingNoScrollbarClassic();
+		
+}
+
+FMargin FDetailsDisplayManager::GetRowPadding() const
+{
+	FDetailsViewStyle ViewStyle = GetDetailsViewStyleKey();
+	ViewStyle.SetIsOuterCategory(bIsOuterCategory);
+	ViewStyle.SetIsScrollbarShowing(bIsScrollbarShowing);
+	return ViewStyle.GetRowPadding();
+}
+
+FMargin FDetailsDisplayManager::GetCategoryButtonsPadding() const
+{
+	return FMargin(0, 0,  bIsOuterCategory ? 10 : 0, 0);
+}
+
+bool FDetailsDisplayManager::GetIsScrollbarShowing() const
+{
+	return bIsScrollbarShowing;
+}
+
+void FDetailsDisplayManager::SetIsScrollbarShowing(bool bInIsScrollbarShowing)
+{
+	bIsScrollbarShowing = bInIsScrollbarShowing;
+}
diff --git a/Engine/Source/Editor/PropertyEditor/Private/DetailsViewStyle.cpp b/Engine/Source/Editor/PropertyEditor/Private/DetailsViewStyle.cpp
--- a/Engine/Source/Editor/PropertyEditor/Private/DetailsViewStyle.cpp
+++ b/Engine/Source/Editor/PropertyEditor/Private/DetailsViewStyle.cpp
@@ -82,7 +82,6 @@ const FSlateBrush* FDetailsViewStyle::GetBackgroundImageForCategoryRow(
 	static const FSlateBrush* CardStyleTopLevelCategoryCollapsedScrollbarVisibleRowBrush = FAppStyle::Get().GetBrush("DetailsView.CardHeaderRounded");
 	static const FSlateBrush* CardStyleTopLevelCategoryCollapsedScrollbarHiddenRowBrush = FAppStyle::Get().GetBrush("DetailsView.CardHeaderLeftSideRounded");
 	static const FSlateBrush* CardStyleTopLevelCategoryExpandedScrollbarVisibleRowBrush = FAppStyle::Get().GetBrush("DetailsView.CardHeaderTopRounded");
-	static const FSlateBrush* CardStyleTopLevelCategoryExpandedScrollbarHiddenRowBrush = FAppStyle::Get().GetBrush("DetailsView.CardHeaderTopLeftSideRounded");
 
 	if (bShowBorder)
 	{
@@ -105,11 +104,7 @@ const FSlateBrush* FDetailsViewStyle::GetBackgroundImageForCategoryRow(
 			}
 			return CardStyleTopLevelCategoryCollapsedScrollbarHiddenRowBrush;
 		}
-		if (bIsScrollBarVisible)
-		{
-			return CardStyleTopLevelCategoryExpandedScrollbarVisibleRowBrush;
-		}
-		return CardStyleTopLevelCategoryExpandedScrollbarHiddenRowBrush;
+		return CardStyleTopLevelCategoryExpandedScrollbarVisibleRowBrush;
 	}
 
 	return nullptr;
@@ -117,7 +112,7 @@ const FSlateBrush* FDetailsViewStyle::GetBackgroundImageForCategoryRow(
 
 void FDetailsViewStyle::InitializeDetailsViewStyles()
 {
-	static FDetailsViewStyle CardStyle{FDetailsViewStyleKeys::Card(), 8.0f, 8.0f};
+	static FDetailsViewStyle CardStyle{FDetailsViewStyleKeys::Card(), 8.0f, 0.0f};
 	StyleKeyToStyleTemplateMap.Add(FDetailsViewStyleKeys::Card().GetName(), &CardStyle);
 
 	static FDetailsViewStyle DefaultStyle{FDetailsViewStyleKeys::Default(), 0.0f, 0.0f};
@@ -145,18 +140,14 @@ const FSlateBrush* FDetailsViewStyle::GetBackgroundImageForScrollBarWell(
 		}
 		const bool bIsCardStyle = this->Key == FDetailsViewStyleKeys::Card();
 			
-		if (!bIsScrollBarVisible)
+		if (!bIsCardStyle)
 		{
-			if (!bIsCardStyle)
-			{
-				return ClassicStyleTopLevelCategoryRowBrush;
-			}
-			if (!bIsCategoryExpanded){
-				return CardStyleCollapsedScrollbarVisibleWellBrush;
-			}
-			return CardStyleExpandedScrollbarVisibleWellBrush;
+			return ClassicStyleTopLevelCategoryRowBrush;
+		}
+		if (!bIsCategoryExpanded){
+			return CardStyleCollapsedScrollbarVisibleWellBrush;
 		}
-		return DetailsBackgroundWellBrush;
+		return CardStyleExpandedScrollbarVisibleWellBrush;
 	}
 
 	return nullptr;
@@ -187,4 +178,9 @@ const FDetailsViewStyle* FDetailsViewStyle::GetStyle(FDetailsViewStyleKey Key)
 void FDetailsViewStyle::SetIsOuterCategory(bool bInIsOuterCategory)
 {
 	bIsOuterCategory = bInIsOuterCategory;
+}
+
+void FDetailsViewStyle::SetIsScrollbarShowing(bool bInIsScrollbarShowing)
+{
+	bIsScrollbarShowing = bInIsScrollbarShowing;
 }
\ No newline at end of file
diff --git a/Engine/Source/Editor/PropertyEditor/Private/DetailsViewStyle.h b/Engine/Source/Editor/PropertyEditor/Private/DetailsViewStyle.h
--- a/Engine/Source/Editor/PropertyEditor/Private/DetailsViewStyle.h
+++ b/Engine/Source/Editor/PropertyEditor/Private/DetailsViewStyle.h
@@ -56,6 +56,13 @@ public:
 	*/
 	void SetIsOuterCategory(bool bIsOuterCategory);
 
+	/**
+	* Sets a bool indicating whether or not the Scrollbar is showing on the details view
+	*
+	* @param bInIsScrollbarShowing a bool indicating whether or not the Scrollbar is showing on the details view
+	*/
+	void SetIsScrollbarShowing(bool bInIsScrollbarShowing);
+	
 	/**
 	 * Returns the padding for details panel rows which are not outer Category rows
 	 */
@@ -120,6 +127,35 @@ public:
 		const FDetailsViewStyleKey& InKey,
 		float InTopCategoryPadding = 0.f,
 		float InHorizontalPadding = 0.f);
+
+	
+	/** the Padding for the entirety of the details panel if no scrollbar is present for Classic style */
+	static const FMargin& TablePaddingNoScrollbarClassic()
+	{
+		static FMargin Margin = FMargin(0, 0, 0, 1);
+		return Margin;
+	}
+	
+	/** the Padding for the entirety of the details panel if no scrollbar is present for Card style */
+	static const FMargin& TablePaddingNoScrollbarCard()
+	{
+		static FMargin Margin = FMargin(8, 0, 8, 1);
+		return Margin;
+	}
+
+	/** the Padding for the entirety of the details panel if the scrollbar is present for Classic style */
+	static const FMargin TablePaddingWithScrollbarClassic()
+	{
+		static FMargin Margin = FMargin(0, 0, 16, 1);
+		return Margin;
+	}
+
+	/** the Padding for the entirety of the details panel if the scrollbar is present for Card style  */
+	static const FMargin TablePaddingWithScrollbarCard()
+	{
+		static FMargin Margin = FMargin(8, 0, 20, 1);
+		return Margin;
+	}
 	
 private:
 	/**
@@ -165,5 +201,8 @@ private:
 	
 	/** Whether the current Category is an outer versus an inner Category  */
    	bool bIsOuterCategory = false;
+
+	/** Whether the scrollbar is showing on the details view  */
+	bool bIsScrollbarShowing = false;
 };
 
diff --git a/Engine/Source/Editor/PropertyEditor/Private/IDetailsViewPrivate.h b/Engine/Source/Editor/PropertyEditor/Private/IDetailsViewPrivate.h
--- a/Engine/Source/Editor/PropertyEditor/Private/IDetailsViewPrivate.h
+++ b/Engine/Source/Editor/PropertyEditor/Private/IDetailsViewPrivate.h
@@ -3,6 +3,7 @@
 #pragma once
 
 #include "AssetThumbnail.h"
+#include "DetailsDisplayManager.h"
 #include "DetailTreeNode.h"
 #include "IDetailsView.h"
 #include "PropertyNode.h"
@@ -199,4 +200,11 @@ public:
 	{
 		return true;
 	}
+
+	/**
+	 * Returns a @code TSharedPtr @endcode to the @code FDetailsDisplayManager @endcode for this
+	 * details view
+	 */
+	virtual TSharedPtr<FDetailsDisplayManager> GetDisplayManager() = 0;
+
 };
diff --git a/Engine/Source/Editor/PropertyEditor/Private/SDetailCategoryTableRow.cpp b/Engine/Source/Editor/PropertyEditor/Private/SDetailCategoryTableRow.cpp
--- a/Engine/Source/Editor/PropertyEditor/Private/SDetailCategoryTableRow.cpp
+++ b/Engine/Source/Editor/PropertyEditor/Private/SDetailCategoryTableRow.cpp
@@ -15,7 +15,10 @@
 #include "SDetailsView.h"
 #include "Serialization/JsonSerializer.h"
 #include "Styling/StyleColors.h"
+#include "UserInterface/Categories/CategoryMenuComboButtonBuilder.h"
 #include "UserInterface/PropertyEditor/PropertyEditorConstants.h"
+#include "Widgets/Input/SComboButton.h"
+#include "Widgets/Layout/SSeparator.h"
 
 void SDetailCategoryTableRow::Construct(const FArguments& InArgs, TSharedRef<FDetailTreeNode> InOwnerTreeNode, const TSharedRef<STableViewBase>& InOwnerTableView)
 {
@@ -27,9 +30,12 @@ void SDetailCategoryTableRow::Construct(const FArguments& InArgs, TSharedRef<FDe
 
 	IDetailsViewPrivate* DetailsView = InOwnerTreeNode->GetDetailsView();
 	FDetailColumnSizeData& ColumnSizeData = DetailsView->GetColumnSizeData();
-	
-	PulseAnimation.AddCurve(0.0f, UE::PropertyEditor::Private::PulseAnimationLength, ECurveEaseFunction::CubicInOut);
+	ObjectName = InArgs._ObjectName;
+	DisplayManager = DetailsView->GetDisplayManager();
+
+	InitializeDisplayManager();
 
+	PulseAnimation.AddCurve(0.0f, UE::PropertyEditor::Private::PulseAnimationLength, ECurveEaseFunction::CubicInOut);
 	CopyAction.ExecuteAction = FExecuteAction::CreateSP(this, &SDetailCategoryTableRow::OnCopyCategory);
 	CopyAction.CanExecuteAction = FCanExecuteAction::CreateSP(this, &SDetailCategoryTableRow::CanCopyCategory);
 
@@ -83,26 +89,31 @@ void SDetailCategoryTableRow::Construct(const FArguments& InArgs, TSharedRef<FDe
 	auto GetScrollbarWellTint = [this]()
 	{
 		return SDetailTableRowBase::IsScrollBarVisible(OwnerTableViewWeak) ?
-			FSlateColor(EStyleColor::White) :
+			FSlateColor(EStyleColor::Header) :
 			this->GetInnerBackgroundColor();
 	};
 
-	FDetailsViewStyle ViewStyle = DetailsView->GetStyleKey();
-	ViewStyle.SetIsOuterCategory(!bIsInnerCategory);
-
 	this->ChildSlot
 	[
 		SNew(SBorder)
 		.BorderImage(FAppStyle::Get().GetBrush("DetailsView.GridLine"))
-		.Padding(ViewStyle.GetRowPadding())
+		.Padding_Lambda([this]
+		{
+			if (!DisplayManager.IsValid())
+			{
+				return FMargin{0};
+			}
+			DisplayManager->SetIsOuterCategory(!bIsInnerCategory);
+			return DisplayManager->GetRowPadding();
+		})
 		[
-			SNew(SHorizontalBox)
-			+ SHorizontalBox::Slot()
+			SNew(SOverlay)
+			+ SOverlay::Slot()
 			[
 				SNew(SBorder)
 				.BorderImage(this, &SDetailCategoryTableRow::GetBackgroundImage)
 				.BorderBackgroundColor(this, &SDetailCategoryTableRow::GetInnerBackgroundColor)
-				.Padding(0)
+					.Padding(0)
 				[
 					SNew(SBox)
 					.MinDesiredHeight(PropertyEditorConstants::PropertyRowHeight)
@@ -111,15 +122,37 @@ void SDetailCategoryTableRow::Construct(const FArguments& InArgs, TSharedRef<FDe
 					]
 				]
 			]
-			+ SHorizontalBox::Slot()
+			 + SOverlay::Slot()
 			.HAlign(HAlign_Right)
-			.AutoWidth()
 			[
-				SNew(SBorder)
-				.BorderImage(this, &SDetailCategoryTableRow::GetBackgroundImageForScrollBarWell)
-				.BorderBackgroundColor_Lambda(GetScrollbarWellTint)
-				.Padding(FMargin(0, 0, SDetailTableRowBase::ScrollBarPadding, 0))
+			SNew(SBorder)
+			 .BorderImage(this, &SDetailCategoryTableRow::GetBackgroundImageForScrollBarWell)
+			.Visibility(EVisibility::Visible)
+			.Padding_Lambda([this]
+			{
+				if (!DisplayManager.IsValid())
+				{
+					return FMargin{0};
+				}
+				DisplayManager->SetIsOuterCategory(!bIsInnerCategory);
+				DisplayManager->SetIsScrollbarShowing(IsScrollBarVisible(OwnerTableViewWeak));
+				return DisplayManager->GetCategoryButtonsPadding();
+			})
+			[
+			DisplayManager.IsValid() ?
+				*FCategoryMenuComboButtonBuilder( DisplayManager.ToSharedRef() )
+					.Set_OnGetContent(FOnGetContent::CreateLambda([this]
+					{
+						const TSharedPtr<SWidget> Menu = DisplayManager->GetCategoryMenu(ObjectName);
+						return Menu.IsValid() ? Menu.ToSharedRef() : SNullWidget::NullWidget;
+					}))
+					.Bind_IsVisible(TAttribute<EVisibility>::CreateLambda([this]
+					{
+						return IsHovered() ? EVisibility::Visible : EVisibility::Collapsed;
+					})) :
+					SNullWidget::NullWidget
 			]
+			] 
 		]
 	];
 
@@ -225,6 +258,8 @@ const FSlateBrush* SDetailCategoryTableRow::GetBackgroundImage() const
 		FDetailsViewStyle ViewStyle = View ? View->GetStyleKey() : PrimaryKey;
 		const bool bIsCategoryExpanded = IsItemExpanded();
 		const bool bIsScrollBarVisible = IsScrollBarVisible(OwnerTableViewWeak);
+		DisplayManager->SetIsScrollbarShowing(bIsScrollBarVisible);
+		ViewStyle.SetIsOuterCategory(!bIsInnerCategory);
 		
 		return ViewStyle.GetBackgroundImageForCategoryRow(bShowBorder, bIsInnerCategory, bIsCategoryExpanded, bIsScrollBarVisible);
 	}
@@ -240,13 +275,15 @@ const FSlateBrush* SDetailCategoryTableRow::GetBackgroundImageForScrollBarWell()
 		const FDetailsViewStyle ViewStyle = View ? View->GetStyleKey() : PrimaryKey;
 		const bool bIsCategoryExpanded = IsItemExpanded();
 		const bool bIsScrollBarVisible = IsScrollBarVisible(OwnerTableViewWeak);
+		DisplayManager->SetIsScrollbarShowing(bIsScrollBarVisible);
+		DisplayManager->SetIsOuterCategory(!bIsInnerCategory);
 
 		return ViewStyle.GetBackgroundImageForScrollBarWell(bShowBorder, bIsInnerCategory, bIsCategoryExpanded, bIsScrollBarVisible);
 	}
 	return nullptr;
 }
 
-FSlateColor SDetailCategoryTableRow::GetInnerBackgroundColor() const
+FSlateColor SDetailCategoryTableRow::	GetInnerBackgroundColor() const
 {
 	FSlateColor Color = FSlateColor(FLinearColor::White);
 	
@@ -475,3 +512,12 @@ FReply SDetailCategoryTableRow::OnMouseButtonDoubleClick(const FGeometry& InMyGe
 {
 	return OnMouseButtonDown(InMyGeometry, InMouseEvent);
 }
+
+void SDetailCategoryTableRow::InitializeDisplayManager()
+{
+	if (DisplayManager.IsValid())
+	{
+		DisplayManager->SetCategoryObjectName( ObjectName );
+		DisplayManager->SetIsOuterCategory( !bIsInnerCategory );
+	}
+}
diff --git a/Engine/Source/Editor/PropertyEditor/Private/SDetailCategoryTableRow.h b/Engine/Source/Editor/PropertyEditor/Private/SDetailCategoryTableRow.h
--- a/Engine/Source/Editor/PropertyEditor/Private/SDetailCategoryTableRow.h
+++ b/Engine/Source/Editor/PropertyEditor/Private/SDetailCategoryTableRow.h
@@ -18,12 +18,14 @@ public:
 	SLATE_BEGIN_ARGS(SDetailCategoryTableRow)
 		: _InnerCategory(false)
 		, _ShowBorder(true)
+		, _ObjectName(NAME_Name)
 	{}
 		SLATE_ARGUMENT(FText, DisplayName)
 		SLATE_ARGUMENT(bool, InnerCategory)
 		SLATE_ARGUMENT(TSharedPtr<SWidget>, HeaderContent)
 		SLATE_ARGUMENT(bool, ShowBorder)
 		SLATE_ARGUMENT(TSharedPtr<FOnPasteFromText>, PasteFromText)
+		SLATE_ARGUMENT(FName, ObjectName)
 	SLATE_END_ARGS()
 
 	void Construct(const FArguments& InArgs, TSharedRef<FDetailTreeNode> InOwnerTreeNode, const TSharedRef<STableViewBase>& InOwnerTableView);
@@ -58,10 +60,16 @@ private:
 	virtual FReply OnMouseButtonDown(const FGeometry& MyGeometry, const FPointerEvent& MouseEvent) override;
 	virtual FReply OnMouseButtonDoubleClick(const FGeometry& InMyGeometry, const FPointerEvent& InMouseEvent) override;
 
+	/** Initializes the @code TSharedPtr<FDetailsDisplayManager> DisplayManager @endcode for this category */
+	void InitializeDisplayManager();
+
 private:
 	/** Cached category name. */
 	FText DisplayName;
 
+	/** The name of the object that is specified by this category */
+	FName ObjectName;
+	
 	/** Previously parsed clipboard data. */
 	UE::PropertyEditor::Internal::FClipboardData PreviousClipboardData;
 
@@ -75,4 +83,4 @@ private:
 
 	/** Animation curve for displaying pulse */
 	FCurveSequence PulseAnimation;
-};
+};
\ No newline at end of file
diff --git a/Engine/Source/Editor/PropertyEditor/Private/SDetailSingleItemRow.cpp b/Engine/Source/Editor/PropertyEditor/Private/SDetailSingleItemRow.cpp
--- a/Engine/Source/Editor/PropertyEditor/Private/SDetailSingleItemRow.cpp
+++ b/Engine/Source/Editor/PropertyEditor/Private/SDetailSingleItemRow.cpp
@@ -674,13 +674,12 @@ void SDetailSingleItemRow::Construct( const FArguments& InArgs, FDetailLayoutCus
 	};
 
 	static const FDetailsViewStyleKey& PrimaryKey = SDetailsView::GetPrimaryDetailsViewStyleKey();
-	FDetailsViewStyle ViewStyle = DetailsView ? DetailsView->GetStyleKey() : PrimaryKey;
 	
 	this->ChildSlot
 	[
 		SNew( SBorder )
 		.BorderImage(FAppStyle::Get().GetBrush("DetailsView.GridLine"))
-		.Padding(ViewStyle.GetRowPadding())
+		.Padding(FMargin(0,0,0,1))
 		.Clipping(EWidgetClipping::ClipToBounds)
 		[
 			SNew(SBox)
@@ -704,16 +703,6 @@ void SDetailSingleItemRow::Construct( const FArguments& InArgs, FDetailLayoutCus
 						]
 					]
 				]
-				+ SHorizontalBox::Slot()
-				.HAlign(HAlign_Right)
-				.VAlign(VAlign_Fill)
-				.AutoWidth()
-				[
-					SNew( SBorder )
-					.BorderImage_Lambda(GetScrollbarWellBrush)
-					.BorderBackgroundColor_Lambda(GetScrollbarWellTint)
-					.Padding(FMargin(0, 0, SDetailTableRowBase::ScrollBarPadding, 0))
-				]
 			]
 		]
 	];
diff --git a/Engine/Source/Editor/PropertyEditor/Private/SDetailTableRowBase.h b/Engine/Source/Editor/PropertyEditor/Private/SDetailTableRowBase.h
--- a/Engine/Source/Editor/PropertyEditor/Private/SDetailTableRowBase.h
+++ b/Engine/Source/Editor/PropertyEditor/Private/SDetailTableRowBase.h
@@ -108,4 +108,10 @@ protected:
 	 */
 	TWeakPtr<STableViewBase> OwnerTableViewWeak;
 	TWeakPtr<FDetailTreeNode> OwnerTreeNode;
+	
+	/**
+	 * The @code DetailsDisplayManager @endcode which provides an API to manage some of the characteristics of the
+	 * details display
+	 */
+	TSharedPtr<FDetailsDisplayManager> DisplayManager;
 };
diff --git a/Engine/Source/Editor/PropertyEditor/Private/SDetailsView.cpp b/Engine/Source/Editor/PropertyEditor/Private/SDetailsView.cpp
--- a/Engine/Source/Editor/PropertyEditor/Private/SDetailsView.cpp
+++ b/Engine/Source/Editor/PropertyEditor/Private/SDetailsView.cpp
@@ -412,7 +412,20 @@ void SDetailsView::Construct(const FArguments& InArgs, const FDetailsViewArgs& I
 			]
 			+SHorizontalBox::Slot()
 			[
-				ConstructTreeView(ExternalScrollbar)
+	           SNew(SBorder)
+				.Padding_Lambda([this]
+				{
+					return DisplayManager->GetTablePadding();
+				})
+				.BorderImage_Lambda([this]
+				{
+					static const FSlateBrush* Background = FAppStyle::GetBrush("DetailsView.GridLine");
+					static const FSlateBrush* Panel = FAppStyle::GetBrush("Brushes.Panel");
+					
+					return RootPropertyNodes.Num() == 0 ? Panel : Background;
+				})
+				.Visibility(this, &SDetailsView::GetTreeVisibility)
+				[ ConstructTreeView(ExternalScrollbar) ]
 			]
 		]
 		+ SOverlay::Slot()
@@ -602,6 +615,12 @@ void SDetailsView::SetObjectFilter(TSharedPtr<FDetailsViewObjectFilter> InFilter
 	{
 		ObjectFilter = MakeShared<FDetailsViewDefaultObjectFilter>(!!DetailsViewArgs.bAllowMultipleTopLevelObjects);
 	}
+
+	RefreshDisplayManager();
+
+	// hook up the OnDetailsNeedsUpdate to the details panel refresh function so that when the display data
+	// is updated, the details panel will reflect it
+	DisplayManager->OnDetailsNeedsUpdate.BindSP(this, &SDetailsView::ForceRefresh);
 }
 
 void SDetailsView::SetClassViewerFilters(const TArray<TSharedRef<class IClassViewerFilter>>& InFilters)
@@ -1485,7 +1504,7 @@ bool SDetailsView::IsDefaultStyle() const
 
 void SDetailsView::UpdateStyleKey()
 {
-	const bool IsDefault = FDetailsViewStyleKeys::IsDefault(ObjectFilter->GetObjectsDetailsViewStyleKey());
+	const bool IsDefault = FDetailsViewStyleKeys::IsDefault(DisplayManager->GetDetailsViewStyleKey());
 
 	if (IsDefault && DetailsViewArgs.StyleKey.IsValid()) 
 	{
@@ -1493,7 +1512,7 @@ void SDetailsView::UpdateStyleKey()
 	}
 	else if (!IsDefault)
 	{
-		FDetailsViewStyleKey Key = ObjectFilter->GetObjectsDetailsViewStyleKey(); 
+		FDetailsViewStyleKey Key = DisplayManager->GetDetailsViewStyleKey(); 
 		StyleKeySP = MakeShared<FDetailsViewStyleKey>(Key);
 	}
 	else
@@ -1509,6 +1528,24 @@ const FDetailsViewStyleKey& SDetailsView::GetStyleKey()
 	return StyleKeySP.IsValid() ? *StyleKeySP.Get() : PrimaryKey;
 }
 
+void SDetailsView::RefreshDisplayManager()
+{
+	if (ObjectFilter.IsValid())
+	{
+		DisplayManager = ObjectFilter->GetDisplayManager();
+	}
+	if (!DisplayManager.IsValid())
+	{
+		DisplayManager = MakeShared<FDetailsDisplayManager>();
+	}
+}
+
+TSharedPtr<FDetailsDisplayManager> SDetailsView::GetDisplayManager()
+{
+	RefreshDisplayManager();
+	return DisplayManager;
+}
+
 const FDetailsViewStyleKey& SDetailsView::GetPrimaryDetailsViewStyleKey()
 {
 	if (CVarDetailsPanelEnableCardLayout.GetValueOnAnyThread())
diff --git a/Engine/Source/Editor/PropertyEditor/Private/SDetailsView.h b/Engine/Source/Editor/PropertyEditor/Private/SDetailsView.h
--- a/Engine/Source/Editor/PropertyEditor/Private/SDetailsView.h
+++ b/Engine/Source/Editor/PropertyEditor/Private/SDetailsView.h
@@ -130,8 +130,8 @@ public:
 	/**
 	* Returns a bool indicating whether the given @code FDetailsViewStyleKey @endcode is the default Details View Style 
 	*/
-	bool IsDefaultStyle() const;
-	
+	virtual bool IsDefaultStyle() const override;
+
 private:
 
 	void SetObjectArrayPrivate(const TArray<UObject*>& InObjects);
@@ -221,7 +221,19 @@ private:
 	*/
 	virtual const FDetailsViewStyleKey& GetStyleKey() override;
 
+	/**
+	 * Returns a @code TSharedPtr @endcode to the @code FDetailsDisplayManager @endcode
+	 */
+	virtual TSharedPtr<FDetailsDisplayManager> GetDisplayManager() override;
+
 private:
+
+	/**
+	 * Refreshes the class member @code TSharedPtr<FDetailsDisplayManager> DisplayManager @endcode. This is required
+	 * anytime the data in the display manager may have grown stale.
+	 */
+	void RefreshDisplayManager();
+
 	/**
 	* The @code FDetailsViewStyle& @endcode shared pointer which provides style
 	* information for this details view.
diff --git a/Engine/Source/Editor/PropertyEditor/Private/SDetailsViewBase.cpp b/Engine/Source/Editor/PropertyEditor/Private/SDetailsViewBase.cpp
--- a/Engine/Source/Editor/PropertyEditor/Private/SDetailsViewBase.cpp
+++ b/Engine/Source/Editor/PropertyEditor/Private/SDetailsViewBase.cpp
@@ -1780,3 +1780,12 @@ void SDetailsViewBase::UnregisterInstancedCustomPropertyTypeLayout(FName Propert
 		LayoutCallbacks->Remove(Identifier);
 	}
 }
+
+TSharedPtr<FDetailsDisplayManager> SDetailsViewBase::GetDisplayManager()
+{
+	if (!DisplayManager.IsValid())
+	{
+		DisplayManager = MakeShared<FDetailsDisplayManager>();
+	}
+	return DisplayManager;
+}
\ No newline at end of file
diff --git a/Engine/Source/Editor/PropertyEditor/Private/SDetailsViewBase.h b/Engine/Source/Editor/PropertyEditor/Private/SDetailsViewBase.h
--- a/Engine/Source/Editor/PropertyEditor/Private/SDetailsViewBase.h
+++ b/Engine/Source/Editor/PropertyEditor/Private/SDetailsViewBase.h
@@ -181,6 +181,12 @@ public:
 	/** Restore all expanded items in root nodes and external root nodes. */
 	void RestoreAllExpandedItems();
 
+	/**
+	 * Returns a @code TSharedPtr @endcode to the @code FDetailsDisplayManager @endcode for this
+	 * details view
+	 */
+	virtual TSharedPtr<FDetailsDisplayManager> GetDisplayManager() override;
+
 	// SWidget interface
 	virtual bool SupportsKeyboardFocus() const override;
 	virtual FReply OnFocusReceived(const FGeometry& MyGeometry, const FFocusEvent& InFocusEvent) override;
@@ -472,6 +478,12 @@ protected:
 	/** The list of nodes whose widgets should be animating. */
 	TArray<FAnimatingNodeCollection> CurrentlyAnimatingNodeCollections;
 
+	/**
+	 * The @code DetailsDisplayManager @endcode which provides an API to manage some of the characteristics of the
+	 * details display
+	 */
+	TSharedPtr<FDetailsDisplayManager> DisplayManager;
+
 	/** Current set of expanded detail nodes (by path) that should be saved when the details panel closes */
 	FStringPrefixTree ExpandedDetailNodes;
 	FStringPrefixTree PreSearchExpandedItems;
diff --git a/Engine/Source/Editor/PropertyEditor/Private/UserInterface/Categories/CategoryMenuComboButtonBuilder.cpp b/Engine/Source/Editor/PropertyEditor/Private/UserInterface/Categories/CategoryMenuComboButtonBuilder.cpp
new file mode 100644
--- /dev/null
+++ b/Engine/Source/Editor/PropertyEditor/Private/UserInterface/Categories/CategoryMenuComboButtonBuilder.cpp
@@ -0,0 +1,44 @@
+//  Copyright Epic Games, Inc. All Rights Reserved.
+
+#include "UserInterface/Categories/CategoryMenuComboButtonBuilder.h"
+
+FCategoryMenuComboButtonBuilder::FCategoryMenuComboButtonBuilder(
+	TSharedRef<FDetailsDisplayManager> InDetailsDisplayManager):
+	FToolElementRegistrationArgs("FCategoryMenuComboButtonBuilder"),
+	DisplayManager(InDetailsDisplayManager)
+{
+}
+
+FCategoryMenuComboButtonBuilder& FCategoryMenuComboButtonBuilder::Set_OnGetContent(FOnGetContent InOnGetContent)
+{
+	OnGetContent.Unbind();
+	OnGetContent = InOnGetContent;
+	return *this;
+}
+
+FCategoryMenuComboButtonBuilder& FCategoryMenuComboButtonBuilder::Bind_IsVisible(TAttribute<EVisibility> InIsVisible)
+{
+	IsVisible = InIsVisible;
+	return *this;
+}
+
+TSharedPtr<SWidget> FCategoryMenuComboButtonBuilder::GenerateWidget()
+{
+	return 	DisplayManager->ShouldShowCategoryMenu()  ?
+		        SNew( SComboButton )
+				.ComboButtonStyle( FAppStyle::Get(), "DetailsView.CategoryComboButton" )
+				.Visibility(IsVisible)
+				.HasDownArrow(true)
+				.OnGetMenuContent(OnGetContent) :
+		        SNullWidget::NullWidget;
+}
+
+TSharedRef<SWidget> FCategoryMenuComboButtonBuilder::operator*()
+{
+	return GenerateWidget().ToSharedRef();
+}
+
+FCategoryMenuComboButtonBuilder::~FCategoryMenuComboButtonBuilder()
+{
+	OnGetContent.Unbind();
+}
\ No newline at end of file
diff --git a/Engine/Source/Editor/PropertyEditor/PropertyEditor.Build.cs b/Engine/Source/Editor/PropertyEditor/PropertyEditor.Build.cs
--- a/Engine/Source/Editor/PropertyEditor/PropertyEditor.Build.cs
+++ b/Engine/Source/Editor/PropertyEditor/PropertyEditor.Build.cs
@@ -45,6 +45,7 @@ public class PropertyEditor : ModuleRules
 				"DesktopPlatform",
 				"PropertyPath",
 				"ToolWidgets",
+				"WidgetRegistration",
 				"Json"
 			}
         );
diff --git a/Engine/Source/Editor/PropertyEditor/Public/DetailsDisplayManager.h b/Engine/Source/Editor/PropertyEditor/Public/DetailsDisplayManager.h
new file mode 100644
--- /dev/null
+++ b/Engine/Source/Editor/PropertyEditor/Public/DetailsDisplayManager.h
@@ -0,0 +1,112 @@
+//  Copyright Epic Games, Inc. All Rights Reserved.
+#pragma once
+
+#include "DetailsViewStyleKey.h"
+#include "Templates/SharedPointer.h"
+#include "Widgets/SWidget.h"
+
+DECLARE_DELEGATE(FOnDetailsNeedsUpdate)
+
+/** An @code FDetailsDisplayManager @endcode provides an API to tweak various settings of your details view, and
+ * provides some utility methods to work with Details.  */
+class FDetailsDisplayManager : public TSharedFromThis<FDetailsDisplayManager>
+{
+public:
+	FOnDetailsNeedsUpdate OnDetailsNeedsUpdate;
+
+	PROPERTYEDITOR_API FDetailsDisplayManager() : bIsOuterCategory(false)
+	{
+	}
+
+	PROPERTYEDITOR_API virtual ~FDetailsDisplayManager();
+
+	/**
+	 * Returns a boolean indicating if the Component Editor should be hidden 
+	 */
+	virtual bool ShouldHideComponentEditor();
+
+	/**
+	 * Returns a @code bool @endcode indicating whether this @code DetailsViewObjectFilter @endcode instance
+	 * should show a category menu
+	 */
+	virtual bool ShouldShowCategoryMenu();
+
+	/**
+	 * Sets the name of the object defined by the currently active category
+	 *
+	 * @param InCategoryObjectName the name of the category object
+	 */
+	void SetCategoryObjectName(FName InCategoryObjectName);
+
+	/**
+	 * Gets the category menu SWidget and returns a shared pointer to it
+	 *
+	 *  @param InCategoryObjectName the name of the category
+	 */
+	virtual TSharedPtr<SWidget> GetCategoryMenu(FName InCategoryObjectName);
+
+	/**
+	 * Updates the current details view
+	 */
+	PROPERTYEDITOR_API void UpdateView() const;
+
+	/**
+	 * Returns the @code FDetailsViewStyleKey @endcode that is the Key to the current FDetailsViewStyle style
+	 */
+	virtual const FDetailsViewStyleKey& GetDetailsViewStyleKey() const;
+
+	/** sets whether the currently active category is an Outer category*/
+	void SetIsOuterCategory(bool bInIsOuterCategory);
+
+	/**
+	 * Returns the padding for details panel rows which are not outer Category rows
+	 */
+	FMargin GetRowPadding() const;
+
+	/**
+	 * Returns the padding for any buttons for Category rows
+	 */
+	FMargin GetCategoryButtonsPadding() const;
+
+	/**
+	 * returns a boolean indicating whether or not the currently active categpory is an inner category
+	 */	
+	bool GetIsInnerCategory() const;
+
+	/**
+   	 * Returns a bool indicating whether or not the scrollbar is showing on the details view
+   	 */
+	PROPERTYEDITOR_API virtual bool GetIsScrollbarShowing() const;
+	
+	/**
+	* Set a bool indicating whether or not the scrollbar is showing on the details view
+	*
+	* @param bInIsScrollbarShowing a bool indicating whether or not the scrollbar is showing on the details view
+	*/
+	PROPERTYEDITOR_API virtual void SetIsScrollbarShowing(bool bInIsScrollbarShowing);
+	
+	/**
+	* Returns an FMargin containing the padding for the entire details view. Note that if there is a scrollbar present,
+	* this padding takes that into account
+	*/
+	PROPERTYEDITOR_API virtual FMargin GetTablePadding() const;
+
+protected:
+
+	/**
+	 * The name of the object defined by the currently active category
+	 */
+	FName CategoryObjectName;
+
+	/**
+	 * A bool indicating whether or not the currently active category is an outer category
+	 */
+	bool bIsOuterCategory;
+
+	/**
+	 * A bool indicating whether or not the scrollbar is showing on the details view
+	 */
+	bool bIsScrollbarShowing;
+
+};
+
diff --git a/Engine/Source/Editor/PropertyEditor/Public/DetailsViewObjectFilter.h b/Engine/Source/Editor/PropertyEditor/Public/DetailsViewObjectFilter.h
--- a/Engine/Source/Editor/PropertyEditor/Public/DetailsViewObjectFilter.h
+++ b/Engine/Source/Editor/PropertyEditor/Public/DetailsViewObjectFilter.h
@@ -4,6 +4,11 @@
 
 #include "CoreMinimal.h"
 #include "DetailsViewStyleKey.h"
+#include "Widgets/SWidget.h"
+#include "Templates/SharedPointer.h"
+#include "DetailsDisplayManager.h"
+
+DECLARE_DELEGATE(FOnUpdateFilteredObjects)
 
 /**
  * An object root is a collection of UObjects that represent a top level set of properties in a details panel
@@ -42,7 +47,14 @@ class FDetailsViewObjectFilter
 {
 public:
 
-	virtual ~FDetailsViewObjectFilter() {}
+	explicit FDetailsViewObjectFilter()
+	{
+		InitializeDetailsDisplayManager();
+	}
+
+	virtual ~FDetailsViewObjectFilter()
+	{
+	}
 
 	/**
 	 * Given a const TArray<UObject*>& SourceObjects, it fills a TArray<FDetailsViewObjectRoot> with the objects
@@ -54,20 +66,38 @@ public:
 	 */
 	virtual TArray<FDetailsViewObjectRoot> FilterObjects(const TArray<UObject*>& SourceObjects) = 0;
 
+	/** Updates the view of anything being filtered by this filter */
+	virtual void UpdateFilterView() const
+	{
+	}
+
 	/**
-	 * Returns the @code FDetailsViewStyleKey& @endcode that is the Key to the current objects' style
-	 */
-	virtual const FDetailsViewStyleKey& GetObjectsDetailsViewStyleKey()
+	* Returns a @code TSharedPtr @endcode to the @code FDetailsDisplayManager @endcode
+	*/
+	virtual TSharedPtr<FDetailsDisplayManager> GetDisplayManager()
 	{
-		static const FDetailsViewStyleKey Default = FDetailsViewStyleKeys::Default();
-		return Default;
+		if (!DisplayManager.IsValid())
+		{
+			InitializeDetailsDisplayManager();
+		}
+		return DisplayManager;
 	}
 
+
+protected:
+	/**
+	 * The @code DetailsDisplayManager @endcode which provides an API to manage some of the characteristics of the
+	 * details display
+	 */
+	TSharedPtr<FDetailsDisplayManager> DisplayManager;
+
 	/**
-	 * Returns a boolean indicating if the Component Editor should be hidden for these object(s)
+	 * Initializes the Details Display Manager
 	 */
-	virtual bool ShouldHideComponentEditor()
+	virtual void InitializeDetailsDisplayManager()
 	{
-		return false;
+		DisplayManager = MakeShared<FDetailsDisplayManager>();
 	}
+	
 };
+ 
\ No newline at end of file
diff --git a/Engine/Source/Editor/PropertyEditor/Public/IDetailsView.h b/Engine/Source/Editor/PropertyEditor/Public/IDetailsView.h
--- a/Engine/Source/Editor/PropertyEditor/Public/IDetailsView.h
+++ b/Engine/Source/Editor/PropertyEditor/Public/IDetailsView.h
@@ -3,6 +3,7 @@
 #pragma once
 
 #include "CoreMinimal.h"
+#include "DetailsDisplayManager.h"
 #include "Widgets/SCompoundWidget.h"
 #include "PropertyEditorDelegates.h"
 #include "Framework/Commands/UICommandList.h"
@@ -311,4 +312,5 @@ public:
    	* Returns a bool indicating whether the given @code FDetailsViewStyleKey @endcode is the default Details View Style 
    	*/
    	virtual bool IsDefaultStyle() const = 0;
+
 };
diff --git a/Engine/Source/Editor/PropertyEditor/Public/UserInterface/Categories/CategoryMenuComboButtonBuilder.h b/Engine/Source/Editor/PropertyEditor/Public/UserInterface/Categories/CategoryMenuComboButtonBuilder.h
new file mode 100644
--- /dev/null
+++ b/Engine/Source/Editor/PropertyEditor/Public/UserInterface/Categories/CategoryMenuComboButtonBuilder.h
@@ -0,0 +1,68 @@
+//  Copyright Epic Games, Inc. All Rights Reserved.
+
+#pragma once
+#include "DetailsDisplayManager.h"
+#include "ToolElementRegistry.h"
+#include "Widgets/SNullWidget.h"
+#include "Widgets/Input/SComboButton.h"
+#include "Templates/SharedPointer.h"
+
+class FDetailsDisplayManager;
+
+DECLARE_DELEGATE_OneParam(FGetIsRowHoveredOver, bool)
+
+class FCategoryMenuComboButtonBuilder : public FToolElementRegistrationArgs
+{
+public:
+
+	/**
+	 * the delegate to get the menu content
+	 */
+	FOnGetContent OnGetContent;
+
+	/**
+	 * The constructor, which takes a @code TSharedRef<FDetailsDisplayManager> @endcode to initialize
+	 * the Details Display Manager
+	 */
+	FCategoryMenuComboButtonBuilder( TSharedRef<FDetailsDisplayManager> InDetailsDisplayManager );
+
+	/**
+	 * Set the OnGetContent for the menu that this button is responsible for
+	 */
+	FCategoryMenuComboButtonBuilder& Set_OnGetContent(FOnGetContent InOnGetContent);
+
+	/**
+	 * Set the IsVisible delegate for the menu that this button is responsible for
+	 */
+	FCategoryMenuComboButtonBuilder& Bind_IsVisible(TAttribute<EVisibility> InIsVisible);
+
+	/** Implements the generation of the Category Menu button SWidget */
+	virtual TSharedPtr<SWidget> GenerateWidget() override;
+
+	/**
+	 * Converts this into the SWidget it builds
+	 */
+	virtual ~FCategoryMenuComboButtonBuilder() override;
+
+	/**
+	 * Converts this into the SWidget it builds
+	 */
+    TSharedRef<SWidget> operator*();
+
+private:
+	/**
+	 * The @code DetailsDisplayManager @endcode which provides an API to manage some of the characteristics of the
+	 * details display
+	 */
+	TSharedRef<FDetailsDisplayManager> DisplayManager;
+
+	/**
+	 * An attribute indicating whether the current row is hovered over
+	 */
+	TAttribute<EVisibility> IsRowHoveredOver;
+
+	/**
+	 * the attribute which tells whether this menu button is visible
+	 */
+	TAttribute<EVisibility> IsVisible;
+};
