Allow math expressions on properties marked "Multiple Values"

diff --git a/Engine/Source/Editor/PropertyEditor/Private/UserInterface/PropertyEditor/PropertyEditorConstants.cpp b/Engine/Source/Editor/PropertyEditor/Private/UserInterface/PropertyEditor/PropertyEditorConstants.cpp
--- a/Engine/Source/Editor/PropertyEditor/Private/UserInterface/PropertyEditor/PropertyEditorConstants.cpp
+++ b/Engine/Source/Editor/PropertyEditor/Private/UserInterface/PropertyEditor/PropertyEditorConstants.cpp
@@ -6,6 +6,7 @@
 #include "Styling/AppStyle.h"
 #include "Styling/ISlateStyle.h"
 #include "UObject/NameTypes.h"
+#include "Internationalization/Text.h"
 
 const FName PropertyEditorConstants::PropertyFontStyle( TEXT("PropertyWindow.NormalFont") );
 const FName PropertyEditorConstants::CategoryFontStyle( TEXT("DetailsView.CategoryFontStyle") );
@@ -14,6 +15,8 @@ const FName PropertyEditorConstants::MD_Bitmask( TEXT("Bitmask") );
 const FName PropertyEditorConstants::MD_BitmaskEnum( TEXT("BitmaskEnum") );
 const FName PropertyEditorConstants::MD_UseEnumValuesAsMaskValuesInEditor( TEXT("UseEnumValuesAsMaskValuesInEditor") );
 
+const FText PropertyEditorConstants::DefaultUndeterminedText(NSLOCTEXT("PropertyEditor", "MultipleValues", "Multiple Values"));
+
 const FSlateBrush* PropertyEditorConstants::GetOverlayBrush(const TSharedRef<class FPropertyEditor> PropertyEditor )
 {
 	return FAppStyle::GetBrush( TEXT("PropertyWindow.NoOverlayColor") );
diff --git a/Engine/Source/Editor/PropertyEditor/Private/UserInterface/PropertyEditor/PropertyEditorConstants.h b/Engine/Source/Editor/PropertyEditor/Private/UserInterface/PropertyEditor/PropertyEditorConstants.h
--- a/Engine/Source/Editor/PropertyEditor/Private/UserInterface/PropertyEditor/PropertyEditorConstants.h
+++ b/Engine/Source/Editor/PropertyEditor/Private/UserInterface/PropertyEditor/PropertyEditorConstants.h
@@ -6,6 +6,7 @@
 #include "Templates/SharedPointer.h"
 
 class FName;
+class FText;
 struct FSlateBrush;
 
 namespace UE::PropertyEditor::Private
@@ -26,6 +27,8 @@ public:
 
 	static constexpr float PropertyRowHeight = 26.0f;
 
+	static const FText DefaultUndeterminedText;
+
 	static const FSlateBrush* GetOverlayBrush(const TSharedRef<class FPropertyEditor> PropertyEditor);
 
 	static FSlateColor GetRowBackgroundColor(int32 IndentLevel, bool IsHovered);
diff --git a/Engine/Source/Editor/PropertyEditor/Private/UserInterface/PropertyEditor/SPropertyEditorArray.cpp b/Engine/Source/Editor/PropertyEditor/Private/UserInterface/PropertyEditor/SPropertyEditorArray.cpp
--- a/Engine/Source/Editor/PropertyEditor/Private/UserInterface/PropertyEditor/SPropertyEditorArray.cpp
+++ b/Engine/Source/Editor/PropertyEditor/Private/UserInterface/PropertyEditor/SPropertyEditorArray.cpp
@@ -6,6 +6,7 @@
 #include "PropertyEditorHelpers.h"
 #include "SDropTarget.h"
 #include "Widgets/Text/STextBlock.h"
+#include "PropertyEditorConstants.h"
 
 #define LOCTEXT_NAMESPACE "PropertyEditor"
 
@@ -63,7 +64,7 @@ FText SPropertyEditorArray::GetArrayTextValue() const
 
 	if (GetValResult == FPropertyAccess::MultipleValues)
 	{
-		return LOCTEXT("MultipleValues", "Multiple Values");
+		return PropertyEditorConstants::DefaultUndeterminedText;
 	}
 
 	const int32 NumChildNodes = PropertyEditor->GetPropertyNode()->GetNumChildNodes();
diff --git a/Engine/Source/Editor/PropertyEditor/Private/UserInterface/PropertyEditor/SPropertyEditorAsset.cpp b/Engine/Source/Editor/PropertyEditor/Private/UserInterface/PropertyEditor/SPropertyEditorAsset.cpp
--- a/Engine/Source/Editor/PropertyEditor/Private/UserInterface/PropertyEditor/SPropertyEditorAsset.cpp
+++ b/Engine/Source/Editor/PropertyEditor/Private/UserInterface/PropertyEditor/SPropertyEditorAsset.cpp
@@ -36,6 +36,7 @@
 #include "AssetThumbnail.h"
 #include "DetailWidgetRow.h"
 #include "Subsystems/AssetEditorSubsystem.h"
+#include "PropertyEditorConstants.h"
 
 #define LOCTEXT_NAMESPACE "PropertyEditor"
 
@@ -823,7 +824,7 @@ FText SPropertyEditorAsset::OnGetAssetName() const
 	}
 	else if( Result == FPropertyAccess::MultipleValues )
 	{
-		Name = LOCTEXT("MultipleValues", "Multiple Values");
+		Name = PropertyEditorConstants::DefaultUndeterminedText;
 	}
 
 	return Name;
@@ -887,7 +888,7 @@ FText SPropertyEditorAsset::OnGetToolTip() const
 	}
 	else if( Result == FPropertyAccess::MultipleValues )
 	{
-		ToolTipText = LOCTEXT("MultipleValues", "Multiple Values");
+		ToolTipText = PropertyEditorConstants::DefaultUndeterminedText;
 	}
 
 	if( ToolTipText.IsEmpty() )
diff --git a/Engine/Source/Editor/PropertyEditor/Private/UserInterface/PropertyEditor/SPropertyEditorNumeric.h b/Engine/Source/Editor/PropertyEditor/Private/UserInterface/PropertyEditor/SPropertyEditorNumeric.h
--- a/Engine/Source/Editor/PropertyEditor/Private/UserInterface/PropertyEditor/SPropertyEditorNumeric.h
+++ b/Engine/Source/Editor/PropertyEditor/Private/UserInterface/PropertyEditor/SPropertyEditorNumeric.h
@@ -30,6 +30,7 @@
 #include "Widgets/Input/SNumericEntryBox.h"
 #include "Widgets/Input/SComboButton.h"
 #include "ObjectPropertyNode.h"
+#include "PropertyEditorConstants.h"
 
 #include "Math/UnitConversion.h"
 #include "Widgets/Input/NumericUnitTypeInterface.inl"
@@ -183,7 +184,7 @@ public:
 				}
 				else
 				{
-					return LOCTEXT("MultipleValues", "Multiple Values");
+					return PropertyEditorConstants::DefaultUndeterminedText;
 				}
 			};
 
@@ -345,7 +346,7 @@ public:
 				.LinearDeltaSensitivity(NumericPropertyParams.GetLinearDeltaSensitivityAttribute())
 				.AllowWheel(bAllowSpin)
 				.WheelStep(NumericPropertyParams.WheelStep)
-				.UndeterminedString(LOCTEXT("MultipleValues", "Multiple Values"))
+				.UndeterminedString(PropertyEditorConstants::DefaultUndeterminedText)
 				.OnValueChanged(this, &SPropertyEditorNumeric<NumericType>::OnValueChanged)
 				.OnValueCommitted(this, &SPropertyEditorNumeric<NumericType>::OnValueCommitted)
 				.OnUndeterminedValueCommitted(this, &SPropertyEditorNumeric<NumericType>::OnUndeterminedValueCommitted)
@@ -547,7 +548,8 @@ private:
 		{
 			NumericType OldNumericValue;
 			TTypeFromString<NumericType>::FromString(OldNumericValue, *Value);
-			TOptional<NumericType> NewNumericValue = TypeInterface->FromString(NewValueString, OldNumericValue);
+			const FString ReplacedNewValueString = NewValueString.Replace(*PropertyEditorConstants::DefaultUndeterminedText.ToString(), *Value);
+			TOptional<NumericType> NewNumericValue = TypeInterface->FromString(ReplacedNewValueString, OldNumericValue);
 
 			if (NewNumericValue.IsSet())
 			{
