Actor Picking: Experimental Editor Path Support
- Option that allows picking Actors inside a Level Instance
- Used by SPropertyEditorAsset when picking SoftObject refs to actors (if CVar "EditorPaths.Enabled" is true)

#rnx

diff --git a/Engine/Source/Editor/PropertyEditor/Private/PropertyCustomizationHelpers.cpp b/Engine/Source/Editor/PropertyEditor/Private/PropertyCustomizationHelpers.cpp
--- a/Engine/Source/Editor/PropertyEditor/Private/PropertyCustomizationHelpers.cpp
+++ b/Engine/Source/Editor/PropertyEditor/Private/PropertyCustomizationHelpers.cpp
@@ -317,6 +317,20 @@ namespace PropertyCustomizationHelpers
 			SNew( SPropertyMenuActorPicker )
 			.InitialActor(InitialActor)
 			.AllowClear(AllowClear)
+			.AllowPickingLevelInstanceContent(false)
+			.ActorFilter(ActorFilter)
+			.OnSet(OnSet)
+			.OnClose(OnClose)
+			.OnUseSelected(OnUseSelected);
+	}
+
+	TSharedRef<SWidget> MakeActorPickerWithMenu(AActor* const InitialActor, const bool AllowClear, const bool AllowPickingLevelInstanceContent, FOnShouldFilterActor ActorFilter, FOnActorSelected OnSet, FSimpleDelegate OnClose, FSimpleDelegate OnUseSelected)
+	{
+		return
+			SNew(SPropertyMenuActorPicker)
+			.InitialActor(InitialActor)
+			.AllowClear(AllowClear)
+			.AllowPickingLevelInstanceContent(AllowPickingLevelInstanceContent)
 			.ActorFilter(ActorFilter)
 			.OnSet(OnSet)
 			.OnClose(OnClose)
diff --git a/Engine/Source/Editor/PropertyEditor/Private/UserInterface/PropertyEditor/SPropertyEditorAsset.cpp b/Engine/Source/Editor/PropertyEditor/Private/UserInterface/PropertyEditor/SPropertyEditorAsset.cpp
--- a/Engine/Source/Editor/PropertyEditor/Private/UserInterface/PropertyEditor/SPropertyEditorAsset.cpp
+++ b/Engine/Source/Editor/PropertyEditor/Private/UserInterface/PropertyEditor/SPropertyEditorAsset.cpp
@@ -37,6 +37,7 @@
 #include "DetailWidgetRow.h"
 #include "Subsystems/AssetEditorSubsystem.h"
 #include "PropertyEditorConstants.h"
+#include "Misc/EditorPathHelper.h"
 
 #define LOCTEXT_NAMESPACE "PropertyEditor"
 
@@ -83,6 +84,66 @@ static bool GetTagOrBoolMetadata(const FProperty* Property, FName TagName, bool
 	return bResult;
 }
 
+static bool GetEditorPathOwnerFromPropertyHandle(const TSharedPtr<IPropertyHandle>& PropertyHandle, UObject*& OutEditorPathOwner)
+{
+	// If we don't get a proper Handle then consider the context null and valid
+	OutEditorPathOwner = nullptr;
+
+	if(PropertyHandle.IsValid())
+	{ 
+		TArray<UObject*> OuterObjects;
+		PropertyHandle->GetOuterObjects(OuterObjects);
+
+		if (OuterObjects.Num() > 0)
+		{
+			UObject* OutReferencer = OuterObjects[0];
+			OutEditorPathOwner = FEditorPathHelper::GetEditorPathOwner(OutReferencer);
+			for (int32 i = 1; i < OuterObjects.Num(); ++i)
+			{
+				if (OutEditorPathOwner != FEditorPathHelper::GetEditorPathOwner(OuterObjects[i]))
+				{
+					OutEditorPathOwner = nullptr;
+					return false;
+				}
+			}
+		}
+	}
+
+	return true;
+}
+
+static FString GetActorEditorPathLabel(const AActor* InActor)
+{
+	check(InActor);
+	
+	TArray<FString> EditorPathOwners;
+	EditorPathOwners.Add(InActor->GetActorLabel());
+	const UObject* Context = InActor;
+
+	while (UObject* EditorPathOwner = FEditorPathHelper::GetEditorPathOwner(Context))
+	{
+		if (AActor* ActorEditorPathOwner = Cast<AActor>(EditorPathOwner))
+		{
+			EditorPathOwners.Add(ActorEditorPathOwner->GetActorLabel());
+		}
+		else
+		{
+			EditorPathOwners.Add(EditorPathOwner->GetName());
+		}
+		Context = EditorPathOwner;
+	}
+
+	TStringBuilder<256> LabelBuilder;
+	LabelBuilder.Append(EditorPathOwners[EditorPathOwners.Num()-1]);
+	for (int32 i = EditorPathOwners.Num() - 2; i >= 0; --i)
+	{
+		LabelBuilder.Append(TEXT("."));
+		LabelBuilder.Append(EditorPathOwners[i]);
+	}
+
+	return LabelBuilder.ToString();
+}
+
 bool SPropertyEditorAsset::ShouldDisplayThumbnail(const FArguments& InArgs, const UClass* InObjectClass) const
 {
 	if (!InArgs._DisplayThumbnail || !InArgs._ThumbnailPool.IsValid())
@@ -281,7 +342,8 @@ void SPropertyEditorAsset::Construct(const FArguments& InArgs, const TSharedPtr<
 	ObjectClass = InArgs._Class != nullptr ? InArgs._Class : GetObjectPropertyClass(Property);
 	bAllowClear = InArgs._AllowClear.IsSet() ? InArgs._AllowClear.GetValue() : (Property ? !(Property->PropertyFlags & CPF_NoClear) : true);
 	bAllowCreate = InArgs._AllowCreate.IsSet() ? InArgs._AllowCreate.GetValue() : (Property ? !Property->HasMetaData("NoCreate") : true);
-
+	bIsSoftObjectPath = CastField<FSoftObjectProperty>(Property) != nullptr;
+	
 	InitializeAssetDataTags(Property);
 	if (DisallowedAssetDataTags.IsValid() || RequiredAssetDataTags.IsValid())
 	{
@@ -338,6 +400,16 @@ void SPropertyEditorAsset::Construct(const FArguments& InArgs, const TSharedPtr<
 	TAttribute<bool> IsEnabledAttribute(this, &SPropertyEditorAsset::CanEdit);
 	TAttribute<FText> TooltipAttribute(this, &SPropertyEditorAsset::OnGetToolTip);
 
+	EditorPathOwner = nullptr;
+	if (bIsActor && bIsSoftObjectPath && FEditorPathHelper::IsEnabled())
+	{
+		if (!GetEditorPathOwnerFromPropertyHandle(GetMostSpecificPropertyHandle(), EditorPathOwner))
+		{
+			IsEnabledAttribute.Set(false);
+			TooltipAttribute.Set(LOCTEXT("InvalidActorEditorPathOwner", "Editing this value with different referencing context is not allowed"));
+		}
+	}
+
 	if (Property)
 	{
 		const FProperty* PropToConsider = GetActualMetadataProperty(Property);
@@ -730,6 +802,7 @@ TSharedRef<SWidget> SPropertyEditorAsset::OnGetMenuContent()
 	{
 		return PropertyCustomizationHelpers::MakeActorPickerWithMenu(Cast<AActor>(Value.Object),
 																	 bAllowClear,
+																	 bIsSoftObjectPath && FEditorPathHelper::IsEnabled(),
 																	 FOnShouldFilterActor::CreateSP( this, &SPropertyEditorAsset::IsFilteredActor ),
 																	 FOnActorSelected::CreateSP( this, &SPropertyEditorAsset::OnActorSelected),
 																	 FSimpleDelegate::CreateSP( this, &SPropertyEditorAsset::CloseComboButton ),
@@ -761,6 +834,13 @@ void SPropertyEditorAsset::OnMenuOpenChanged(bool bOpen)
 bool SPropertyEditorAsset::IsFilteredActor( const AActor* const Actor ) const
 {
 	bool IsAllowed = Actor != nullptr && Actor->IsA(ObjectClass) && !Actor->IsChildActor() && IsClassAllowed(Actor->GetClass());
+
+	if (IsAllowed)
+	{
+		// If we have an EditorPathOwner referenced actor needs to be in same EditorPathOwner
+		IsAllowed = !EditorPathOwner || FEditorPathHelper::IsInEditorPath(EditorPathOwner, Actor);
+	}
+
 	if (IsAllowed && OnShouldFilterActor.IsBound())
 	{
 		IsAllowed = OnShouldFilterActor.Execute(Actor);
@@ -797,7 +877,7 @@ FText SPropertyEditorAsset::OnGetAssetName() const
 
 				if (Actor)
 				{
-					Name = FText::AsCultureInvariant(Actor->GetActorLabel());
+					Name = FText::AsCultureInvariant(GetActorEditorPathLabel(Actor));
 				}
 				else
 				{
@@ -987,7 +1067,7 @@ FPropertyAccess::Result SPropertyEditorAsset::GetValue( FObjectOrAssetData& OutV
 		}
 #endif
 
-		OutValue = FObjectOrAssetData( Object );
+		OutValue = FObjectOrAssetData( Object, EditorPathOwner );
 	}
 	else
 	{
@@ -1094,8 +1174,33 @@ void SPropertyEditorAsset::OnAssetSelected( const struct FAssetData& AssetData )
 	SetValue(AssetData);
 }
 
+SPropertyEditorAsset::FObjectOrAssetData::FObjectOrAssetData(UObject* InObject, UObject* InEditorPathOwner)
+	: Object(InObject)
+{
+	if (AActor* Actor = Cast<AActor>(InObject))
+	{
+		ObjectPath = FEditorPathHelper::GetEditorPathFromEditorPathOwner(Actor, InEditorPathOwner);
+	}
+	else if(InObject != nullptr)
+	{
+		AssetData = FAssetData(InObject);
+		ObjectPath = InObject;
+	}
+}
+
 void SPropertyEditorAsset::OnActorSelected( AActor* InActor )
 {
+	if (FEditorPathHelper::IsEnabled() && bIsSoftObjectPath)
+	{
+		// Even if SetValue ends up calling FSoftObjectProperty::ImportText_Internal the FAssetData validation needs to validate the reference domain which is /Temp when referencing Level Instance objects. So we convert the FAssetData to the EditorPath version to pass validation.
+		FSoftObjectPath EditorPath = FEditorPathHelper::GetEditorPathFromEditorPathOwner(InActor, EditorPathOwner);
+		if (FSoftObjectPath(InActor) != EditorPath)
+		{
+			FAssetData EditorAssetData(EditorPath.GetLongPackageName(), EditorPath.ToString(), FTopLevelAssetPath(InActor->GetClass()->GetPathName()));
+			SetValue(EditorAssetData);
+			return;
+		}
+	}
 	SetValue(InActor);
 }
 
diff --git a/Engine/Source/Editor/PropertyEditor/Private/UserInterface/PropertyEditor/SPropertyEditorAsset.h b/Engine/Source/Editor/PropertyEditor/Private/UserInterface/PropertyEditor/SPropertyEditorAsset.h
--- a/Engine/Source/Editor/PropertyEditor/Private/UserInterface/PropertyEditor/SPropertyEditorAsset.h
+++ b/Engine/Source/Editor/PropertyEditor/Private/UserInterface/PropertyEditor/SPropertyEditorAsset.h
@@ -111,23 +111,21 @@ private:
 		UObject* Object;
 		FSoftObjectPath ObjectPath;
 		FAssetData AssetData;
+		UObject* EditorPathOwner;
 
-		FObjectOrAssetData( UObject* InObject = nullptr )
-			: Object( InObject )
-			, ObjectPath( Object )
-		{
-			AssetData = InObject != nullptr && !InObject->IsA<AActor>() ? FAssetData( InObject ) : FAssetData();
-		}
+		FObjectOrAssetData(UObject* InObject = nullptr, UObject* EditorPathOwner = nullptr);
 
 		FObjectOrAssetData( const FSoftObjectPath& InObjectPath )
 			: Object(nullptr)
 			, ObjectPath(InObjectPath)
+			, EditorPathOwner(nullptr)
 		{}
 
 		FObjectOrAssetData( const FAssetData& InAssetData )
-			: Object( nullptr )
+			: Object(nullptr)
 			, ObjectPath( InAssetData.ToSoftObjectPath() )
 			, AssetData( InAssetData )
+			, EditorPathOwner(nullptr)
 		{}
 
 		bool IsValid() const
@@ -406,6 +404,12 @@ private:
 	/** Whether the classes in the AllowedClassFilters list are exact or whether valid classes can be derived. */
 	bool bExactClass;
 
+	/** Wheter the property is a soft reference */
+	bool bIsSoftObjectPath;
+
+	/** Editor Path context if any */
+	UObject* EditorPathOwner;
+
 	/** The number of additional buttons this picker has. */
 	int32 NumButtons;
 
diff --git a/Engine/Source/Editor/PropertyEditor/Private/UserInterface/PropertyEditor/SPropertyMenuActorPicker.cpp b/Engine/Source/Editor/PropertyEditor/Private/UserInterface/PropertyEditor/SPropertyMenuActorPicker.cpp
--- a/Engine/Source/Editor/PropertyEditor/Private/UserInterface/PropertyEditor/SPropertyMenuActorPicker.cpp
+++ b/Engine/Source/Editor/PropertyEditor/Private/UserInterface/PropertyEditor/SPropertyMenuActorPicker.cpp
@@ -26,6 +26,7 @@ void SPropertyMenuActorPicker::Construct( const FArguments& InArgs )
 {
 	CurrentActor = InArgs._InitialActor;
 	bAllowClear = InArgs._AllowClear;
+	bAllowPickingLevelInstanceContent = InArgs._AllowPickingLevelInstanceContent;
 	ActorFilter = InArgs._ActorFilter;
 	OnSet = InArgs._OnSet;
 	OnClose = InArgs._OnClose;
@@ -96,7 +97,7 @@ void SPropertyMenuActorPicker::Construct( const FArguments& InArgs )
 			.WidthOverride(PropertyEditorAssetConstants::SceneOutlinerWindowSize.X)
 			.HeightOverride(PropertyEditorAssetConstants::SceneOutlinerWindowSize.Y)
 			[
-				SceneOutlinerModule.CreateActorPicker(InitOptions, FOnActorPicked::CreateSP(this, &SPropertyMenuActorPicker::OnActorSelected))
+				SceneOutlinerModule.CreateActorPicker(InitOptions, FOnActorPicked::CreateSP(this, &SPropertyMenuActorPicker::OnActorSelected), nullptr, !bAllowPickingLevelInstanceContent)
 			];
 
 		MenuBuilder.AddWidget(MenuContent.ToSharedRef(), FText::GetEmpty(), true);
diff --git a/Engine/Source/Editor/PropertyEditor/Private/UserInterface/PropertyEditor/SPropertyMenuActorPicker.h b/Engine/Source/Editor/PropertyEditor/Private/UserInterface/PropertyEditor/SPropertyMenuActorPicker.h
--- a/Engine/Source/Editor/PropertyEditor/Private/UserInterface/PropertyEditor/SPropertyMenuActorPicker.h
+++ b/Engine/Source/Editor/PropertyEditor/Private/UserInterface/PropertyEditor/SPropertyMenuActorPicker.h
@@ -15,10 +15,12 @@ public:
 	SLATE_BEGIN_ARGS( SPropertyMenuActorPicker )
 		: _InitialActor(NULL)
 		, _AllowClear(true)
+		, _AllowPickingLevelInstanceContent(false)
 		, _ActorFilter()
 	{}
 		SLATE_ARGUMENT( AActor*, InitialActor )
 		SLATE_ARGUMENT( bool, AllowClear )
+		SLATE_ARGUMENT( bool, AllowPickingLevelInstanceContent )
 		SLATE_ARGUMENT( FOnShouldFilterActor, ActorFilter )
 		SLATE_EVENT( FOnActorSelected, OnSet )
 		SLATE_EVENT( FSimpleDelegate, OnClose )
@@ -82,6 +84,9 @@ private:
 	/** Whether the asset can be 'None' in this case */
 	bool bAllowClear;
 
+	/** Whether we allow picking actors inside level instances */
+	bool bAllowPickingLevelInstanceContent;
+
 	/** Delegate used to test whether an actor should be displayed or not */
 	FOnShouldFilterActor ActorFilter;
 
diff --git a/Engine/Source/Editor/PropertyEditor/Public/PropertyCustomizationHelpers.h b/Engine/Source/Editor/PropertyEditor/Public/PropertyCustomizationHelpers.h
--- a/Engine/Source/Editor/PropertyEditor/Public/PropertyCustomizationHelpers.h
+++ b/Engine/Source/Editor/PropertyEditor/Public/PropertyCustomizationHelpers.h
@@ -109,6 +109,7 @@ namespace PropertyCustomizationHelpers
 	PROPERTYEDITOR_API TSharedRef<SWidget> MakeAssetPickerWithMenu( const FAssetData& InitialObject, const bool AllowClear, const bool AllowCopyPaste, const TArray<const UClass*>& AllowedClasses, const TArray<const UClass*>& DisallowedClasses, const TArray<UFactory*>& NewAssetFactories, FOnShouldFilterAsset OnShouldFilterAsset, FOnAssetSelected OnSet, FSimpleDelegate OnClose, const TSharedPtr<IPropertyHandle>& PropertyHandle = TSharedPtr<IPropertyHandle>(), const TArray<FAssetData>& OwnerAssetArray = TArray<FAssetData>());
 	PROPERTYEDITOR_API TSharedRef<SWidget> MakeActorPickerAnchorButton( FOnGetActorFilters OnGetActorFilters, FOnActorSelected OnActorSelectedFromPicker );
 	PROPERTYEDITOR_API TSharedRef<SWidget> MakeActorPickerWithMenu( AActor* const InitialActor, const bool AllowClear, FOnShouldFilterActor ActorFilter, FOnActorSelected OnSet, FSimpleDelegate OnClose, FSimpleDelegate OnUseSelected );
+	PROPERTYEDITOR_API TSharedRef<SWidget> MakeActorPickerWithMenu(AActor* const InitialActor, const bool AllowClear, const bool AllowPickingLevelInstanceContent, FOnShouldFilterActor ActorFilter, FOnActorSelected OnSet, FSimpleDelegate OnClose, FSimpleDelegate OnUseSelected);
 	PROPERTYEDITOR_API TSharedRef<SWidget> MakeComponentPickerWithMenu( UActorComponent* const InitialComponent, const bool AllowClear, FOnShouldFilterActor ActorFilter, FOnShouldFilterComponent ComponentFilter, FOnComponentSelected OnSet, FSimpleDelegate OnClose );
 	PROPERTYEDITOR_API TSharedRef<SWidget> MakeInteractiveActorPicker(FOnGetAllowedClasses OnGetAllowedClasses, FOnShouldFilterActor OnShouldFilterActor, FOnActorSelected OnActorSelectedFromPicker);
 	PROPERTYEDITOR_API TSharedRef<SWidget> MakeSceneDepthPicker(FOnSceneDepthLocationSelected OnSceneDepthLocationSelected);
