Reset buttons in details panel hard to hover/click

Problem: With recent changes to the details panel some styling updates caused the scrollbar section to overlap with the reset button when the scrollbar is not showing.

Solution: In SDetailsViewBase.cpp, in the GetScrollbarVisibility method call into DisplayManager.GetIsScrollbarShowing, to tell whether or not we need to show the section which has the scrollbar

diff --git a/Engine/Source/Editor/PropertyEditor/Private/SDetailsViewBase.cpp b/Engine/Source/Editor/PropertyEditor/Private/SDetailsViewBase.cpp
--- a/Engine/Source/Editor/PropertyEditor/Private/SDetailsViewBase.cpp
+++ b/Engine/Source/Editor/PropertyEditor/Private/SDetailsViewBase.cpp
@@ -493,8 +493,9 @@ EVisibility SDetailsViewBase::GetTreeVisibility() const
 EVisibility SDetailsViewBase::GetScrollBarVisibility() const
 {
 	const bool bHasAnythingToShow = RootTreeNodes.Num() > 0;
+	const bool bIsScrollbarNeeded = DisplayManager.IsValid() ? DisplayManager->GetIsScrollbarShowing() : true;
 	const bool bShowScrollBar = DetailsViewArgs.bShowScrollBar && bHasAnythingToShow;
-	return bShowScrollBar ? EVisibility::Visible : EVisibility::Collapsed;
+	return bShowScrollBar && bIsScrollbarNeeded ? EVisibility::Visible : EVisibility::Collapsed;
 }
 
 /** Returns the image used for the icon on the filter button */
