Details Panel - Component card layout actions

~ finish up clean up and refactor of details view style. Make sure that all of the numbers are moved to the style file and anything state related stays in the display manager.

diff --git a/Engine/Source/Editor/PropertyEditor/Private/DetailsDisplayManager.cpp b/Engine/Source/Editor/PropertyEditor/Private/DetailsDisplayManager.cpp
--- a/Engine/Source/Editor/PropertyEditor/Private/DetailsDisplayManager.cpp
+++ b/Engine/Source/Editor/PropertyEditor/Private/DetailsDisplayManager.cpp
@@ -3,8 +3,14 @@
 #include "DetailsDisplayManager.h"
 
 #include "DetailsViewStyle.h"
+#include "SDetailsView.h"
 
 
+FDetailsDisplayManager::FDetailsDisplayManager(): bIsOuterCategory(false)
+{
+	PrimaryStyleKey = SDetailsView::GetPrimaryDetailsViewStyleKey();
+}
+
 FDetailsDisplayManager::~FDetailsDisplayManager()
 {
 	OnDetailsNeedsUpdate.Unbind();
@@ -37,7 +43,7 @@ void FDetailsDisplayManager::UpdateView() const
 
 const FDetailsViewStyleKey& FDetailsDisplayManager::GetDetailsViewStyleKey() const
 {
-	return FDetailsViewStyleKeys::Default();
+	return PrimaryStyleKey;
 }
 
 void FDetailsDisplayManager::SetIsOuterCategory(bool bInIsOuterCategory)
@@ -45,39 +51,26 @@ void FDetailsDisplayManager::SetIsOuterCategory(bool bInIsOuterCategory)
 	bIsOuterCategory = bInIsOuterCategory;
 }
 
-FMargin FDetailsDisplayManager::GetTablePadding() const
+const FDetailsViewStyle* FDetailsDisplayManager::GetDetailsViewStyle() const
 {
-	if (GetDetailsViewStyleKey() == FDetailsViewStyleKeys::Card())
-	{
-		return bIsScrollbarShowing ?
-					FDetailsViewStyle::TablePaddingWithScrollbarCard() :
-					FDetailsViewStyle::TablePaddingNoScrollbarCard();
-	}
-	return bIsScrollbarShowing ?
-					FDetailsViewStyle::TablePaddingWithScrollbarClassic() :
-					FDetailsViewStyle::TablePaddingNoScrollbarClassic();
-		
+	const FDetailsViewStyle* ViewStyle = FDetailsViewStyle::GetStyle(GetDetailsViewStyleKey());
+	return ViewStyle;
 }
 
-FMargin FDetailsDisplayManager::GetRowPadding() const
-{
-	FDetailsViewStyle ViewStyle = GetDetailsViewStyleKey();
-	ViewStyle.SetIsOuterCategory(bIsOuterCategory);
-	ViewStyle.SetIsScrollbarShowing(bIsScrollbarShowing);
-	return ViewStyle.GetRowPadding();
-}
 
-FMargin FDetailsDisplayManager::GetCategoryButtonsPadding() const
+
+FMargin FDetailsDisplayManager::GetTablePadding() const
 {
-	return FMargin(0, 0,  bIsOuterCategory ? 10 : 0, 0);
+	const FDetailsViewStyle* Style = GetDetailsViewStyle();
+	return Style ? Style->GetTablePadding(bIsScrollBarNeeded) : 0;
 }
 
-bool FDetailsDisplayManager::GetIsScrollbarShowing() const
+bool FDetailsDisplayManager::GetIsScrollBarNeeded() const
 {
-	return bIsScrollbarShowing;
+	return bIsScrollBarNeeded;
 }
 
-void FDetailsDisplayManager::SetIsScrollbarShowing(bool bInIsScrollbarShowing)
+void FDetailsDisplayManager::SetIsScrollBarNeeded(bool bInIsScrollBarNeeded)
 {
-	bIsScrollbarShowing = bInIsScrollbarShowing;
-}
+	bIsScrollBarNeeded = bInIsScrollBarNeeded;
+}
\ No newline at end of file
diff --git a/Engine/Source/Editor/PropertyEditor/Private/DetailsViewStyle.cpp b/Engine/Source/Editor/PropertyEditor/Private/DetailsViewStyle.cpp
--- a/Engine/Source/Editor/PropertyEditor/Private/DetailsViewStyle.cpp
+++ b/Engine/Source/Editor/PropertyEditor/Private/DetailsViewStyle.cpp
@@ -3,18 +3,40 @@
 
 #include "DetailsViewStyle.h"
 
+FDetailsViewStyle::FDetailsViewStyle()
+{
+}
+
 FDetailsViewStyle::FDetailsViewStyle(
 	const FDetailsViewStyleKey& InKey, 
-	float InTopCategoryPadding,
-	float InHorizontalPadding) :
+	float InTopCategoryPadding) :
 		FSlateWidgetStyle(),
 		Key(InKey),
-		TopCategoryPadding(InTopCategoryPadding),
-		HorizontalPadding(InHorizontalPadding)
+		TopCategoryPadding(InTopCategoryPadding)
 {
 	Initialize(Key);
 }
 
+FMargin FDetailsViewStyle::GetTablePadding(bool bIsScrollBarNeeded) const
+{
+	if (bIsScrollBarNeeded)
+	{
+		return TablePaddingWithScrollbar;
+	}
+	return TablePaddingWithNoScrollbar;
+}
+
+void FDetailsViewStyle::Initialize(const FDetailsViewStyle* Style)
+{
+	if ( Style )
+	{
+		Key = Style->Key;
+		TopCategoryPadding = Style->TopCategoryPadding;
+		TablePaddingWithScrollbar = Style->TablePaddingWithScrollbar;
+		TablePaddingWithNoScrollbar = Style->TablePaddingWithNoScrollbar;
+	}
+}
+
 FDetailsViewStyle::FDetailsViewStyle(FDetailsViewStyleKey& InKey) :
 		FSlateWidgetStyle(),
 		Key(InKey)
@@ -24,32 +46,32 @@ FDetailsViewStyle::FDetailsViewStyle(FDetailsViewStyleKey& InKey) :
 
 void FDetailsViewStyle::Initialize(FDetailsViewStyleKey& InKey)
 {
-	if (const FDetailsViewStyle* Style = GetStyle(InKey))
+	if (const FDetailsViewStyle** Style = StyleKeyToStyleTemplateMap.Find(InKey.GetName()))
 	{
-		Initialize(InKey, Style->HorizontalPadding, Style->TopCategoryPadding);
+		Initialize( *Style );
 	}
 }
 
 FDetailsViewStyle::FDetailsViewStyle(FDetailsViewStyle& InStyle) :
-	FDetailsViewStyle(InStyle.Key, InStyle.TopCategoryPadding, InStyle.HorizontalPadding) 
+	FDetailsViewStyle(InStyle.Key, InStyle.TopCategoryPadding ) 
 {
 }
 
 FDetailsViewStyle::FDetailsViewStyle(const FDetailsViewStyle& InStyle) :
-	FDetailsViewStyle(InStyle.Key, InStyle.TopCategoryPadding, InStyle.HorizontalPadding) 
+	FDetailsViewStyle(InStyle.Key, InStyle.TopCategoryPadding ) 
 {
 }
 
 FMargin FDetailsViewStyle::GetOuterCategoryRowPadding() const
 {
-	return FMargin(HorizontalPadding, TopCategoryPadding, HorizontalPadding, 1);
+	return FMargin(0, TopCategoryPadding, 0, 1);
 }
 
-FMargin FDetailsViewStyle::GetRowPadding() const
+FMargin FDetailsViewStyle::GetRowPadding(bool bIsOuterCategory) const
 {
 	return bIsOuterCategory ?
-		FMargin(HorizontalPadding, TopCategoryPadding, HorizontalPadding, 1) :
-		FMargin(HorizontalPadding, 0, HorizontalPadding, 1);
+		FMargin(0, TopCategoryPadding, 0, 1) :
+		FMargin(0, 0, 0, 1);
 }
 
 bool FDetailsViewStyle::operator==(FDetailsViewStyle& OtherLayoutType) const
@@ -61,7 +83,7 @@ FDetailsViewStyle& FDetailsViewStyle::operator=(FDetailsViewStyleKey& OtherLayou
 {
 	if (const FDetailsViewStyle* Style = GetStyle(OtherLayoutTypeKey))
 	{
-		Initialize(OtherLayoutTypeKey, Style->HorizontalPadding, Style->TopCategoryPadding);
+		Initialize(OtherLayoutTypeKey, Style->TopCategoryPadding);
 	}
 	return *this;
 }
@@ -75,13 +97,13 @@ const FSlateBrush* FDetailsViewStyle::GetBackgroundImageForCategoryRow(
 	const bool bShowBorder, 
 	const bool bIsInnerCategory,
 	const bool bIsCategoryExpanded,
-	const bool bIsScrollBarVisible) const
+	const bool bIsScrollBarNeeded) const
 {
 	static const FSlateBrush* InnerCategoryRowBrush = FAppStyle::Get().GetBrush("DetailsView.CategoryMiddle");
 	static const FSlateBrush* ClassicStyleTopLevelCategoryRowBrush = FAppStyle::Get().GetBrush("DetailsView.CategoryTop");
-	static const FSlateBrush* CardStyleTopLevelCategoryCollapsedScrollbarVisibleRowBrush = FAppStyle::Get().GetBrush("DetailsView.CardHeaderRounded");
+	static const FSlateBrush* CardStyleTopLevelCategoryCollapsedScrollBarNeededRowBrush = FAppStyle::Get().GetBrush("DetailsView.CardHeaderRounded");
 	static const FSlateBrush* CardStyleTopLevelCategoryCollapsedScrollbarHiddenRowBrush = FAppStyle::Get().GetBrush("DetailsView.CardHeaderLeftSideRounded");
-	static const FSlateBrush* CardStyleTopLevelCategoryExpandedScrollbarVisibleRowBrush = FAppStyle::Get().GetBrush("DetailsView.CardHeaderTopRounded");
+	static const FSlateBrush* CardStyleTopLevelCategoryExpandedScrollBarNeededRowBrush = FAppStyle::Get().GetBrush("DetailsView.CardHeaderTopRounded");
 
 	if (bShowBorder)
 	{
@@ -98,13 +120,13 @@ const FSlateBrush* FDetailsViewStyle::GetBackgroundImageForCategoryRow(
 		}
 		if (!bIsCategoryExpanded)
 		{
-			if (bIsScrollBarVisible)
+			if (bIsScrollBarNeeded)
 			{
-				return CardStyleTopLevelCategoryCollapsedScrollbarVisibleRowBrush;
+				return CardStyleTopLevelCategoryCollapsedScrollBarNeededRowBrush;
 			}
 			return CardStyleTopLevelCategoryCollapsedScrollbarHiddenRowBrush;
 		}
-		return CardStyleTopLevelCategoryExpandedScrollbarVisibleRowBrush;
+		return CardStyleTopLevelCategoryExpandedScrollBarNeededRowBrush;
 	}
 
 	return nullptr;
@@ -112,25 +134,28 @@ const FSlateBrush* FDetailsViewStyle::GetBackgroundImageForCategoryRow(
 
 void FDetailsViewStyle::InitializeDetailsViewStyles()
 {
-	static FDetailsViewStyle CardStyle{FDetailsViewStyleKeys::Card(), 8.0f, 0.0f};
+	static FDetailsViewStyle CardStyle{FDetailsViewStyleKeys::Card(), 8.0f};
+	CardStyle.TablePaddingWithScrollbar = FMargin(8, 0, 20, 8);
+	CardStyle.TablePaddingWithNoScrollbar = FMargin(8, 0, 8, 8);
 	StyleKeyToStyleTemplateMap.Add(FDetailsViewStyleKeys::Card().GetName(), &CardStyle);
 
-	static FDetailsViewStyle DefaultStyle{FDetailsViewStyleKeys::Default(), 0.0f, 0.0f};
+	static FDetailsViewStyle DefaultStyle{FDetailsViewStyleKeys::Default(), 0.0f };
 	StyleKeyToStyleTemplateMap.Add(FDetailsViewStyleKeys::Default().GetName(), &DefaultStyle);
 	
+	static FDetailsViewStyle ClassicStyle{FDetailsViewStyleKeys::Classic(), 0.0f };
+	StyleKeyToStyleTemplateMap.Add(FDetailsViewStyleKeys::Classic().GetName(), &ClassicStyle);
 }
 
 const FSlateBrush* FDetailsViewStyle::GetBackgroundImageForScrollBarWell(
 	const bool bShowBorder,
 	const bool bIsInnerCategory,
 	const bool bIsCategoryExpanded,
-	const bool bIsScrollBarVisible) const
+	const bool bIsScrollBarNeeded) const
 {
 	static const FSlateBrush* InnerCategoryWellBrush = FAppStyle::Get().GetBrush("DetailsView.CategoryMiddle");
 	static const FSlateBrush* ClassicStyleTopLevelCategoryRowBrush = FAppStyle::Get().GetBrush("DetailsView.CategoryTop");
-	static const FSlateBrush* DetailsBackgroundWellBrush = FAppStyle::Get().GetBrush("DetailsView.GridLine");
-	static const FSlateBrush* CardStyleCollapsedScrollbarVisibleWellBrush = FAppStyle::Get().GetBrush("DetailsView.CardHeaderRightSideRounded");
-	static const FSlateBrush* CardStyleExpandedScrollbarVisibleWellBrush = FAppStyle::Get().GetBrush("DetailsView.CardHeaderTopRightSideRounded");
+	static const FSlateBrush* CardStyleCollapsedScrollBarNeededWellBrush = FAppStyle::Get().GetBrush("DetailsView.CardHeaderRightSideRounded");
+	static const FSlateBrush* CardStyleExpandedScrollBarNeededWellBrush = FAppStyle::Get().GetBrush("DetailsView.CardHeaderTopRightSideRounded");
 
 	if (bShowBorder)
 	{
@@ -145,9 +170,9 @@ const FSlateBrush* FDetailsViewStyle::GetBackgroundImageForScrollBarWell(
 			return ClassicStyleTopLevelCategoryRowBrush;
 		}
 		if (!bIsCategoryExpanded){
-			return CardStyleCollapsedScrollbarVisibleWellBrush;
+			return CardStyleCollapsedScrollBarNeededWellBrush;
 		}
-		return CardStyleExpandedScrollbarVisibleWellBrush;
+		return CardStyleExpandedScrollBarNeededWellBrush;
 	}
 
 	return nullptr;
@@ -155,32 +180,18 @@ const FSlateBrush* FDetailsViewStyle::GetBackgroundImageForScrollBarWell(
 
 void FDetailsViewStyle::Initialize(
 	FDetailsViewStyleKey& InKey,
-	const float InHorizontalPadding,
 	const float InTopCategoryPadding)
 {
 	Key = InKey;
-	HorizontalPadding = InHorizontalPadding;
 	TopCategoryPadding = InTopCategoryPadding;
 }
 
-const FDetailsViewStyle* FDetailsViewStyle::GetStyle(FDetailsViewStyleKey Key)
+const FDetailsViewStyle* FDetailsViewStyle::GetStyle(const FDetailsViewStyleKey& Key)
 {
-	const FDetailsViewStyle** StylePtr;
-		
-	StylePtr = StyleKeyToStyleTemplateMap.Find(Key.GetName());
+	const FDetailsViewStyle** StylePtr = StyleKeyToStyleTemplateMap.Find(Key.GetName());
 	StylePtr = StylePtr ?
 					StylePtr :
 					StyleKeyToStyleTemplateMap.Find(FDetailsViewStyleKeys::Default().GetName());
 		
 	return StylePtr ? *StylePtr : nullptr;
-}
-
-void FDetailsViewStyle::SetIsOuterCategory(bool bInIsOuterCategory)
-{
-	bIsOuterCategory = bInIsOuterCategory;
-}
-
-void FDetailsViewStyle::SetIsScrollbarShowing(bool bInIsScrollbarShowing)
-{
-	bIsScrollbarShowing = bInIsScrollbarShowing;
 }
\ No newline at end of file
diff --git a/Engine/Source/Editor/PropertyEditor/Private/DetailsViewStyle.h b/Engine/Source/Editor/PropertyEditor/Private/DetailsViewStyle.h
--- a/Engine/Source/Editor/PropertyEditor/Private/DetailsViewStyle.h
+++ b/Engine/Source/Editor/PropertyEditor/Private/DetailsViewStyle.h
@@ -23,8 +23,6 @@ public:
 	 * Constructs this @code FDetailsViewStyle @endcode with the specified parameters
 	 *
 	 * @param InKey the Key for  the @code FDetailsViewStyle @endcode
-	 * @param InTopCategoryPadding the top padding for the Category row
-	 * @param InHorizontalPadding the horizontal padding for the Category row
 	 */
 	FDetailsViewStyle(FDetailsViewStyleKey& InKey);
 
@@ -49,27 +47,13 @@ public:
 	 */
 	 FMargin GetOuterCategoryRowPadding() const;
 	
-	/**
-	* Sets a bool indicating whether or not the Category this is a Style for is an outer Category
-	*
-	* @param bIsOuterCategory a bool indicating whether or not the Category this is a Style for is an outer Category
-	*/
-	void SetIsOuterCategory(bool bIsOuterCategory);
-
-	/**
-	* Sets a bool indicating whether or not the Scrollbar is showing on the details view
-	*
-	* @param bInIsScrollbarShowing a bool indicating whether or not the Scrollbar is showing on the details view
-	*/
-	void SetIsScrollbarShowing(bool bInIsScrollbarShowing);
-	
 	/**
 	 * Returns the padding for details panel rows which are not outer Category rows
 	 */
-	FMargin GetRowPadding() const;
+	FMargin GetRowPadding(bool bIsOuterCategory) const;
 
 	/**
-	 * The equality operator for @code FDetailsViewStule @endcode
+	 * The equality operator for @code FDetailsViewStyle @endcode
 	 */
 	bool operator==(FDetailsViewStyle& OtherLayoutType) const;
 
@@ -121,71 +105,55 @@ public:
 	 *
 	 * @param InKey the name of the FDetailsViewStyle
 	 * @param InTopCategoryPadding the top padding for the Category row
-	 * @param InHorizontalPadding the horizontal padding for the Category row
 	 */
 	FDetailsViewStyle(
 		const FDetailsViewStyleKey& InKey,
-		float InTopCategoryPadding = 0.f,
-		float InHorizontalPadding = 0.f);
-
-	
-	/** the Padding for the entirety of the details panel if no scrollbar is present for Classic style */
-	static const FMargin& TablePaddingNoScrollbarClassic()
-	{
-		static FMargin Margin = FMargin(0, 0, 0, 1);
-		return Margin;
-	}
-	
-	/** the Padding for the entirety of the details panel if no scrollbar is present for Card style */
-	static const FMargin& TablePaddingNoScrollbarCard()
-	{
-		static FMargin Margin = FMargin(8, 0, 8, 1);
-		return Margin;
-	}
+		float InTopCategoryPadding = 0.f);
 
-	/** the Padding for the entirety of the details panel if the scrollbar is present for Classic style */
-	static const FMargin TablePaddingWithScrollbarClassic()
-	{
-		static FMargin Margin = FMargin(0, 0, 16, 1);
-		return Margin;
-	}
+	/**
+	 * Returns the FMargin which provides the padding around the whole details view table
+	 *
+	 * @param bIsScrollBarVisible whether the scrollbar is visible
+	 */
+	FMargin GetTablePadding(bool bIsScrollBarVisible) const;
 
-	/** the Padding for the entirety of the details panel if the scrollbar is present for Card style  */
-	static const FMargin TablePaddingWithScrollbarCard()
+	/**
+	 * Returns the FMargin which provides the padding around the Category buttons
+	 */
+	FMargin GetCategoryButtonsMargin() const
 	{
-		static FMargin Margin = FMargin(8, 0, 20, 1);
-		return Margin;
+		return CategoryButtonsMargin;
 	}
 	
+	/**
+	* Returns the @code FDetailsViewStyle @endcode for which this is the @code FDetailsViewStyleKey @endcode .
+	*
+	* @param InKey the name of the FDetailsViewStyle
+	*/
+	static const FDetailsViewStyle* GetStyle(const FDetailsViewStyleKey& InKey);
+	
 private:
+
+	void Initialize( const FDetailsViewStyle* Style );
+
 	/**
    	 * Initializes this FDetailsViewStyle with the specified parameters
    	 *
    	 * @param InKey the name of the FDetailsViewStyle
    	 * @param InTopCategoryPadding the top padding for the Category row
-   	 * @param InHorizontalPadding the horizontal padding for the Category row
    	 */
-	void Initialize(
-		FDetailsViewStyleKey& InKey, 
-		const float InHorizontalPadding, 
+	void Initialize( FDetailsViewStyleKey& InKey, 
 		const float InTopCategoryPadding);
 
-		/**
-		* Initializes this FDetailsViewStyle with the style specified by @code FDetailsViewStyleKey @endcode InKey
-		*
-		* @param InKey the name of the FDetailsViewStyle
-		*/
-		void Initialize(FDetailsViewStyleKey& InKey);
-
-	/** the Name of the Style */
-	FDetailsViewStyleKey Key;
-
 	/**
-	* Returns the @code FDetailsViewStyle @endcode for which this is the @code FDetailsViewStyleKey @endcode .
+	* Initializes this FDetailsViewStyle with the style specified by @code FDetailsViewStyleKey @endcode InKey
 	*
 	* @param InKey the name of the FDetailsViewStyle
 	*/
-	static const FDetailsViewStyle* GetStyle(FDetailsViewStyleKey InKey);
+	void Initialize(FDetailsViewStyleKey& InKey);
+
+	/** the Name of the Style */
+	FDetailsViewStyleKey Key;
 
 	/**
 	 * A map of @code FDetailsViewStyleKey @endcode instances to the
@@ -196,13 +164,21 @@ private:
 	/** the Slate Units of the top padding for an outer Category row */
 	float TopCategoryPadding = 0.f;
 	
-	/** the Slate Units of the horizontal padding for all details rows */
-	float HorizontalPadding = 0.f;
-	
-	/** Whether the current Category is an outer versus an inner Category  */
-   	bool bIsOuterCategory = false;
+ 	/**
+	* The FMargin which provides the padding around the whole details view table with scrollbar
+	*/
+	FMargin TablePaddingWithScrollbar = FMargin(0, 0, 16, 1);
 
-	/** Whether the scrollbar is showing on the details view  */
-	bool bIsScrollbarShowing = false;
+	/**
+	* The FMargin which provides the padding around the whole details view table with no scrollbar
+	*/
+	FMargin TablePaddingWithNoScrollbar = FMargin(0, 0, 0, 1);
+	
+	/**
+	* The FMargin which provides the padding around the Category buttons
+	*/
+	FMargin CategoryButtonsMargin = FMargin(0, 0,  10, 0);
+	
+	
 };
 
diff --git a/Engine/Source/Editor/PropertyEditor/Private/SDetailCategoryTableRow.cpp b/Engine/Source/Editor/PropertyEditor/Private/SDetailCategoryTableRow.cpp
--- a/Engine/Source/Editor/PropertyEditor/Private/SDetailCategoryTableRow.cpp
+++ b/Engine/Source/Editor/PropertyEditor/Private/SDetailCategoryTableRow.cpp
@@ -33,6 +33,11 @@ void SDetailCategoryTableRow::Construct(const FArguments& InArgs, TSharedRef<FDe
 	ObjectName = InArgs._ObjectName;
 	DisplayManager = DetailsView->GetDisplayManager();
 
+	if (DisplayManager.IsValid() && DisplayManager->GetDetailsViewStyle())
+	{
+		DetailsViewStyle = DisplayManager->GetDetailsViewStyle();
+	}
+
 	InitializeDisplayManager();
 
 	PulseAnimation.AddCurve(0.0f, UE::PropertyEditor::Private::PulseAnimationLength, ECurveEaseFunction::CubicInOut);
@@ -86,25 +91,13 @@ void SDetailCategoryTableRow::Construct(const FArguments& InArgs, TSharedRef<FDe
 
 	OwnerTableViewWeak = InOwnerTableView;
 	
-	auto GetScrollbarWellTint = [this]()
-	{
-		return SDetailTableRowBase::IsScrollBarVisible(OwnerTableViewWeak) ?
-			FSlateColor(EStyleColor::Header) :
-			this->GetInnerBackgroundColor();
-	};
-
 	this->ChildSlot
 	[
 		SNew(SBorder)
 		.BorderImage(FAppStyle::Get().GetBrush("DetailsView.GridLine"))
 		.Padding_Lambda([this]
 		{
-			if (!DisplayManager.IsValid())
-			{
-				return FMargin{0};
-			}
-			DisplayManager->SetIsOuterCategory(!bIsInnerCategory);
-			return DisplayManager->GetRowPadding();
+			return DetailsViewStyle ? DetailsViewStyle->GetRowPadding(!bIsInnerCategory) : 0;
 		})
 		[
 			SNew(SOverlay)
@@ -130,13 +123,7 @@ void SDetailCategoryTableRow::Construct(const FArguments& InArgs, TSharedRef<FDe
 			.Visibility(EVisibility::Visible)
 			.Padding_Lambda([this]
 			{
-				if (!DisplayManager.IsValid())
-				{
-					return FMargin{0};
-				}
-				DisplayManager->SetIsOuterCategory(!bIsInnerCategory);
-				DisplayManager->SetIsScrollbarShowing(IsScrollBarVisible(OwnerTableViewWeak));
-				return DisplayManager->GetCategoryButtonsPadding();
+				return DetailsViewStyle ? DetailsViewStyle->GetCategoryButtonsMargin() : 0;
 			})
 			[
 			DisplayManager.IsValid() ?
@@ -251,17 +238,11 @@ const FSlateBrush* SDetailCategoryTableRow::GetBackgroundImage() const
 		{
 			return FAppStyle::Get().GetBrush("DetailsView.CategoryMiddle");
 		}
-		
-		IDetailsViewPrivate* View = OwnerTreeNode.IsValid() && OwnerTreeNode.Pin()->GetDetailsView() ? OwnerTreeNode.Pin()->GetDetailsView() : nullptr;
-		
-		static FDetailsViewStyleKey PrimaryKey = SDetailsView::GetPrimaryDetailsViewStyleKey();
-		FDetailsViewStyle ViewStyle = View ? View->GetStyleKey() : PrimaryKey;
+
+		const bool bIsScrollBarNeeded = IsScrollBarVisible(OwnerTableViewWeak);
+		DisplayManager->SetIsScrollBarNeeded(bIsScrollBarNeeded);
 		const bool bIsCategoryExpanded = IsItemExpanded();
-		const bool bIsScrollBarVisible = IsScrollBarVisible(OwnerTableViewWeak);
-		DisplayManager->SetIsScrollbarShowing(bIsScrollBarVisible);
-		ViewStyle.SetIsOuterCategory(!bIsInnerCategory);
-		
-		return ViewStyle.GetBackgroundImageForCategoryRow(bShowBorder, bIsInnerCategory, bIsCategoryExpanded, bIsScrollBarVisible);
+		return DetailsViewStyle ? DetailsViewStyle->GetBackgroundImageForCategoryRow(bShowBorder, bIsInnerCategory, bIsCategoryExpanded, bIsScrollBarNeeded) : nullptr;
 	}
 	return nullptr;
 }
@@ -270,15 +251,10 @@ const FSlateBrush* SDetailCategoryTableRow::GetBackgroundImageForScrollBarWell()
 {
 	if (bShowBorder)
 	{
-		IDetailsViewPrivate* View = OwnerTreeNode.IsValid() && OwnerTreeNode.Pin()->GetDetailsView() ? OwnerTreeNode.Pin()->GetDetailsView() : nullptr;
-		static FDetailsViewStyleKey PrimaryKey = SDetailsView::GetPrimaryDetailsViewStyleKey();
-		const FDetailsViewStyle ViewStyle = View ? View->GetStyleKey() : PrimaryKey;
 		const bool bIsCategoryExpanded = IsItemExpanded();
 		const bool bIsScrollBarVisible = IsScrollBarVisible(OwnerTableViewWeak);
-		DisplayManager->SetIsScrollbarShowing(bIsScrollBarVisible);
-		DisplayManager->SetIsOuterCategory(!bIsInnerCategory);
-
-		return ViewStyle.GetBackgroundImageForScrollBarWell(bShowBorder, bIsInnerCategory, bIsCategoryExpanded, bIsScrollBarVisible);
+		return DetailsViewStyle ? DetailsViewStyle->GetBackgroundImageForScrollBarWell(
+			bShowBorder, bIsInnerCategory, bIsCategoryExpanded, bIsScrollBarVisible) : nullptr;
 	}
 	return nullptr;
 }
diff --git a/Engine/Source/Editor/PropertyEditor/Private/SDetailCategoryTableRow.h b/Engine/Source/Editor/PropertyEditor/Private/SDetailCategoryTableRow.h
--- a/Engine/Source/Editor/PropertyEditor/Private/SDetailCategoryTableRow.h
+++ b/Engine/Source/Editor/PropertyEditor/Private/SDetailCategoryTableRow.h
@@ -11,6 +11,8 @@
 #include "Widgets/DeclarativeSyntaxSupport.h"
 #include "Widgets/Views/STableRow.h"
 #include "Widgets/Views/STableViewBase.h"
+#include "DetailsViewStyle.h"
+
 
 class SDetailCategoryTableRow : public SDetailTableRowBase
 {
@@ -69,7 +71,13 @@ private:
 
 	/** The name of the object that is specified by this category */
 	FName ObjectName;
-	
+
+	/**
+	* The style of the details view. This holds information and methods to get things like the row border image,
+	* table padding, and so forth.
+	*/
+	const FDetailsViewStyle* DetailsViewStyle = nullptr;
+
 	/** Previously parsed clipboard data. */
 	UE::PropertyEditor::Internal::FClipboardData PreviousClipboardData;
 
diff --git a/Engine/Source/Editor/PropertyEditor/Private/SDetailsView.cpp b/Engine/Source/Editor/PropertyEditor/Private/SDetailsView.cpp
--- a/Engine/Source/Editor/PropertyEditor/Private/SDetailsView.cpp
+++ b/Engine/Source/Editor/PropertyEditor/Private/SDetailsView.cpp
@@ -415,7 +415,12 @@ void SDetailsView::Construct(const FArguments& InArgs, const FDetailsViewArgs& I
 	           SNew(SBorder)
 				.Padding_Lambda([this]
 				{
-					return DisplayManager->GetTablePadding();
+					if (DisplayManager.IsValid())
+					{
+						return DisplayManager->GetTablePadding();
+					}
+					static const FMargin Margin{0};
+					return Margin;
 				})
 				.BorderImage_Lambda([this]
 				{
diff --git a/Engine/Source/Editor/PropertyEditor/Private/SDetailsViewBase.cpp b/Engine/Source/Editor/PropertyEditor/Private/SDetailsViewBase.cpp
--- a/Engine/Source/Editor/PropertyEditor/Private/SDetailsViewBase.cpp
+++ b/Engine/Source/Editor/PropertyEditor/Private/SDetailsViewBase.cpp
@@ -493,9 +493,9 @@ EVisibility SDetailsViewBase::GetTreeVisibility() const
 EVisibility SDetailsViewBase::GetScrollBarVisibility() const
 {
 	const bool bHasAnythingToShow = RootTreeNodes.Num() > 0;
-	const bool bIsScrollbarNeeded = DisplayManager.IsValid() ? DisplayManager->GetIsScrollbarShowing() : true;
-	const bool bShowScrollBar = DetailsViewArgs.bShowScrollBar && bHasAnythingToShow;
-	return bShowScrollBar && bIsScrollbarNeeded ? EVisibility::Visible : EVisibility::Collapsed;
+	const bool bIsScrollbarNeeded = DisplayManager.IsValid() ? DisplayManager->GetIsScrollBarNeeded() : true;
+	const bool bShowScrollBar = DetailsViewArgs.bShowScrollBar && bHasAnythingToShow && bIsScrollbarNeeded;
+	return bShowScrollBar ? EVisibility::Visible : EVisibility::Collapsed;
 }
 
 /** Returns the image used for the icon on the filter button */
diff --git a/Engine/Source/Editor/PropertyEditor/Public/DetailsDisplayManager.h b/Engine/Source/Editor/PropertyEditor/Public/DetailsDisplayManager.h
--- a/Engine/Source/Editor/PropertyEditor/Public/DetailsDisplayManager.h
+++ b/Engine/Source/Editor/PropertyEditor/Public/DetailsDisplayManager.h
@@ -7,6 +7,8 @@
 
 DECLARE_DELEGATE(FOnDetailsNeedsUpdate)
 
+class FDetailsViewStyle;
+
 /** An @code FDetailsDisplayManager @endcode provides an API to tweak various settings of your details view, and
  * provides some utility methods to work with Details.  */
 class FDetailsDisplayManager : public TSharedFromThis<FDetailsDisplayManager>
@@ -14,9 +16,7 @@ class FDetailsDisplayManager : public TSharedFromThis<FDetailsDisplayManager>
 public:
 	FOnDetailsNeedsUpdate OnDetailsNeedsUpdate;
 
-	PROPERTYEDITOR_API FDetailsDisplayManager() : bIsOuterCategory(false)
-	{
-	}
+	PROPERTYEDITOR_API FDetailsDisplayManager();
 
 	PROPERTYEDITOR_API virtual ~FDetailsDisplayManager();
 
@@ -59,40 +59,39 @@ public:
 	void SetIsOuterCategory(bool bInIsOuterCategory);
 
 	/**
-	 * Returns the padding for details panel rows which are not outer Category rows
-	 */
-	FMargin GetRowPadding() const;
-
-	/**
-	 * Returns the padding for any buttons for Category rows
+	 * Returns the @code FDetailsViewStyle @endcode that is the current FDetailsViewStyle style
 	 */
-	FMargin GetCategoryButtonsPadding() const;
-
-	/**
-	 * returns a boolean indicating whether or not the currently active categpory is an inner category
-	 */	
-	bool GetIsInnerCategory() const;
-
-	/**
-   	 * Returns a bool indicating whether or not the scrollbar is showing on the details view
-   	 */
-	PROPERTYEDITOR_API virtual bool GetIsScrollbarShowing() const;
+	PROPERTYEDITOR_API const FDetailsViewStyle* GetDetailsViewStyle() const;
+	
+	    /**
+		* Returns a bool indicating whether or not the scrollbar is needed on the details view. Note that the "needed"
+		* here means that in this value the work has been done to figure out if the scrollbar should show, and
+		* anything can query this to see if it needs to alter the display accordingly
+		*/
+	PROPERTYEDITOR_API virtual bool GetIsScrollBarNeeded() const;
 	
 	/**
-	* Set a bool indicating whether or not the scrollbar is showing on the details view
+	* Set a bool indicating whether or not the scrollbar is needed on the details view. Note that the "needed"
+	* here means that in this value the work has been done to figure out if the scrollbar should show, and
+	* anything can query this to see if it needs to alter the display accordingly
 	*
-	* @param bInIsScrollbarShowing a bool indicating whether or not the scrollbar is showing on the details view
+	* @param bInIsScrollBarNeeded a bool indicating whether or not the scrollbar is Needed on the details view
 	*/
-	PROPERTYEDITOR_API virtual void SetIsScrollbarShowing(bool bInIsScrollbarShowing);
+	PROPERTYEDITOR_API virtual void SetIsScrollBarNeeded(bool bInIsScrollBarNeeded);
 	
 	/**
-	* Returns an FMargin containing the padding for the entire details view. Note that if there is a scrollbar present,
-	* this padding takes that into account
-	*/
-	PROPERTYEDITOR_API virtual FMargin GetTablePadding() const;
+	 * Returns the FMargin which provides the padding around the whole details view table
+	 */
+	FMargin GetTablePadding() const;
 
 protected:
 
+	
+	/**
+	 * The primary style key for the details view. 
+	 */
+	FDetailsViewStyleKey PrimaryStyleKey;
+	
 	/**
 	 * The name of the object defined by the currently active category
 	 */
@@ -104,9 +103,10 @@ protected:
 	bool bIsOuterCategory;
 
 	/**
-	 * A bool indicating whether or not the scrollbar is showing on the details view
+	* a bool indicating whether or not the scrollbar is needed on the details view. Note that the "needed"
+	* here means that in this value the work has been done to figure out if the scrollbar should show 
 	 */
-	bool bIsScrollbarShowing;
+	bool bIsScrollBarNeeded = false;
 
 };
 
