Material Slot widget fixes:
- Use FillWidth so that its responsive to its parent size.
- Removed the delete button when the widget is being used in the level editor details panel. It was always disabled but the icon was taking space.
- Changed the text "Slot name" to "Slot"

diff --git a/Engine/Source/Editor/DetailCustomizations/Private/ComponentMaterialCategory.cpp b/Engine/Source/Editor/DetailCustomizations/Private/ComponentMaterialCategory.cpp
--- a/Engine/Source/Editor/DetailCustomizations/Private/ComponentMaterialCategory.cpp
+++ b/Engine/Source/Editor/DetailCustomizations/Private/ComponentMaterialCategory.cpp
@@ -587,7 +587,7 @@ TSharedRef<SWidget> FComponentMaterialCategory::OnGenerateWidgetsForMaterial(UMa
 		SNew(SMaterialSlotWidget, SlotIndex, true)
 		.MaterialName(this, &FComponentMaterialCategory::GetMaterialNameText, SlotIndex)
 		.IsMaterialSlotNameReadOnly(true)
-		.CanDeleteMaterialSlot(false);
+		.DeleteMaterialSlotVisibility(EVisibility::Collapsed);
 }
 
 FText FComponentMaterialCategory::GetMaterialNameText(int32 MaterialIndex) const
diff --git a/Engine/Source/Editor/PropertyEditor/Private/MaterialList.cpp b/Engine/Source/Editor/PropertyEditor/Private/MaterialList.cpp
--- a/Engine/Source/Editor/PropertyEditor/Private/MaterialList.cpp
+++ b/Engine/Source/Editor/PropertyEditor/Private/MaterialList.cpp
@@ -187,7 +187,7 @@ TSharedRef<SWidget> FMaterialItemView::CreateValueContent(IDetailLayoutBuilder&
 		[
 			SNew(SHorizontalBox)
 			+SHorizontalBox::Slot()
-			.AutoWidth()
+			.FillWidth(1.f)
 			[
 				SNew(SObjectPropertyEntryBox)
 				.ObjectPath(this, &FMaterialItemView::OnGetObjectPath)
@@ -234,7 +234,7 @@ TSharedRef<SWidget> FMaterialItemView::CreateValueContent(IDetailLayoutBuilder&
 						]
 						+SHorizontalBox::Slot()
 						.Padding(3.0f, 0.0f)
-						.AutoWidth()
+						.FillWidth(1.f)
 						[
 							OnGenerateCustomMaterialWidgets.IsBound() ? OnGenerateCustomMaterialWidgets.Execute(MaterialItem.Material.Get(), MaterialItem.SlotIndex) : SNullWidget::NullWidget
 						]
diff --git a/Engine/Source/Editor/PropertyEditor/Private/PropertyCustomizationHelpers.cpp b/Engine/Source/Editor/PropertyEditor/Private/PropertyCustomizationHelpers.cpp
--- a/Engine/Source/Editor/PropertyEditor/Private/PropertyCustomizationHelpers.cpp
+++ b/Engine/Source/Editor/PropertyEditor/Private/PropertyCustomizationHelpers.cpp
@@ -1462,6 +1462,8 @@ void SMaterialSlotWidget::Construct(const FArguments& InArgs, int32 SlotIndex, b
 			LOCTEXT("CustomNameMaterialNotUsedDeleteTooltip", "Delete this material slot"),
 			InArgs._CanDeleteMaterialSlot);
 
+	DeleteButton->SetVisibility(InArgs._DeleteMaterialSlotVisibility);
+
 	ChildSlot
 	[
 		SAssignNew(SlotNameBox, SHorizontalBox)
@@ -1475,7 +1477,7 @@ void SMaterialSlotWidget::Construct(const FArguments& InArgs, int32 SlotIndex, b
 			[
 				SNew(STextBlock)
 				.Font(IDetailLayoutBuilder::GetDetailFont())
-				.Text(LOCTEXT("MaterialArrayNameLabelStringKey", "Slot Name"))
+				.Text(LOCTEXT("MaterialArrayNameLabelStringKey", "Slot"))
 			]
 		]
 		+ SHorizontalBox::Slot()
diff --git a/Engine/Source/Editor/PropertyEditor/Public/PropertyCustomizationHelpers.h b/Engine/Source/Editor/PropertyEditor/Public/PropertyCustomizationHelpers.h
--- a/Engine/Source/Editor/PropertyEditor/Public/PropertyCustomizationHelpers.h
+++ b/Engine/Source/Editor/PropertyEditor/Public/PropertyCustomizationHelpers.h
@@ -557,6 +557,7 @@ private:
 class SMaterialSlotWidget : public SCompoundWidget
 {
 	SLATE_BEGIN_ARGS(SMaterialSlotWidget)
+		: _DeleteMaterialSlotVisibility(EVisibility::Visible)
 	{}
 		SLATE_ATTRIBUTE(FText, MaterialName)
 		SLATE_ATTRIBUTE(bool, IsMaterialSlotNameReadOnly)
@@ -564,6 +565,7 @@ class SMaterialSlotWidget : public SCompoundWidget
 		SLATE_EVENT(FOnTextCommitted, OnMaterialNameCommitted)
 		SLATE_ATTRIBUTE(bool, CanDeleteMaterialSlot)
 		SLATE_EVENT(FSimpleDelegate, OnDeleteMaterialSlot)
+		SLATE_ATTRIBUTE(EVisibility, DeleteMaterialSlotVisibility)
 	SLATE_END_ARGS()
 
 	PROPERTYEDITOR_API void Construct(const FArguments& InArgs, int32 SlotIndex, bool bIsMaterialUsed);
