Fixed property clamping not taking into account the instance meta data if present on the property handle

diff --git a/Engine/Source/Editor/PropertyEditor/Private/PropertyHandleImpl.cpp b/Engine/Source/Editor/PropertyEditor/Private/PropertyHandleImpl.cpp
--- a/Engine/Source/Editor/PropertyEditor/Private/PropertyHandleImpl.cpp
+++ b/Engine/Source/Editor/PropertyEditor/Private/PropertyHandleImpl.cpp
@@ -869,46 +869,38 @@ bool FPropertyValueImpl::IsPropertyTypeOf(FFieldClass* ClassType ) const
 }
 
 template< typename Type>
-static Type ClampValueFromMetaData(Type InValue, FPropertyNode& InPropertyNode )
+static Type ClampValueFromMetaData(Type InValue, FPropertyHandleBase& InPropertyHandle)
 {
-	FProperty* Property = InPropertyNode.GetProperty();
-
 	Type RetVal = InValue;
-	if( Property )
+	//enforce min
+	const FString& MinString = InPropertyHandle.GetMetaData(TEXT("ClampMin"));
+	if(MinString.Len())
 	{
-		//enforce min
-		const FString& MinString = Property->GetMetaData(TEXT("ClampMin"));
-		if(MinString.Len())
-		{
-			checkSlow(MinString.IsNumeric());
-			Type MinValue;
-			TTypeFromString<Type>::FromString(MinValue, *MinString);
-			RetVal = FMath::Max<Type>(MinValue, RetVal);
-		}
-		//Enforce max 
-		const FString& MaxString = Property->GetMetaData(TEXT("ClampMax"));
-		if(MaxString.Len())
-		{
-			checkSlow(MaxString.IsNumeric());
-			Type MaxValue;
-			TTypeFromString<Type>::FromString(MaxValue, *MaxString);
-			RetVal = FMath::Min<Type>(MaxValue, RetVal);
-		}
+		checkSlow(MinString.IsNumeric());
+		Type MinValue;
+		TTypeFromString<Type>::FromString(MinValue, *MinString);
+		RetVal = FMath::Max<Type>(MinValue, RetVal);
+	}
+	//Enforce max 
+	const FString& MaxString = InPropertyHandle.GetMetaData(TEXT("ClampMax"));
+	if(MaxString.Len())
+	{
+		checkSlow(MaxString.IsNumeric());
+		Type MaxValue;
+		TTypeFromString<Type>::FromString(MaxValue, *MaxString);
+		RetVal = FMath::Min<Type>(MaxValue, RetVal);
 	}
 
 	return RetVal;
 }
 
 template <typename Type>
-static Type ClampIntegerValueFromMetaData(Type InValue, FPropertyNode& InPropertyNode )
+static Type ClampIntegerValueFromMetaData(Type InValue, FPropertyHandleBase& InPropertyHandle, FPropertyNode& InPropertyNode)
 {
-	Type RetVal = ClampValueFromMetaData<Type>( InValue, InPropertyNode );
-
-	FProperty* Property = InPropertyNode.GetProperty();
-
+	Type RetVal = ClampValueFromMetaData<Type>(InValue, InPropertyHandle);
 
 	//if there is "Multiple" meta data, the selected number is a multiple
-	const FString& MultipleString = Property->GetMetaData(TEXT("Multiple"));
+	const FString& MultipleString = InPropertyHandle.GetMetaData(TEXT("Multiple"));
 	if (MultipleString.Len())
 	{
 		check(MultipleString.IsNumeric());
@@ -921,7 +913,7 @@ static Type ClampIntegerValueFromMetaData(Type InValue, FPropertyNode& InPropert
 	}
 
 	//enforce array bounds
-	const FString& ArrayClampString = Property->GetMetaData(TEXT("ArrayClamp"));
+	const FString& ArrayClampString = InPropertyHandle.GetMetaData(TEXT("ArrayClamp"));
 	if (ArrayClampString.Len())
 	{
 		FObjectPropertyNode* ObjectPropertyNode = InPropertyNode.FindObjectItemParent();
@@ -932,7 +924,7 @@ static Type ClampIntegerValueFromMetaData(Type InValue, FPropertyNode& InPropert
 		}
 		else
 		{
-			UE_LOG(LogPropertyNode, Warning, TEXT("Array Clamping isn't supported in multi-select (Param Name: %s)"), *Property->GetName());
+			UE_LOG(LogPropertyNode, Warning, TEXT("Array Clamping isn't supported in multi-select (Param Name: %s)"), *InPropertyHandle.GetProperty()->GetName());
 		}
 	}
 
@@ -3589,7 +3581,7 @@ FPropertyAccess::Result FPropertyHandleInt::SetValue(const int8& NewValue, EProp
 {
 	FPropertyAccess::Result Res;
 	// Clamp the value from any meta data ranges stored on the property value
-	int8 FinalValue = ClampIntegerValueFromMetaData<int8>( NewValue, *Implementation->GetPropertyNode() );
+	int8 FinalValue = ClampIntegerValueFromMetaData<int8>( NewValue, *this, *Implementation->GetPropertyNode());
 
 	const FString ValueStr = LexToString(FinalValue);
 	Res = Implementation->ImportText(ValueStr, Flags);
@@ -3602,7 +3594,7 @@ FPropertyAccess::Result FPropertyHandleInt::SetValue(const int16& NewValue, EPro
 {
 	FPropertyAccess::Result Res;
 	// Clamp the value from any meta data ranges stored on the property value
-	int16 FinalValue = ClampIntegerValueFromMetaData<int16>(NewValue, *Implementation->GetPropertyNode());
+	int16 FinalValue = ClampIntegerValueFromMetaData<int16>(NewValue, *this, *Implementation->GetPropertyNode());
 
 	const FString ValueStr = LexToString(FinalValue);
 	Res = Implementation->ImportText(ValueStr, Flags);
@@ -3615,7 +3607,7 @@ FPropertyAccess::Result FPropertyHandleInt::SetValue( const int32& NewValue, EPr
 {
 	FPropertyAccess::Result Res;
 	// Clamp the value from any meta data ranges stored on the property value
-	int32 FinalValue = ClampIntegerValueFromMetaData<int32>( NewValue, *Implementation->GetPropertyNode() );
+	int32 FinalValue = ClampIntegerValueFromMetaData<int32>( NewValue, *this, *Implementation->GetPropertyNode());
 
 	const FString ValueStr = LexToString(FinalValue);
 	Res = Implementation->ImportText( ValueStr, Flags );
@@ -3628,7 +3620,7 @@ FPropertyAccess::Result FPropertyHandleInt::SetValue(const int64& NewValue, EPro
 	FPropertyAccess::Result Res;
 
 	// Clamp the value from any meta data ranges stored on the property value
-	int64 FinalValue = ClampIntegerValueFromMetaData<int64>(NewValue, *Implementation->GetPropertyNode());
+	int64 FinalValue = ClampIntegerValueFromMetaData<int64>(NewValue, *this, *Implementation->GetPropertyNode());
 
 	const FString ValueStr = LexToString(FinalValue);
 	Res = Implementation->ImportText(ValueStr, Flags);
@@ -3639,7 +3631,7 @@ FPropertyAccess::Result FPropertyHandleInt::SetValue(const uint16& NewValue, EPr
 {
 	FPropertyAccess::Result Res;
 	// Clamp the value from any meta data ranges stored on the property value
-	uint16 FinalValue = ClampIntegerValueFromMetaData<uint16>(NewValue, *Implementation->GetPropertyNode());
+	uint16 FinalValue = ClampIntegerValueFromMetaData<uint16>(NewValue, *this, *Implementation->GetPropertyNode());
 
 	const FString ValueStr = LexToString(FinalValue);
 	Res = Implementation->ImportText(ValueStr, Flags);
@@ -3652,7 +3644,7 @@ FPropertyAccess::Result FPropertyHandleInt::SetValue(const uint32& NewValue, EPr
 {
 	FPropertyAccess::Result Res;
 	// Clamp the value from any meta data ranges stored on the property value
-	uint32 FinalValue = ClampIntegerValueFromMetaData<uint32>(NewValue, *Implementation->GetPropertyNode());
+	uint32 FinalValue = ClampIntegerValueFromMetaData<uint32>(NewValue, *this, *Implementation->GetPropertyNode());
 
 	const FString ValueStr = LexToString(FinalValue);
 	Res = Implementation->ImportText(ValueStr, Flags);
@@ -3664,7 +3656,7 @@ FPropertyAccess::Result FPropertyHandleInt::SetValue(const uint64& NewValue, EPr
 {
 	FPropertyAccess::Result Res;
 	// Clamp the value from any meta data ranges stored on the property value
-	uint64 FinalValue = ClampIntegerValueFromMetaData<uint64>(NewValue, *Implementation->GetPropertyNode());
+	uint64 FinalValue = ClampIntegerValueFromMetaData<uint64>(NewValue, *this, *Implementation->GetPropertyNode());
 
 	const FString ValueStr = LexToString(FinalValue);
 	Res = Implementation->ImportText(ValueStr, Flags);
@@ -3701,7 +3693,7 @@ FPropertyAccess::Result FPropertyHandleFloat::SetValue( const float& NewValue, E
 {
 	FPropertyAccess::Result Res;
 	// Clamp the value from any meta data ranges stored on the property value
-	float FinalValue = ClampValueFromMetaData<float>( NewValue, *Implementation->GetPropertyNode() );
+	float FinalValue = ClampValueFromMetaData<float>( NewValue, *this);
 
 	const FString ValueStr = FString::Printf( TEXT("%f"), FinalValue );
 	Res = Implementation->ImportText( ValueStr, Flags );
@@ -3739,7 +3731,7 @@ FPropertyAccess::Result FPropertyHandleDouble::SetValue( const double& NewValue,
 {
 	FPropertyAccess::Result Res;
 	// Clamp the value from any meta data ranges stored on the property value
-	double FinalValue = ClampValueFromMetaData<double>( NewValue, *Implementation->GetPropertyNode() );
+	double FinalValue = ClampValueFromMetaData<double>( NewValue, *this);
 	
 	FDoubleProperty* NumericProperty = CastFieldChecked<FDoubleProperty>(Implementation->GetPropertyNode()->GetProperty());
 	const FString ValueStr = NumericProperty->GetNumericPropertyValueToString(&FinalValue);
@@ -4355,7 +4347,7 @@ FPropertyAccess::Result FPropertyHandleMixed::SetValue(const double& NewValue, E
 {
 	FPropertyAccess::Result Res;
 	// Clamp the value from any meta data ranges stored on the property value
-	double FinalValue = ClampValueFromMetaData<double>(NewValue, *Implementation->GetPropertyNode());
+	double FinalValue = ClampValueFromMetaData<double>(NewValue, *this);
 
 	const FString ValueStr = FString::Printf(TEXT("%f"), FinalValue);
 	Res = Implementation->ImportText(ValueStr, Flags);
