Expand the use of ForceRebuildProperty to include properties inside nested properties such as UStructs.

diff --git a/Engine/Source/Editor/PropertyEditor/Private/PropertyNode.cpp b/Engine/Source/Editor/PropertyEditor/Private/PropertyNode.cpp
--- a/Engine/Source/Editor/PropertyEditor/Private/PropertyNode.cpp
+++ b/Engine/Source/Editor/PropertyEditor/Private/PropertyNode.cpp
@@ -3093,6 +3093,13 @@ void FPropertyNode::NotifyPostChange( FPropertyChangedEvent& InPropertyChangedEv
 	BroadcastPropertyChangedDelegates(InPropertyChangedEvent);
 	BroadcastPropertyChangedDelegates();
 
+	// Reset these values
+	if (PropertyChain->Num() > 0)
+	{
+		PropertyChain->SetActiveMemberPropertyNode(OriginalActiveProperty);
+		PropertyChain->SetActivePropertyNode(InPropertyChangedEvent.Property);
+	}
+
 	// Call through to the property window's notify hook.
 	if( InNotifyHook )
 	{
@@ -3102,32 +3109,37 @@ void FPropertyNode::NotifyPostChange( FPropertyChangedEvent& InPropertyChangedEv
 		}
 		else
 		{
-			PropertyChain->SetActiveMemberPropertyNode( OriginalActiveProperty );
-			PropertyChain->SetActivePropertyNode( InPropertyChangedEvent.Property);
-		
 			InPropertyChangedEvent.SetActiveMemberProperty(OriginalActiveProperty);
 			InNotifyHook->NotifyPostChange( InPropertyChangedEvent, &PropertyChain.Get() );
 		}
 	}
 
 
-	if( OriginalActiveProperty )
+	// For each Property in the Property Chain, see if it has ForceRebuildProperty metadata and find the sibling PropertyNode to rebuild.
+	// To do that, we need to match up the FPropertyNode (Editor representation) with the FProperty (Engine representation)
+	TSharedPtr<FPropertyNode> CurrentPropertyNode = FindObjectItemParent()->AsShared();
+	for (auto PropertyChainNode = PropertyChain->GetActiveMemberNode(); PropertyChainNode && CurrentPropertyNode.IsValid() ; PropertyChainNode = PropertyChainNode->GetNextNode())
 	{
-		//if i have metadata forcing other property windows to rebuild
-		const FString& MetaData = OriginalActiveProperty->GetMetaData(TEXT("ForceRebuildProperty"));
-
-		if( MetaData.Len() > 0 )
+		if (const FProperty* CurrentProperty = PropertyChainNode->GetValue())
 		{
-			// We need to find the property node beginning at the root/parent, not at our own node.
-			ObjectNode = FindObjectItemParent();
-			check(ObjectNode != NULL);
-
-			TSharedPtr<FPropertyNode> ForceRebuildNode = ObjectNode->FindChildPropertyNode( FName(*MetaData), true );
-
-			if( ForceRebuildNode.IsValid() )
+			const static FName NAME_ForceRebuildProperty(TEXT("ForceRebuildProperty"));
+			const FString& ForceRebuildPropertyName = CurrentProperty->GetMetaData(NAME_ForceRebuildProperty);
+			if (!ForceRebuildPropertyName.IsEmpty())
 			{
-				ForceRebuildNode->RequestRebuildChildren();
+				constexpr bool bRecursive = true;
+				TSharedPtr<FPropertyNode> ForceRebuildNode = CurrentPropertyNode->FindChildPropertyNode(FName(*ForceRebuildPropertyName, FNAME_Find), bRecursive);
+
+				if (ForceRebuildNode.IsValid())
+				{
+					ForceRebuildNode->RequestRebuildChildren();
+				}
+				else
+				{
+					UE_LOG(LogPropertyNode, Error, TEXT("Could not find named property '%s' referenced from %s ForceRebuildProperty"), *ForceRebuildPropertyName, *CurrentPropertyNode->GetDisplayName().ToString());
+				}
 			}
+
+			CurrentPropertyNode = CurrentPropertyNode->FindChildPropertyNode(CurrentProperty->GetFName());
 		}
 	}
 
diff --git a/Engine/Source/Runtime/CoreUObject/Public/UObject/ObjectMacros.h b/Engine/Source/Runtime/CoreUObject/Public/UObject/ObjectMacros.h
--- a/Engine/Source/Runtime/CoreUObject/Public/UObject/ObjectMacros.h
+++ b/Engine/Source/Runtime/CoreUObject/Public/UObject/ObjectMacros.h
@@ -1337,6 +1337,9 @@ namespace UM
 		/// [PropertyMetadata] Deprecated.
 		FixedIncrement,
 
+		/// [PropertyMetaData] Force the specified sibling Editor PropertyNode to refresh itself when this property changes.  Useful to trigger an Editor refresh of an unrelated Property when this one is modified in PostEditChangeProperty or similar.
+		ForceRebuildProperty,
+
 		/// [PropertyMetadata] Used by asset properties. Indicates that the asset pickers should always show engine content
 		ForceShowEngineContent,
 
