Recent change resulted in a crash in Niagara due to the assumption of a valid parent property to always exist

diff --git a/Engine/Source/Editor/PropertyEditor/Private/PropertyNode.cpp b/Engine/Source/Editor/PropertyEditor/Private/PropertyNode.cpp
--- a/Engine/Source/Editor/PropertyEditor/Private/PropertyNode.cpp
+++ b/Engine/Source/Editor/PropertyEditor/Private/PropertyNode.cpp
@@ -3117,29 +3117,32 @@ void FPropertyNode::NotifyPostChange( FPropertyChangedEvent& InPropertyChangedEv
 
 	// For each Property in the Property Chain, see if it has ForceRebuildProperty metadata and find the sibling PropertyNode to rebuild.
 	// To do that, we need to match up the FPropertyNode (Editor representation) with the FProperty (Engine representation)
-	TSharedPtr<FPropertyNode> CurrentPropertyNode = FindObjectItemParent()->AsShared();
-	for (auto PropertyChainNode = PropertyChain->GetActiveMemberNode(); PropertyChainNode && CurrentPropertyNode.IsValid() ; PropertyChainNode = PropertyChainNode->GetNextNode())
+	if(FindObjectItemParent() != nullptr)
 	{
-		if (const FProperty* CurrentProperty = PropertyChainNode->GetValue())
+		TSharedPtr<FPropertyNode> CurrentPropertyNode = FindObjectItemParent()->AsShared();
+		for (auto PropertyChainNode = PropertyChain->GetActiveMemberNode(); PropertyChainNode && CurrentPropertyNode.IsValid() ; PropertyChainNode = PropertyChainNode->GetNextNode())
 		{
-			const static FName NAME_ForceRebuildProperty(TEXT("ForceRebuildProperty"));
-			const FString& ForceRebuildPropertyName = CurrentProperty->GetMetaData(NAME_ForceRebuildProperty);
-			if (!ForceRebuildPropertyName.IsEmpty())
+			if (const FProperty* CurrentProperty = PropertyChainNode->GetValue())
 			{
-				constexpr bool bRecursive = true;
-				TSharedPtr<FPropertyNode> ForceRebuildNode = CurrentPropertyNode->FindChildPropertyNode(FName(*ForceRebuildPropertyName, FNAME_Find), bRecursive);
-
-				if (ForceRebuildNode.IsValid())
-				{
-					ForceRebuildNode->RequestRebuildChildren();
-				}
-				else
+				const static FName NAME_ForceRebuildProperty(TEXT("ForceRebuildProperty"));
+				const FString& ForceRebuildPropertyName = CurrentProperty->GetMetaData(NAME_ForceRebuildProperty);
+				if (!ForceRebuildPropertyName.IsEmpty())
 				{
-					UE_LOG(LogPropertyNode, Error, TEXT("Could not find named property '%s' referenced from %s ForceRebuildProperty"), *ForceRebuildPropertyName, *CurrentPropertyNode->GetDisplayName().ToString());
+					constexpr bool bRecursive = true;
+					TSharedPtr<FPropertyNode> ForceRebuildNode = CurrentPropertyNode->FindChildPropertyNode(FName(*ForceRebuildPropertyName, FNAME_Find), bRecursive);
+
+					if (ForceRebuildNode.IsValid())
+					{
+						ForceRebuildNode->RequestRebuildChildren();
+					}
+					else
+					{
+						UE_LOG(LogPropertyNode, Error, TEXT("Could not find named property '%s' referenced from %s ForceRebuildProperty"), *ForceRebuildPropertyName, *CurrentPropertyNode->GetDisplayName().ToString());
+					}
 				}
-			}
 
-			CurrentPropertyNode = CurrentPropertyNode->FindChildPropertyNode(CurrentProperty->GetFName());
+				CurrentPropertyNode = CurrentPropertyNode->FindChildPropertyNode(CurrentProperty->GetFName());
+			}
 		}
 	}
 
