Support showing components with no UProperties and add a disallowed types list.

~added isEmpty property to details categories for showing stub categories
~styling support to show them properly

diff --git a/Engine/Source/Editor/PropertyEditor/Private/DetailCategoryBuilderImpl.cpp b/Engine/Source/Editor/PropertyEditor/Private/DetailCategoryBuilderImpl.cpp
--- a/Engine/Source/Editor/PropertyEditor/Private/DetailCategoryBuilderImpl.cpp
+++ b/Engine/Source/Editor/PropertyEditor/Private/DetailCategoryBuilderImpl.cpp
@@ -602,6 +602,16 @@ FDetailLayoutCustomization* FDetailCategoryImpl::GetDefaultCustomization(TShared
 	return Customization;
 }
 
+bool FDetailCategoryImpl::IsEmpty() const
+{
+	return bIsEmpty;
+}
+
+void FDetailCategoryImpl::SetIsEmpty(bool bInIsEmpty)
+{
+	bIsEmpty = bInIsEmpty;
+}
+ 
 bool FDetailCategoryImpl::ShouldAdvancedBeExpanded() const
 {
 	return bUserShowAdvanced || bForceAdvanced;
@@ -824,6 +834,7 @@ TSharedRef<ITableRow> FDetailCategoryImpl::GenerateWidgetForTableView(const TSha
 	return SNew(SDetailCategoryTableRow, AsShared(), OwnerTable)
 		.PasteFromText(OnPasteFromText())
 		.ObjectName( ObjectName )
+		.IsEmpty( bIsEmpty )
 		.InnerCategory(ParentLayout.IsValid() ? ParentLayout->IsLayoutForExternalRoot() : false)
 		.DisplayName(GetDisplayName())
 		.HeaderContent(HeaderContent);
diff --git a/Engine/Source/Editor/PropertyEditor/Private/DetailCategoryBuilderImpl.h b/Engine/Source/Editor/PropertyEditor/Private/DetailCategoryBuilderImpl.h
--- a/Engine/Source/Editor/PropertyEditor/Private/DetailCategoryBuilderImpl.h
+++ b/Engine/Source/Editor/PropertyEditor/Private/DetailCategoryBuilderImpl.h
@@ -322,6 +322,18 @@ public:
 	bool GetShouldBeInitiallyCollapsed() const { return bShouldBeInitiallyCollapsed; }
 
 	FDetailLayoutCustomization* GetDefaultCustomization(TSharedRef<FPropertyNode> PropertyNode);
+	
+	/**
+    * If true, this Category should have no UProperty data associated with it, and will be shown as an empty stub
+    * with no expansion arrow
+    */
+	virtual bool IsEmpty() const override;
+	
+	/**
+	 * Sets whether this Category is "Empty" ~ that is, should have no UProperty data associated with it, and will be shown
+	 * as an empty stub with no expansion arrow
+	 */
+	virtual void SetIsEmpty(bool bInIsEmpty) override;
 
 private:
 	virtual void OnItemExpansionChanged(bool bIsExpanded, bool bShouldSaveState) override;
@@ -452,4 +464,10 @@ private:
 	bool bHasVisibleAdvanced : 1;
 	bool bPendingRefresh : 1;
 	bool bPendingRefreshNeedsRefilter : 1;
+
+	/**
+	* If true, this Category should have no UProperty data associated with it, and will be shown as an empty stub
+	* with no expansion arrow
+	*/
+	bool bIsEmpty 	: 1;
 };
diff --git a/Engine/Source/Editor/PropertyEditor/Private/DetailLayoutBuilderImpl.cpp b/Engine/Source/Editor/PropertyEditor/Private/DetailLayoutBuilderImpl.cpp
--- a/Engine/Source/Editor/PropertyEditor/Private/DetailLayoutBuilderImpl.cpp
+++ b/Engine/Source/Editor/PropertyEditor/Private/DetailLayoutBuilderImpl.cpp
@@ -1055,4 +1055,32 @@ bool FDetailLayoutBuilderImpl::IsPropertyPathAllowed(const FString& InPath) cons
 	}
 
 	return false;
-}
\ No newline at end of file
+}
+
+bool FDetailLayoutBuilderImpl::AddEmptyCategoryIfNeeded(TSharedPtr<FComplexPropertyNode> Node)
+{
+	if (DefaultCategoryMap.IsEmpty())
+	{
+		if (Node.IsValid() &&
+			Node->AsObjectNode() &&
+			Node->AsObjectNode()->GetNumObjects() > 0)
+		{
+			const FName PropertyCategoryName = FObjectEditorUtils::GetCategoryFName(Node->GetProperty());
+			FDetailCategoryImpl& Category = DefaultCategory(PropertyCategoryName);
+			const FName InstanceName = Node->GetProperty()->GetFName();
+			
+			FObjectPropertyNode* ObjectPropertyNode =  Node->AsObjectNode();
+			UObject* OutObject = ObjectPropertyNode->GetUObject(0);
+
+			if (DetailsView &&
+				DetailsView->GetDisplayManager().IsValid() &&
+				DetailsView->GetDisplayManager()->ShowEmptyCategoryIfRootUObjectHasNoPropertyData(OutObject))
+			{
+				Category.AddPropertyNode(Node.ToSharedRef(), InstanceName);
+				Category.SetIsEmpty(true);
+				return true;
+			}
+		}
+	}
+	return false;
+}
diff --git a/Engine/Source/Editor/PropertyEditor/Private/DetailLayoutBuilderImpl.h b/Engine/Source/Editor/PropertyEditor/Private/DetailLayoutBuilderImpl.h
--- a/Engine/Source/Editor/PropertyEditor/Private/DetailLayoutBuilderImpl.h
+++ b/Engine/Source/Editor/PropertyEditor/Private/DetailLayoutBuilderImpl.h
@@ -55,6 +55,17 @@ public:
 	virtual void SortCategories(const FOnCategorySortOrderFunction& SortFunction) override;
 	virtual void SetPropertyGenerationAllowListPaths(const TSet<FString>& InPropertyGenerationAllowListPaths) override;
 	virtual bool IsPropertyPathAllowed(const FString& InPath) const override;
+
+	/**
+	 * Creates an empty category row if there currently are no categories and one is required for the
+	 * @code TSharedPtr<FComplexPropertyNode> @endcode Node
+	 *
+	 * @param Node The @code TSharedPtr<FComplexPropertyNode> @endcode that we will add an empty category for, if needed
+	 *
+	 * @return true if an empty/stub category was added, else it returns false 
+	 */
+	bool AddEmptyCategoryIfNeeded(TSharedPtr<FComplexPropertyNode> Node);
+
 	/**
 	 * Creates a default category. The SDetails view will generate widgets in default categories
 	 *
diff --git a/Engine/Source/Editor/PropertyEditor/Private/DetailsDisplayManager.cpp b/Engine/Source/Editor/PropertyEditor/Private/DetailsDisplayManager.cpp
--- a/Engine/Source/Editor/PropertyEditor/Private/DetailsDisplayManager.cpp
+++ b/Engine/Source/Editor/PropertyEditor/Private/DetailsDisplayManager.cpp
@@ -62,14 +62,17 @@ const FDetailsViewStyle* FDetailsDisplayManager::GetDetailsViewStyle() const
 	return ViewStyle;
 }
 
-
-
 FMargin FDetailsDisplayManager::GetTablePadding() const
 {
 	const FDetailsViewStyle* Style = GetDetailsViewStyle();
 	return Style ? Style->GetTablePadding(bIsScrollBarNeeded) : 0;
 }
 
+bool FDetailsDisplayManager::ShowEmptyCategoryIfRootUObjectHasNoPropertyData(UObject* InNode) const
+{
+	return false;
+}
+
 bool FDetailsDisplayManager::GetIsScrollBarNeeded() const
 {
 	return bIsScrollBarNeeded;
diff --git a/Engine/Source/Editor/PropertyEditor/Private/DetailsViewStyle.cpp b/Engine/Source/Editor/PropertyEditor/Private/DetailsViewStyle.cpp
--- a/Engine/Source/Editor/PropertyEditor/Private/DetailsViewStyle.cpp
+++ b/Engine/Source/Editor/PropertyEditor/Private/DetailsViewStyle.cpp
@@ -96,13 +96,11 @@ const FName FDetailsViewStyle::GetTypeName() const
 const FSlateBrush* FDetailsViewStyle::GetBackgroundImageForCategoryRow(
 	const bool bShowBorder, 
 	const bool bIsInnerCategory,
-	const bool bIsCategoryExpanded,
-	const bool bIsScrollBarNeeded) const
+	const bool bIsCategoryExpanded) const
 {
 	static const FSlateBrush* InnerCategoryRowBrush = FAppStyle::Get().GetBrush("DetailsView.CategoryMiddle");
 	static const FSlateBrush* ClassicStyleTopLevelCategoryRowBrush = FAppStyle::Get().GetBrush("DetailsView.CategoryTop");
 	static const FSlateBrush* CardStyleTopLevelCategoryCollapsedScrollBarNeededRowBrush = FAppStyle::Get().GetBrush("DetailsView.CardHeaderRounded");
-	static const FSlateBrush* CardStyleTopLevelCategoryCollapsedScrollbarHiddenRowBrush = FAppStyle::Get().GetBrush("DetailsView.CardHeaderLeftSideRounded");
 	static const FSlateBrush* CardStyleTopLevelCategoryExpandedScrollBarNeededRowBrush = FAppStyle::Get().GetBrush("DetailsView.CardHeaderTopRounded");
 
 	if (bShowBorder)
@@ -120,11 +118,7 @@ const FSlateBrush* FDetailsViewStyle::GetBackgroundImageForCategoryRow(
 		}
 		if (!bIsCategoryExpanded)
 		{
-			if (bIsScrollBarNeeded)
-			{
-				return CardStyleTopLevelCategoryCollapsedScrollBarNeededRowBrush;
-			}
-			return CardStyleTopLevelCategoryCollapsedScrollbarHiddenRowBrush;
+			return CardStyleTopLevelCategoryCollapsedScrollBarNeededRowBrush;
 		}
 		return CardStyleTopLevelCategoryExpandedScrollBarNeededRowBrush;
 	}
diff --git a/Engine/Source/Editor/PropertyEditor/Private/DetailsViewStyle.h b/Engine/Source/Editor/PropertyEditor/Private/DetailsViewStyle.h
--- a/Engine/Source/Editor/PropertyEditor/Private/DetailsViewStyle.h
+++ b/Engine/Source/Editor/PropertyEditor/Private/DetailsViewStyle.h
@@ -72,13 +72,11 @@ public:
 	 *
  	 * @param bShowBorder a bool to indicate whether the border should be shown, at all
 	 * @param bIsCategoryExpanded a bool that indicates whether this Category is expanded 
-	 * @param bIsScrollBarVisible a bool that indicates whether the scrollbar for the details view is visible
 	 */	
 	const FSlateBrush* GetBackgroundImageForCategoryRow(
 		const bool bShowBorder,
 		const bool bIsInnerCategory,
-		const bool bIsCategoryExpanded,
-		const bool bIsScrollBarVisible) const;
+		const bool bIsCategoryExpanded) const;
 
 	/**
 	 * Returns the background image for the scroll bar well
diff --git a/Engine/Source/Editor/PropertyEditor/Private/SDetailCategoryTableRow.cpp b/Engine/Source/Editor/PropertyEditor/Private/SDetailCategoryTableRow.cpp
--- a/Engine/Source/Editor/PropertyEditor/Private/SDetailCategoryTableRow.cpp
+++ b/Engine/Source/Editor/PropertyEditor/Private/SDetailCategoryTableRow.cpp
@@ -19,6 +19,7 @@
 #include "UserInterface/PropertyEditor/PropertyEditorConstants.h"
 #include "Widgets/Input/SComboButton.h"
 #include "Widgets/Layout/SSeparator.h"
+#include "Brushes/SlateColorBrush.h"
 
 void SDetailCategoryTableRow::Construct(const FArguments& InArgs, TSharedRef<FDetailTreeNode> InOwnerTreeNode, const TSharedRef<STableViewBase>& InOwnerTableView)
 {
@@ -27,6 +28,7 @@ void SDetailCategoryTableRow::Construct(const FArguments& InArgs, TSharedRef<FDe
 	DisplayName = InArgs._DisplayName;
 	bIsInnerCategory = InArgs._InnerCategory;
 	bShowBorder = InArgs._ShowBorder;
+	bIsEmpty = InArgs._IsEmpty;
 
 	IDetailsViewPrivate* DetailsView = InOwnerTreeNode->GetDetailsView();
 	FDetailColumnSizeData& ColumnSizeData = DetailsView->GetColumnSizeData();
@@ -67,6 +69,15 @@ void SDetailCategoryTableRow::Construct(const FArguments& InArgs, TSharedRef<FDe
 		.AutoWidth()
 		[
 			SNew(SDetailExpanderArrow, SharedThis(this))
+				//if this is a stub category
+				.Visibility_Lambda([this]
+				{
+					if (bIsEmpty)
+					{
+						return EVisibility::Hidden;
+					}
+					return EVisibility::Visible;
+				})
 		]
 		+ SHorizontalBox::Slot()
 		.VAlign(VAlign_Center)
@@ -119,7 +130,7 @@ void SDetailCategoryTableRow::Construct(const FArguments& InArgs, TSharedRef<FDe
 			.HAlign(HAlign_Right)
 			[
 			SNew(SBorder)
-			 .BorderImage(this, &SDetailCategoryTableRow::GetBackgroundImageForScrollBarWell)
+			.BorderImage(new FSlateColorBrush(FLinearColor::Transparent))
 			.Visibility(EVisibility::Visible)
 			.Padding_Lambda([this]
 			{
@@ -241,8 +252,8 @@ const FSlateBrush* SDetailCategoryTableRow::GetBackgroundImage() const
 
 		const bool bIsScrollBarNeeded = IsScrollBarVisible(OwnerTableViewWeak);
 		DisplayManager->SetIsScrollBarNeeded(bIsScrollBarNeeded);
-		const bool bIsCategoryExpanded = IsItemExpanded();
-		return DetailsViewStyle ? DetailsViewStyle->GetBackgroundImageForCategoryRow(bShowBorder, bIsInnerCategory, bIsCategoryExpanded, bIsScrollBarNeeded) : nullptr;
+		const bool bIsCategoryExpanded = IsItemExpanded() && !bIsEmpty;
+		return DetailsViewStyle ? DetailsViewStyle->GetBackgroundImageForCategoryRow(bShowBorder, bIsInnerCategory, bIsCategoryExpanded) : nullptr;
 	}
 	return nullptr;
 }
@@ -251,7 +262,7 @@ const FSlateBrush* SDetailCategoryTableRow::GetBackgroundImageForScrollBarWell()
 {
 	if (bShowBorder)
 	{
-		const bool bIsCategoryExpanded = IsItemExpanded();
+		const bool bIsCategoryExpanded = IsItemExpanded() && !bIsEmpty;
 		const bool bIsScrollBarVisible = IsScrollBarVisible(OwnerTableViewWeak);
 		return DetailsViewStyle ? DetailsViewStyle->GetBackgroundImageForScrollBarWell(
 			bShowBorder, bIsInnerCategory, bIsCategoryExpanded, bIsScrollBarVisible) : nullptr;
diff --git a/Engine/Source/Editor/PropertyEditor/Private/SDetailCategoryTableRow.h b/Engine/Source/Editor/PropertyEditor/Private/SDetailCategoryTableRow.h
--- a/Engine/Source/Editor/PropertyEditor/Private/SDetailCategoryTableRow.h
+++ b/Engine/Source/Editor/PropertyEditor/Private/SDetailCategoryTableRow.h
@@ -20,12 +20,17 @@ public:
 	SLATE_BEGIN_ARGS(SDetailCategoryTableRow)
 		: _InnerCategory(false)
 		, _ShowBorder(true)
+		, _IsEmpty(false)
 		, _ObjectName(NAME_Name)
 	{}
 		SLATE_ARGUMENT(FText, DisplayName)
 		SLATE_ARGUMENT(bool, InnerCategory)
 		SLATE_ARGUMENT(TSharedPtr<SWidget>, HeaderContent)
 		SLATE_ARGUMENT(bool, ShowBorder)
+		
+		/** If true, this Category should have no UProperty data associated with it, and will be shown as an 
+		* empty stub with no expansion arrow */
+		SLATE_ARGUMENT(bool, IsEmpty)
 		SLATE_ARGUMENT(TSharedPtr<FOnPasteFromText>, PasteFromText)
 		SLATE_ARGUMENT(FName, ObjectName)
 	SLATE_END_ARGS()
@@ -83,6 +88,13 @@ private:
 
 	bool bIsInnerCategory = false;
 	bool bShowBorder = false;
+
+	/**
+	* If true, this Category should have no UProperty data associated with it, and will be shown as an empty stub
+	* with no expansion arrow
+	*/
+	bool bIsEmpty = false;
+
 	FUIAction CopyAction;
 	FUIAction PasteAction;
 
diff --git a/Engine/Source/Editor/PropertyEditor/Private/SDetailSingleItemRow.cpp b/Engine/Source/Editor/PropertyEditor/Private/SDetailSingleItemRow.cpp
--- a/Engine/Source/Editor/PropertyEditor/Private/SDetailSingleItemRow.cpp
+++ b/Engine/Source/Editor/PropertyEditor/Private/SDetailSingleItemRow.cpp
@@ -492,11 +492,11 @@ void SDetailSingleItemRow::Construct( const FArguments& InArgs, FDetailLayoutCus
 				if (PropertyNode->IsReorderable() || 
 					(CastField<FArrayProperty>(PropertyNode->GetProperty()) != nullptr && 
 					CastField<FObjectProperty>(CastField<FArrayProperty>(PropertyNode->GetProperty())->Inner) != nullptr)) // Is an object array
-				{
+						{
 					DragLeaveDelegate = FOnTableRowDragLeave::CreateSP(this, &SDetailSingleItemRow::OnArrayOrCustomDragLeave);
 					AcceptDropDelegate = FOnAcceptDrop::CreateSP(this, PropertyNode->IsReorderable() ? &SDetailSingleItemRow::OnArrayAcceptDrop : &SDetailSingleItemRow::OnArrayHeaderAcceptDrop);
 					CanAcceptDropDelegate = FOnCanAcceptDrop::CreateSP(this, &SDetailSingleItemRow::OnArrayCanAcceptDrop);
-				}
+						}
 			}
 
 			NameColumnBox->AddSlot()
@@ -674,38 +674,49 @@ void SDetailSingleItemRow::Construct( const FArguments& InArgs, FDetailLayoutCus
 	};
 
 	static const FDetailsViewStyleKey& PrimaryKey = SDetailsView::GetPrimaryDetailsViewStyleKey();
-	
-	this->ChildSlot
-	[
-		SNew( SBorder )
-		.BorderImage(FAppStyle::Get().GetBrush("DetailsView.GridLine"))
-		.Padding(FMargin(0,0,0,1))
-		.Clipping(EWidgetClipping::ClipToBounds)
+
+	// If this is a stub category with no UProperty data, just show a null widget, we don't have anything useful to show here
+	if (Category.IsValid() && Category->IsEmpty())
+	{
+		this->ChildSlot
+		[ 
+			SNullWidget::NullWidget		
+		];
+	}
+	else
+	{
+		this->ChildSlot
 		[
-			SNew(SBox)
-			.MinDesiredHeight(PropertyEditorConstants::PropertyRowHeight)
+			SNew( SBorder )
+			.BorderImage(FAppStyle::Get().GetBrush("DetailsView.GridLine"))
+			.Padding(FMargin(0,0,0,1))
+			.Clipping(EWidgetClipping::ClipToBounds)
 			[
-				SNew( SHorizontalBox )
-				+ SHorizontalBox::Slot()
-				.HAlign(HAlign_Fill)
-				.VAlign(VAlign_Fill)
+				SNew(SBox)
+				.MinDesiredHeight(PropertyEditorConstants::PropertyRowHeight)
 				[
-					SNew( SBorder )
-					.BorderImage(FAppStyle::Get().GetBrush("DetailsView.Highlight"))
-					.Padding_Lambda(GetHighlightBorderPadding)
+					SNew( SHorizontalBox )
+					+ SHorizontalBox::Slot()
+					.HAlign(HAlign_Fill)
+					.VAlign(VAlign_Fill)
 					[
 						SNew( SBorder )
-						.BorderImage(FAppStyle::Get().GetBrush("DetailsView.CategoryMiddle"))
-						.BorderBackgroundColor(this, &SDetailSingleItemRow::GetOuterBackgroundColor)
-						.Padding(0)
+						.BorderImage(FAppStyle::Get().GetBrush("DetailsView.Highlight"))
+						.Padding_Lambda(GetHighlightBorderPadding)
 						[
-							Widget
+							SNew( SBorder )
+							.BorderImage(FAppStyle::Get().GetBrush("DetailsView.CategoryMiddle"))
+							.BorderBackgroundColor(this, &SDetailSingleItemRow::GetOuterBackgroundColor)
+							.Padding(0)
+							[
+								Widget
+							]
 						]
 					]
 				]
 			]
-		]
-	];
+		];
+	}
 
 	STableRow< TSharedPtr< FDetailTreeNode > >::ConstructInternal(
 		STableRow::FArguments()
diff --git a/Engine/Source/Editor/PropertyEditor/Private/SDetailsViewBase.cpp b/Engine/Source/Editor/PropertyEditor/Private/SDetailsViewBase.cpp
--- a/Engine/Source/Editor/PropertyEditor/Private/SDetailsViewBase.cpp
+++ b/Engine/Source/Editor/PropertyEditor/Private/SDetailsViewBase.cpp
@@ -659,6 +659,8 @@ void SDetailsViewBase::UpdateSinglePropertyMap(TSharedPtr<FComplexPropertyNode>
 	Args.bUpdateFavoriteSystemOnly = false;
 	DetailLayoutHelpers::UpdateSinglePropertyMapRecursive(*RootPropertyNode, NAME_None, RootPropertyNode.Get(), Args);
 
+	DetailLayout->AddEmptyCategoryIfNeeded(RootPropertyNode);
+	
 	CustomUpdatePropertyMap(LayoutData.DetailLayout);
 
 	// Ask for custom detail layouts, unless disabled. One reason for disabling custom layouts is that the custom layouts
diff --git a/Engine/Source/Editor/PropertyEditor/Public/DetailCategoryBuilder.h b/Engine/Source/Editor/PropertyEditor/Public/DetailCategoryBuilder.h
--- a/Engine/Source/Editor/PropertyEditor/Public/DetailCategoryBuilder.h
+++ b/Engine/Source/Editor/PropertyEditor/Public/DetailCategoryBuilder.h
@@ -215,4 +215,16 @@ public:
 
 	/** Optional PasteFromText delegate for this category */
 	virtual TSharedPtr<FOnPasteFromText> OnPasteFromText() const { return nullptr; }
+
+	/**
+	 * If true, this Category should have no UProperty data associated with it, and will be shown as an empty stub
+	 * with no expansion arrow
+	 */
+	virtual bool IsEmpty() const { return false; }
+
+	/**
+	 * Sets whether this Category is "Empty" ~ that is, should have no UProperty data associated with it, and will be shown
+	 * as an empty stub with no expansion arrow
+	 */
+	virtual void SetIsEmpty(bool bInIsEmpty) { }
 };
diff --git a/Engine/Source/Editor/PropertyEditor/Public/DetailsDisplayManager.h b/Engine/Source/Editor/PropertyEditor/Public/DetailsDisplayManager.h
--- a/Engine/Source/Editor/PropertyEditor/Public/DetailsDisplayManager.h
+++ b/Engine/Source/Editor/PropertyEditor/Public/DetailsDisplayManager.h
@@ -42,7 +42,7 @@ public:
 	 *
 	 * @param InCategoryObjectName the name of the category object
 	 */
-	void SetCategoryObjectName(FName InCategoryObjectName);
+	void SetCategoryObjectName(FName InCategoryObjectName);            
 
 	/**
 	 * Gets the category menu SWidget and returns a shared pointer to it
@@ -90,6 +90,14 @@ public:
 	 */
 	FMargin GetTablePadding() const;
 
+	/**
+	* Returns true if the specified UObject is a Root Node Object and should show an empty Category stub even if it
+	* has no UProperty Data, else it returns false
+	*
+	* @param InNode the @code UObject* @endcode which will be tested to see if it needs an empty category stub
+	*/
+	virtual bool ShowEmptyCategoryIfRootUObjectHasNoPropertyData(UObject* InNode) const;
+
 protected:
 
 	
