Fixed an issue in detail categories, where the is empty flag was not initialized causing categories to be empty sometimes

#rnx

diff --git a/Engine/Source/Editor/PropertyEditor/Private/DetailCategoryBuilderImpl.cpp b/Engine/Source/Editor/PropertyEditor/Private/DetailCategoryBuilderImpl.cpp
--- a/Engine/Source/Editor/PropertyEditor/Private/DetailCategoryBuilderImpl.cpp
+++ b/Engine/Source/Editor/PropertyEditor/Private/DetailCategoryBuilderImpl.cpp
@@ -201,6 +201,7 @@ FDetailCategoryImpl::FDetailCategoryImpl(FName InCategoryName, TSharedRef<FDetai
 	, bHasVisibleAdvanced(false)
 	, bPendingRefresh(false)
 	, bPendingRefreshNeedsRefilter(false)
+	, bIsEmpty(false)
 {
 	const UStruct* BaseStruct = InDetailLayout->GetRootNode()->GetBaseStructure();
 
