Use string comparison rather than binary comparison to avoid floating point error when propagating default values

diff --git a/Engine/Source/Editor/PropertyEditor/Private/PropertyHandleImpl.cpp b/Engine/Source/Editor/PropertyEditor/Private/PropertyHandleImpl.cpp
--- a/Engine/Source/Editor/PropertyEditor/Private/PropertyHandleImpl.cpp
+++ b/Engine/Source/Editor/PropertyEditor/Private/PropertyHandleImpl.cpp
@@ -320,7 +320,7 @@ FPropertyAccess::Result FPropertyValueImpl::ImportText( const TArray<FObjectBase
 
 			// Cache the value of the property before modifying it.
 			FString PreviousValue;
-			FPropertyTextUtilities::PropertyToTextHelper(PreviousValue, InPropertyNode, NodeProperty, Cur, PPF_None);
+			FPropertyTextUtilities::PropertyToTextHelper(PreviousValue, InPropertyNode, NodeProperty, Cur, PPF_ForDiff);
 
 			// If this property is the inner-property of a container, cache the current value as well
 			FString PreviousContainerValue;
@@ -495,7 +495,7 @@ FPropertyAccess::Result FPropertyValueImpl::ImportText( const TArray<FObjectBase
 
 			// Cache the value of the property after having modified it.
 			FString ValueAfterImport;
-			FPropertyTextUtilities::PropertyToTextHelper(ValueAfterImport, InPropertyNode, NodeProperty, Cur, PPF_None);
+			FPropertyTextUtilities::PropertyToTextHelper(ValueAfterImport, InPropertyNode, NodeProperty, Cur, PPF_ForDiff);
 
 			if (CurObject)
 			{
diff --git a/Engine/Source/Editor/PropertyEditor/Private/PropertyNode.cpp b/Engine/Source/Editor/PropertyEditor/Private/PropertyNode.cpp
--- a/Engine/Source/Editor/PropertyEditor/Private/PropertyNode.cpp
+++ b/Engine/Source/Editor/PropertyEditor/Private/PropertyNode.cpp
@@ -2301,11 +2301,9 @@ bool FPropertyNode::GetDiffersFromDefault(const uint8* PropertyValueAddress, con
 			// try to compare the values at the current and default property addresses
 			if( PropertyValueAddress != nullptr && PropertyDefaultAddress != nullptr )
 			{
-				bDiffersFromDefaultValue = !InProperty->Identical(
-					PropertyValueAddress,
-					PropertyDefaultAddress,
-					PortFlags
-					);
+				FString DefaultValue = GetDefaultValueAsString(PropertyDefaultAddress, InProperty, EValueAsStringMode::ForDiff);
+				FString CurrentValue = GetDefaultValueAsString(PropertyValueAddress, InProperty, EValueAsStringMode::ForDiff);
+				bDiffersFromDefaultValue = !(DefaultValue.Equals(CurrentValue, ESearchCase::CaseSensitive));
 			}
 		}
 	}
@@ -2419,11 +2417,21 @@ bool FPropertyNode::GetDiffersFromDefault()
 }
 
 
-FString FPropertyNode::GetDefaultValueAsString(const uint8* PropertyDefaultAddress, const FProperty* InProperty, const bool bUseDisplayName) const
+FString FPropertyNode::GetDefaultValueAsString(const uint8* PropertyDefaultAddress, const FProperty* InProperty, EValueAsStringMode Mode) const
 {
+	const bool bUseDisplayName = (Mode == EValueAsStringMode::UseDisplayName);
 	FString DefaultValue;
 
-	uint32 PortFlags = bUseDisplayName ? PPF_PropertyWindow : PPF_None;
+	uint32 PortFlags = PPF_None;
+	if (Mode == EValueAsStringMode::UseDisplayName)
+	{
+		PortFlags |= PPF_PropertyWindow;
+	}
+	else if (Mode == EValueAsStringMode::ForDiff)
+	{
+		PortFlags |= PPF_ForDiff;
+	}
+
 	if (InProperty->ContainsInstancedObjectProperty())
 	{
 		PortFlags |= PPF_DeepComparison;
@@ -2456,7 +2464,7 @@ FString FPropertyNode::GetDefaultValueAsString(const uint8* PropertyDefaultAddre
 	return DefaultValue;
 }
 
-FString FPropertyNode::GetDefaultValueAsStringForObject( FPropertyItemValueDataTrackerSlate& ValueTracker, UObject* InObject, FProperty* InProperty, bool bUseDisplayName)
+FString FPropertyNode::GetDefaultValueAsStringForObject( FPropertyItemValueDataTrackerSlate& ValueTracker, UObject* InObject, FProperty* InProperty, EValueAsStringMode Mode)
 {
 	check( InObject );
 	check( InProperty );
@@ -2468,7 +2476,7 @@ FString FPropertyNode::GetDefaultValueAsStringForObject( FPropertyItemValueDataT
 	{
 		if ( ValueTracker.IsValidTracker() && ValueTracker.HasDefaultValue() )
 		{
-			DefaultValue = GetDefaultValueAsString(ValueTracker.GetPropertyDefaultAddress(), InProperty, bUseDisplayName);
+			DefaultValue = GetDefaultValueAsString(ValueTracker.GetPropertyDefaultAddress(), InProperty, Mode);
 		}
 	}
 
@@ -2480,6 +2488,7 @@ FString FPropertyNode::GetDefaultValueAsString(bool bUseDisplayName)
 	FString DefaultValue;
 	FString DelimitedValue;
 	bool bAllSame = true;
+	const EValueAsStringMode Mode = bUseDisplayName ? EValueAsStringMode::UseDisplayName : EValueAsStringMode::None;
 
 	FProperty* Prop = GetProperty();
 	if (!Prop)
@@ -2529,7 +2538,7 @@ FString FPropertyNode::GetDefaultValueAsString(bool bUseDisplayName)
 					PropertyDefaultAddress = PropertyDefaultBaseAddress;
 				}
 
-				NodeDefaultValue = GetDefaultValueAsString(PropertyDefaultAddress, Prop, bUseDisplayName);
+				NodeDefaultValue = GetDefaultValueAsString(PropertyDefaultAddress, Prop, Mode);
 			}
 			
 			if (DefaultValue.IsEmpty())
@@ -2559,7 +2568,7 @@ FString FPropertyNode::GetDefaultValueAsString(bool bUseDisplayName)
 
 			if (Object && ValueTracker.IsValid())
 			{
-				const FString NodeDefaultValue = GetDefaultValueAsStringForObject( *ValueTracker, Object, Prop, bUseDisplayName );
+				const FString NodeDefaultValue = GetDefaultValueAsStringForObject( *ValueTracker, Object, Prop, Mode);
 
 				if (DefaultValue.IsEmpty())
 				{
@@ -3855,28 +3864,12 @@ void FPropertyNode::PropagatePropertyChange( UObject* ModifiedObject, const TCHA
 					ComplexPropertyNode = ParentNodeWeakPtr.Pin();
 				}
 				
-				uint8* DestComplexPropAddr = ComplexPropertyNode->GetValueBaseAddressFromObject(ActualObjToChange);
-				uint8* ModifiedComplexPropAddr = ComplexPropertyNode->GetValueBaseAddressFromObject(ModifiedObject);
+				const uint8* DestComplexPropAddr = ComplexPropertyNode->GetValueBaseAddressFromObject(ActualObjToChange);
 
-				bool bShouldImport = false;
-				{
-					uint8* TempComplexPropAddr = (uint8*)FMemory::Malloc(ComplexProperty->GetSize(), ComplexProperty->GetMinAlignment());
-					ComplexProperty->InitializeValue(TempComplexPropAddr);
-					ON_SCOPE_EXIT
-					{
-						ComplexProperty->DestroyValue(TempComplexPropAddr);
-						FMemory::Free(TempComplexPropAddr);
-					};
-
-					// Importing the previous value into the temporary property can potentially affect shared state (such as FText display string values), so we back-up the current value 
-					// before we do this, so that we can restore it once we've checked whether the two properties are identical
-					// This ensures that shared state keeps the correct value, even if the destination property itself isn't imported (or only partly imported, as is the case with arrays/maps/sets)
-					FString CurrentValue;
-					ComplexProperty->ExportText_Direct(CurrentValue, ModifiedComplexPropAddr, ModifiedComplexPropAddr, ModifiedObject, PPF_None);
-					ComplexProperty->ImportText_Direct(*PreviousValue, TempComplexPropAddr, ModifiedObject, PPF_SerializedAsImportText);
-					bShouldImport = ComplexProperty->Identical(DestComplexPropAddr, TempComplexPropAddr, PPF_DeepComparison);
-					ComplexProperty->ImportText_Direct(*CurrentValue, TempComplexPropAddr, ModifiedObject, PPF_None);
-				}
+				FString ActualCurrentValue;
+				ComplexProperty->ExportText_Direct(ActualCurrentValue, DestComplexPropAddr, DestComplexPropAddr, ActualObjToChange, PPF_ForDiff);
+
+				const bool bShouldImport = ActualCurrentValue.Equals(PreviousValue, ESearchCase::CaseSensitive);
 
 				// Only import if the value matches the previous value of the property that changed
 				if (bShouldImport)
diff --git a/Engine/Source/Editor/PropertyEditor/Private/PropertyNode.h b/Engine/Source/Editor/PropertyEditor/Private/PropertyNode.h
--- a/Engine/Source/Editor/PropertyEditor/Private/PropertyNode.h
+++ b/Engine/Source/Editor/PropertyEditor/Private/PropertyNode.h
@@ -1120,8 +1120,14 @@ protected:
 	bool GetDiffersFromDefault(const uint8* PropertyValueAddress, const uint8* PropertyDefaultAddress, const uint8* DefaultPropertyValueBaseAddress, const FProperty* InProperty) const;
 	bool GetDiffersFromDefaultForObject( FPropertyItemValueDataTrackerSlate& ValueTracker, FProperty* InProperty );
 
-	FString GetDefaultValueAsString(const uint8* PropertyDefaultAddress, const FProperty* InProperty, const bool bUseDisplayName) const;
-	FString GetDefaultValueAsStringForObject( FPropertyItemValueDataTrackerSlate& ValueTracker, UObject* InObject, FProperty* InProperty, bool bUseDisplayName );
+	enum class EValueAsStringMode
+	{
+		None,
+		UseDisplayName,
+		ForDiff,
+	};
+	FString GetDefaultValueAsString(const uint8* PropertyDefaultAddress, const FProperty* InProperty, EValueAsStringMode Mode) const;
+	FString GetDefaultValueAsStringForObject( FPropertyItemValueDataTrackerSlate& ValueTracker, UObject* InObject, FProperty* InProperty, EValueAsStringMode Mode);
 	
 	/**
 	 * Helper function to obtain the display name for an enum property
diff --git a/Engine/Source/Runtime/CoreUObject/Private/UObject/TextProperty.cpp b/Engine/Source/Runtime/CoreUObject/Private/UObject/TextProperty.cpp
--- a/Engine/Source/Runtime/CoreUObject/Private/UObject/TextProperty.cpp
+++ b/Engine/Source/Runtime/CoreUObject/Private/UObject/TextProperty.cpp
@@ -128,6 +128,7 @@ void FTextProperty::SerializeItem(FStructuredArchive::FSlot Slot, void* Value, v
 
 void FTextProperty::ExportText_Internal( FString& ValueStr, const void* PropertyValueOrContainer, EPropertyPointerType PropertyPointerType, const void* DefaultValue, UObject* Parent, int32 PortFlags, UObject* ExportRootScope ) const
 {
+	// get current value:
 	FText TextValue;
 	if (PropertyPointerType == EPropertyPointerType::Container && HasGetter())
 	{
@@ -138,7 +139,8 @@ void FTextProperty::ExportText_Internal( FString& ValueStr, const void* Property
 		TextValue = GetPropertyValue(PointerToValuePtr(PropertyValueOrContainer, PropertyPointerType));
 	}
 
-	if (PortFlags & PPF_PropertyWindow)
+	// export based on flags:
+	auto WriteSimpleTextToBuffer = [&ValueStr, &TextValue, PortFlags]()
 	{
 		if (PortFlags & PPF_Delimited)
 		{
@@ -150,10 +152,34 @@ void FTextProperty::ExportText_Internal( FString& ValueStr, const void* Property
 		{
 			ValueStr += TextValue.ToString();
 		}
+	};
+
+	auto WriteComplexTextToBuffer = [&ValueStr, &TextValue, PortFlags]()
+	{
+		FTextStringHelper::WriteToBuffer(ValueStr, TextValue, !!(PortFlags & PPF_Delimited));
+	};
+
+	if (PortFlags & PPF_ForDiff)
+	{
+		if (TextValue.IsCultureInvariant() || TextValue.IsFromStringTable())
+		{
+			// Invariant and StringTable text values still need to export in their complex text form for diffing to avoid invariant and non-invariant text, 
+			// and different StringTable values that have a different ID but the same display string from being considered identical
+			WriteComplexTextToBuffer();
+		}
+		else
+		{
+			// Any other kind of text should exported as a simple string, to avoid ID changes from instancing text into other packages from being treated as significant
+			WriteSimpleTextToBuffer();
+		}
+	}
+	else if (PortFlags & PPF_PropertyWindow)
+	{
+		WriteSimpleTextToBuffer();
 	}
 	else
 	{
-		FTextStringHelper::WriteToBuffer(ValueStr, TextValue, !!(PortFlags & PPF_Delimited));
+		WriteComplexTextToBuffer();
 	}
 }
 
