Several changes to property panel diff code to make it usable in other situations (like the upcoming property bag fixup tool)
[Bugfix] fixed detail's diff code that assumed non-multiedit mode.
[Feature] Moved code in FDetailsDiffControl::GetLinkedScrollRate to FAsyncDetailViewDiff::GenerateScrollSyncRate so that it can be used without FDetailsDiffControl
[Feature] Made details splitter highlight colors use a callback rather than a map of conflicts so it's more versatile
[Feature] Changed the TreeDiffSpecification to use a virtual class TTreeDiffSpecification instead so that individual methods can be overriden for a specific use case
[Feature] Added callback to details splitter to tell it to ignore a row when highlighting

diff --git a/Engine/Source/Editor/Kismet/Private/DiffControl.cpp b/Engine/Source/Editor/Kismet/Private/DiffControl.cpp
--- a/Engine/Source/Editor/Kismet/Private/DiffControl.cpp
+++ b/Engine/Source/Editor/Kismet/Private/DiffControl.cpp
@@ -466,39 +466,7 @@ TAttribute<TArray<FVector2f>> FDetailsDiffControl::GetLinkedScrollRateAttribute(
 TArray<FVector2f> FDetailsDiffControl::GetLinkedScrollRate(TSharedRef<IDetailsView> LeftDetailsView, TSharedRef<IDetailsView> RightDetailsView) const
 {
 	const UObject* LeftObject = LeftDetailsView->GetSelectedObjects()[0].Get();
-	
-	TArray<FIntVector2> MatchingRows;
-
-	// iterate matching rows of both details panels simultaneously
-	auto [LeftRowCount, RightRowCount] = PropertyTreeDifferences[LeftObject].Right->ForEachRow(
-		[&MatchingRows](const TUniquePtr<FAsyncDetailViewDiff::DiffNodeType>& DiffNode, int32 LeftRow, int32 RightRow)->ETreeTraverseControl
-		{
-			// if both trees share this row, sync scrolling here
-			if (DiffNode->ValueA.IsValid() && DiffNode->ValueB.IsValid())
-			{
-				if (MatchingRows.IsEmpty() || (MatchingRows.Last().X != LeftRow && MatchingRows.Last().Y != RightRow))
-				{
-					MatchingRows.Emplace(LeftRow, RightRow);
-				}
-			}
-			return ETreeTraverseControl::Continue;
-		}
-	);
-
-	TArray<FVector2f> FixedPoints;
-	FixedPoints.Emplace(0.f, 0.f);
-	
-	// normalize fixed points
-	for (FIntVector2& MatchingRow : MatchingRows)
-	{
-		FixedPoints.Emplace(
-			StaticCast<float>(MatchingRow.X) / LeftRowCount,
-			StaticCast<float>(MatchingRow.Y) / RightRowCount
-		);
-	}
-	FixedPoints.Emplace(1.f, 1.f);
-	
-	return FixedPoints;
+	return PropertyTreeDifferences[LeftObject].Right->GenerateScrollSyncRate();
 }
 
 
diff --git a/Engine/Source/Editor/Kismet/Private/SDetailsDiff.cpp b/Engine/Source/Editor/Kismet/Private/SDetailsDiff.cpp
--- a/Engine/Source/Editor/Kismet/Private/SDetailsDiff.cpp
+++ b/Engine/Source/Editor/Kismet/Private/SDetailsDiff.cpp
@@ -474,8 +474,8 @@ SDetailsDiff::FDiffControl SDetailsDiff::GenerateDetailsPanel(const TFunction<co
 					}
 				}
 			}
-			
-			Splitter->AddHighlights(Highlights);
+
+			Splitter->HighlightFromMergeResults(MergeConflicts);
 		}
 	});
 	if (OutputObject)
diff --git a/Engine/Source/Editor/PropertyEditor/Internal/AsyncDetailViewDiff.h b/Engine/Source/Editor/PropertyEditor/Internal/AsyncDetailViewDiff.h
--- a/Engine/Source/Editor/PropertyEditor/Internal/AsyncDetailViewDiff.h
+++ b/Engine/Source/Editor/PropertyEditor/Internal/AsyncDetailViewDiff.h
@@ -4,20 +4,40 @@
 #include "DiffUtils.h"
 #include "IDetailsView.h"
 
-
-namespace TreeDiffSpecification
+template<>
+class PROPERTYEDITOR_API TTreeDiffSpecification<TWeakPtr<FDetailTreeNode>>
 {
-	template<>
-	bool PROPERTYEDITOR_API AreValuesEqual(const TWeakPtr<FDetailTreeNode>& TreeNodeA, const TWeakPtr<FDetailTreeNode>& TreeNodeB);
-	
-	template<>
-	bool PROPERTYEDITOR_API AreMatching(const TWeakPtr<FDetailTreeNode>& TreeNodeA, const TWeakPtr<FDetailTreeNode>& TreeNodeB);
-	
-	template<>
-	void PROPERTYEDITOR_API GetChildren(const TWeakPtr<FDetailTreeNode>& InParent, TArray<TWeakPtr<FDetailTreeNode>>& OutChildren);
-
-	template<>
-	bool PROPERTYEDITOR_API ShouldMatchByValue(const TWeakPtr<FDetailTreeNode>& TreeNodeA);
+public:
+	virtual ~TTreeDiffSpecification() = default;
+	/**
+	 * determine whether the values stored in two nodes are equal.
+	 * @param TreeNodeA node from the first user provided tree (guaranteed not to be null)
+	 * @param TreeNodeB node from the second user provided tree (guaranteed not to be null)
+	 */
+	virtual bool AreValuesEqual(const TWeakPtr<FDetailTreeNode>& TreeNodeA, const TWeakPtr<FDetailTreeNode>& TreeNodeB) const;
+
+	/**
+	 * determine whether two nodes occupy the same space in their trees
+	 * for example if you have a tree key/value pairs, AreMatching should compare the keys while AreValuesEqual
+	 * should compare the values
+	 * @param TreeNodeA node from the first user provided tree (guaranteed not to be null)
+	 * @param TreeNodeB node from the second user provided tree (guaranteed not to be null)
+	 */
+	virtual bool AreMatching(const TWeakPtr<FDetailTreeNode>& TreeNodeA, const TWeakPtr<FDetailTreeNode>& TreeNodeB) const;
+
+	/**
+	 * retrieves an array of children nodes from the parent node
+	 * @param InParent node from one of the two user provided trees (guaranteed not to be null)
+	 * @param[out] OutChildren to be filled with the children of parent
+	 */
+	virtual void GetChildren(const TWeakPtr<FDetailTreeNode>& InParent, TArray<TWeakPtr<FDetailTreeNode>>& OutChildren) const;
+
+	/**
+	 * return true for nodes that match using AreValuesEqual first, and pair up by position second
+	 * this is useful for arrays since we often want to keep elements with the same data paired while diffing other elements in order
+	 * @param TreeNode node from one of the two user provided trees (guaranteed not to be null)
+	*/
+	virtual bool ShouldMatchByValue(const TWeakPtr<FDetailTreeNode>& TreeNode) const;
 };
 
 // asynchronously update a difference tree with changes to details views
@@ -36,6 +56,9 @@ public:
 	// returns the number of rows in each details view
 	TPair<int32, int32> ForEachRow(const TFunction<ETreeTraverseControl(const TUniquePtr<DiffNodeType>&, int32, int32)>& Method) const;
 
+	// generates scroll data used by LinkableScrollbar
+	TArray<FVector2f> GenerateScrollSyncRate() const;
+
 private:
 	static TAttribute<TArray<TWeakPtr<FDetailTreeNode>>> RootNodesAttribute(TWeakPtr<IDetailsView> DetailsView);
 	TWeakPtr<IDetailsView> LeftView;
diff --git a/Engine/Source/Editor/PropertyEditor/Private/AsyncDetailViewDiff.cpp b/Engine/Source/Editor/PropertyEditor/Private/AsyncDetailViewDiff.cpp
--- a/Engine/Source/Editor/PropertyEditor/Private/AsyncDetailViewDiff.cpp
+++ b/Engine/Source/Editor/PropertyEditor/Private/AsyncDetailViewDiff.cpp
@@ -4,14 +4,16 @@
 #include "DetailTreeNode.h"
 #include "IDetailsViewPrivate.h"
 #include "DiffUtils.h"
+#include "PropertyHandleImpl.h"
 
-static const UObject* GetObject(const TSharedPtr<FDetailTreeNode>& TreeNode)
+static TArray<TWeakObjectPtr<UObject>> GetObjects(const TSharedPtr<FDetailTreeNode>& TreeNode)
 {
+	TArray<UObject*> Result;
 	if (const IDetailsViewPrivate* DetailsView = TreeNode->GetDetailsView())
 	{
-		return DetailsView->GetSelectedObjects()[0].Get();
+		return DetailsView->GetSelectedObjects();
 	}
-	return nullptr;
+	return {};
 }
 
 static FResolvedProperty GetResolvedProperty(const TSharedPtr<FPropertyNode>& PropertyNode, const UObject* Object)
@@ -27,8 +29,8 @@ static FResolvedProperty GetResolvedProperty(const TSharedPtr<FPropertyNode>& Pr
 	return FResolvedProperty();
 }
 
-template<>
-bool TreeDiffSpecification::AreValuesEqual<TWeakPtr<FDetailTreeNode>>(const TWeakPtr<FDetailTreeNode>& TreeNodeA,  const TWeakPtr<FDetailTreeNode>& TreeNodeB)
+
+bool TTreeDiffSpecification<TWeakPtr<FDetailTreeNode>>::AreValuesEqual(const TWeakPtr<FDetailTreeNode>& TreeNodeA, const TWeakPtr<FDetailTreeNode>& TreeNodeB) const
 {
 	const TSharedPtr<FDetailTreeNode> PinnedTreeNodeA = TreeNodeA.Pin();
 	const TSharedPtr<FDetailTreeNode> PinnedTreeNodeB = TreeNodeB.Pin();
@@ -37,44 +39,38 @@ bool TreeDiffSpecification::AreValuesEqual<TWeakPtr<FDetailTreeNode>>(const TWea
 		return PinnedTreeNodeA == PinnedTreeNodeB;
 	}
 
-	TArray<TSharedRef<FPropertyNode>> PropertyNodesA;
-	TArray<TSharedRef<FPropertyNode>> PropertyNodesB;
-	PinnedTreeNodeA->GetAllPropertyNodes(PropertyNodesA);
-	PinnedTreeNodeB->GetAllPropertyNodes(PropertyNodesB);
-	
-	ensure(PropertyNodesA.Num() == PropertyNodesB.Num()); // AreMatching(...) should've stopped this from happening
-	if (PropertyNodesA.IsEmpty() || PropertyNodesB.IsEmpty())
+	const TSharedPtr<IPropertyHandle> PropertyHandleA = PinnedTreeNodeA->CreatePropertyHandle();
+	const TSharedPtr<IPropertyHandle> PropertyHandleB = PinnedTreeNodeB->CreatePropertyHandle();
+	if (!PropertyHandleA || !PropertyHandleB)
 	{
 		// category nodes
 		return PinnedTreeNodeA->GetNodeName() == PinnedTreeNodeB->GetNodeName();
 	}
-
-	const UObject* OwningObjectA = GetObject(PinnedTreeNodeA);
-	const UObject* OwningObjectB = GetObject(PinnedTreeNodeB);
-
-	for (int32 PropNodeIndex = 0; PropNodeIndex < PropertyNodesA.Num(); ++PropNodeIndex)
+	const TSharedPtr<FPropertyNode>& PropertyNodeA = StaticCastSharedPtr<FPropertyHandleBase>(PropertyHandleA)->GetPropertyNode();
+	const TSharedPtr<FPropertyNode>& PropertyNodeB = StaticCastSharedPtr<FPropertyHandleBase>(PropertyHandleB)->GetPropertyNode();
+	if (!PropertyNodeA || !PropertyNodeB)
 	{
-		const TSharedRef<FPropertyNode>& PropertyNodeA = PropertyNodesA[PropNodeIndex];
-		const TSharedRef<FPropertyNode>& PropertyNodeB = PropertyNodesB[PropNodeIndex];
+		return PinnedTreeNodeA->GetNodeName() == PinnedTreeNodeB->GetNodeName();
+	}
 
-		const FResolvedProperty ResolvedA = GetResolvedProperty(PropertyNodeA, OwningObjectA);
-		const FResolvedProperty ResolvedB = GetResolvedProperty(PropertyNodeB, OwningObjectB);
-		if (ResolvedA.Property && ResolvedB.Property)
-		{
-			if (!ResolvedA.Property->SameType(ResolvedB.Property))
-			{
-				return false;
-			}
+	TArray<void*> DataValuesA;
+	TArray<void*> DataValuesB;
+	PropertyHandleA->AccessRawData(DataValuesA);
+	PropertyHandleB->AccessRawData(DataValuesB);
+	
+	if(DataValuesA.IsEmpty() || DataValuesB.IsEmpty())
+	{
+		return true;
+	}
+	
+	const TArray<TWeakObjectPtr<UObject>> OwningObjectsA = GetObjects(PinnedTreeNodeA);
+	const TArray<TWeakObjectPtr<UObject>> OwningObjectsB = GetObjects(PinnedTreeNodeB);
 
-			// use DiffUtils::Identical instead of FProperty::Identical so that sub-object pointers are considered identical
-			// if they're isomorphic rather than only if they have the same pointer values
-			if(!DiffUtils::Identical(ResolvedA, ResolvedB, OwningObjectA->GetPackage(), OwningObjectB->GetPackage()))
-			{
-				return false;
-			}
-		}
+	if (!ensure(OwningObjectsA.Num() == DataValuesA.Num() && OwningObjectsB.Num() == DataValuesB.Num()))
+	{
+		return true;
 	}
-	return true;
+	return DiffUtils::Identical(PropertyHandleA, PropertyHandleB, OwningObjectsA, OwningObjectsB);
 }
 
 static bool MapKeysMatch(const TSharedRef<FPropertyNode>& MapPropertyNodeA, const TSharedRef<FPropertyNode>& MapPropertyNodeB, int32 KeyIndexA, int32 KeyIndexB,
@@ -127,8 +123,7 @@ static bool SetKeysMatch(const TSharedRef<FPropertyNode>& SetPropertyNodeA, cons
 	return false;
 }
 
-template<>
-bool TreeDiffSpecification::AreMatching<TWeakPtr<FDetailTreeNode>>(const TWeakPtr<FDetailTreeNode>& TreeNodeA, const TWeakPtr<FDetailTreeNode>& TreeNodeB)
+bool TTreeDiffSpecification<TWeakPtr<FDetailTreeNode>>::AreMatching(const TWeakPtr<FDetailTreeNode>& TreeNodeA, const TWeakPtr<FDetailTreeNode>& TreeNodeB) const
 {
 	const TSharedPtr<FDetailTreeNode> PinnedTreeNodeA = TreeNodeA.Pin();
 	const TSharedPtr<FDetailTreeNode> PinnedTreeNodeB = TreeNodeB.Pin();
@@ -151,9 +146,18 @@ bool TreeDiffSpecification::AreMatching<TWeakPtr<FDetailTreeNode>>(const TWeakPt
 		// category nodes
 		return PinnedTreeNodeA->GetNodeName() == PinnedTreeNodeB->GetNodeName();
 	}
+
+	const TArray<TWeakObjectPtr<UObject>> OwningObjectsA = GetObjects(PinnedTreeNodeA);
+	const TArray<TWeakObjectPtr<UObject>> OwningObjectsB = GetObjects(PinnedTreeNodeB);
+	if (OwningObjectsA.IsEmpty() || OwningObjectsB.IsEmpty())
+	{
+		return PinnedTreeNodeA->GetNodeName() == PinnedTreeNodeB->GetNodeName();
+	}
 	
-	const UObject* OwningObjectA = GetObject(PinnedTreeNodeA);
-	const UObject* OwningObjectB = GetObject(PinnedTreeNodeB);
+	// because multi-edit puts matching keys together, we can assume all objects share the same key for this property
+	// and just look at the first instance
+	const UObject* OwningObjectB = OwningObjectsB[0].Get();
+	const UObject* OwningObjectA = OwningObjectsA[0].Get();
 	
 	for (int32 PropNodeIndex = 0; PropNodeIndex < PropertyNodesA.Num(); ++PropNodeIndex)
 	{
@@ -206,8 +210,7 @@ bool TreeDiffSpecification::AreMatching<TWeakPtr<FDetailTreeNode>>(const TWeakPt
 	return true;
 }
 
-template<>
-void TreeDiffSpecification::GetChildren<TWeakPtr<FDetailTreeNode>>(const TWeakPtr<FDetailTreeNode>& InParent, TArray<TWeakPtr<FDetailTreeNode>>& OutChildren)
+void TTreeDiffSpecification<TWeakPtr<FDetailTreeNode>>::GetChildren(const TWeakPtr<FDetailTreeNode>& InParent, TArray<TWeakPtr<FDetailTreeNode>>& OutChildren) const
 {
 	const TSharedPtr<FDetailTreeNode> PinnedParent = InParent.Pin();
 	if (PinnedParent)
@@ -221,8 +224,7 @@ void TreeDiffSpecification::GetChildren<TWeakPtr<FDetailTreeNode>>(const TWeakPt
 	}
 }
 
-template<>
-bool TreeDiffSpecification::ShouldMatchByValue<TWeakPtr<FDetailTreeNode>>(const TWeakPtr<FDetailTreeNode>& TreeNode)
+bool TTreeDiffSpecification<TWeakPtr<FDetailTreeNode>>::ShouldMatchByValue(const TWeakPtr<FDetailTreeNode>& TreeNode) const
 {
 	const TSharedPtr<FDetailTreeNode> PinnedTreeNode = TreeNode.Pin();
 	if (!PinnedTreeNode)
@@ -345,6 +347,42 @@ TPair<int32, int32> FAsyncDetailViewDiff::ForEachRow(const TFunction<ETreeTraver
 	return {LeftRowNum, RightRowNum};
 }
 
+TArray<FVector2f> FAsyncDetailViewDiff::GenerateScrollSyncRate() const
+{
+	TArray<FIntVector2> MatchingRows;
+
+	// iterate matching rows of both details panels simultaneously
+	auto [LeftRowCount, RightRowCount] = ForEachRow(
+		[&MatchingRows](const TUniquePtr<DiffNodeType>& DiffNode, int32 LeftRow, int32 RightRow)->ETreeTraverseControl
+		{
+			// if both trees share this row, sync scrolling here
+			if (DiffNode->ValueA.IsValid() && DiffNode->ValueB.IsValid())
+			{
+				if (MatchingRows.IsEmpty() || (MatchingRows.Last().X != LeftRow && MatchingRows.Last().Y != RightRow))
+				{
+					MatchingRows.Emplace(LeftRow, RightRow);
+				}
+			}
+			return ETreeTraverseControl::Continue;
+		}
+	);
+
+	TArray<FVector2f> FixedPoints;
+	FixedPoints.Emplace(0.f, 0.f);
+	
+	// normalize fixed points
+	for (FIntVector2& MatchingRow : MatchingRows)
+	{
+		FixedPoints.Emplace(
+			StaticCast<float>(MatchingRow.X) / LeftRowCount,
+			StaticCast<float>(MatchingRow.Y) / RightRowCount
+		);
+	}
+	FixedPoints.Emplace(1.f, 1.f);
+	
+	return FixedPoints;
+}
+
 TAttribute<TArray<TWeakPtr<FDetailTreeNode>>> FAsyncDetailViewDiff::RootNodesAttribute(TWeakPtr<IDetailsView> DetailsView)
 {
 	return TAttribute<TArray<TWeakPtr<FDetailTreeNode>>>::CreateLambda([DetailsView]()
diff --git a/Engine/Source/Editor/PropertyEditor/Private/SDetailsSplitter.cpp b/Engine/Source/Editor/PropertyEditor/Private/SDetailsSplitter.cpp
--- a/Engine/Source/Editor/PropertyEditor/Private/SDetailsSplitter.cpp
+++ b/Engine/Source/Editor/PropertyEditor/Private/SDetailsSplitter.cpp
@@ -8,6 +8,15 @@
 #include "Styling/StyleColors.h"
 #include "Framework/Application/SlateApplication.h"
 
+static const TArray<TWeakObjectPtr<UObject>>* GetObjects(const TSharedPtr<FDetailTreeNode>& TreeNode)
+{
+	if (const IDetailsViewPrivate* DetailsView = TreeNode->GetDetailsView())
+	{
+		return &DetailsView->GetSelectedObjects();
+	}
+	return nullptr;
+}
+
 // converts from a container helper like FScriptArrayHelper to a property like FArrayProperty
 template<typename HelperType>
 using TContainerPropertyType =
@@ -17,32 +26,48 @@ using TContainerPropertyType =
 	void>>>;
 
 template<typename HelperType>
-static bool TryGetSourceContainer(TSharedPtr<FDetailTreeNode> DetailsNode, FPropertyNode*& OutPropNode, TUniquePtr<HelperType>& OutResolvedProperty)
+static bool TryGetSourceContainer(TSharedPtr<FDetailTreeNode> DetailsNode, FPropertyNode*& OutPropNode, TArray<TUniquePtr<HelperType>>& OutContainerHelper)
 {
 	using ContainerPropertyType = TContainerPropertyType<HelperType>;
 	if (TryGetSourceContainer<ContainerPropertyType>(DetailsNode, OutPropNode))
 	{
 		const FPropertySoftPath SoftPropertyPath(FPropertyNode::CreatePropertyPath(OutPropNode->AsShared()).Get());
-		const UObject* Object = DetailsNode->GetDetailsView()->GetSelectedObjects()[0].Get();
-		const FResolvedProperty Resolved = SoftPropertyPath.Resolve(Object);
-		const ContainerPropertyType* ContainerProperty = CastFieldChecked<ContainerPropertyType>(Resolved.Property);
-		OutResolvedProperty = MakeUnique<HelperType>(ContainerProperty, ContainerProperty->template ContainerPtrToValuePtr<UObject*>(Resolved.Object));
+		if (const TArray<TWeakObjectPtr<UObject>>* Objects = GetObjects(DetailsNode))
+		{
+			for (const TWeakObjectPtr<UObject> WeakObject : *Objects)
+			{
+				if (const UObject* Object = WeakObject.Get())
+				{
+					const FResolvedProperty Resolved = SoftPropertyPath.Resolve(Object);
+					const ContainerPropertyType* ContainerProperty = CastFieldChecked<ContainerPropertyType>(Resolved.Property);
+					OutContainerHelper.Add(MakeUnique<HelperType>(ContainerProperty, ContainerProperty->template ContainerPtrToValuePtr<UObject*>(Resolved.Object)));
+				}
+			}
+		}
 		return true;
 	}
 	return false;
 }
 
 template<typename HelperType>
-static bool TryGetDestinationContainer(TSharedPtr<FDetailTreeNode> DetailsNode, FPropertyNode*& OutPropNode, TUniquePtr<HelperType>& OutContainerHelper, int32& OutInsertIndex)
+static bool TryGetDestinationContainer(TSharedPtr<FDetailTreeNode> DetailsNode, FPropertyNode*& OutPropNode, TArray<TUniquePtr<HelperType>>& OutContainerHelper, int32& OutInsertIndex)
 {
 	using ContainerPropertyType = TContainerPropertyType<HelperType>;
 	if (TryGetDestinationContainer<ContainerPropertyType>(DetailsNode, OutPropNode, OutInsertIndex))
 	{
 		const FPropertySoftPath SoftPropertyPath(FPropertyNode::CreatePropertyPath(OutPropNode->AsShared()).Get());
-		const UObject* Object = DetailsNode->GetDetailsView()->GetSelectedObjects()[0].Get();
-		const FResolvedProperty Resolved = SoftPropertyPath.Resolve(Object);
-		const ContainerPropertyType* ContainerProperty = CastFieldChecked<ContainerPropertyType>(Resolved.Property);
-		OutContainerHelper = MakeUnique<HelperType>(ContainerProperty, ContainerProperty->template ContainerPtrToValuePtr<UObject*>(Resolved.Object));
+		if (const TArray<TWeakObjectPtr<UObject>>* Objects = GetObjects(DetailsNode))
+		{
+			for (const TWeakObjectPtr<UObject> WeakObject : *Objects)
+			{
+				if (const UObject* Object = WeakObject.Get())
+				{
+					const FResolvedProperty Resolved = SoftPropertyPath.Resolve(Object);
+					const ContainerPropertyType* ContainerProperty = CastFieldChecked<ContainerPropertyType>(Resolved.Property);
+					OutContainerHelper.Add(MakeUnique<HelperType>(ContainerProperty, ContainerProperty->template ContainerPtrToValuePtr<UObject*>(Resolved.Object)));
+				}
+			}
+		}
 		return true;
 	}
 	return false;
@@ -123,66 +148,84 @@ static void SendPropertyChangeNotification(FPropertyNode* PropertyNode, UObject*
 
 static void CopyPropertyValueForInsert(const TSharedPtr<FDetailTreeNode>& SourceDetailsNode, const TSharedPtr<FDetailTreeNode>& DestinationDetailsNode)
 {
-	TUniquePtr<FScriptArrayHelper> SourceArray;
+	const TSharedPtr<IPropertyHandle> DestinationHandle = DestinationDetailsNode->CreatePropertyHandle();
+
+	// Array
+	TArray<TUniquePtr<FScriptArrayHelper>> SourceArrays;
 	FPropertyNode* SourceArrayPropertyNode;
-	if (TryGetSourceContainer(SourceDetailsNode, SourceArrayPropertyNode, SourceArray))
+	if (TryGetSourceContainer(SourceDetailsNode, SourceArrayPropertyNode, SourceArrays))
 	{
 		int32 InsertIndex;
-		TUniquePtr<FScriptArrayHelper> DestinationArray;
+		TArray<TUniquePtr<FScriptArrayHelper>> DestinationArrays;
 		FPropertyNode* DestinationArrayPropertyNode;
-		if (ensure(TryGetDestinationContainer(DestinationDetailsNode, DestinationArrayPropertyNode, DestinationArray, InsertIndex)))
+		if (ensure(TryGetDestinationContainer(DestinationDetailsNode, DestinationArrayPropertyNode, DestinationArrays, InsertIndex)))
 		{
-			DestinationArray->InsertValues(InsertIndex, 1);
-			const void* SourceData = SourceArray->GetElementPtr(SourceDetailsNode->GetPropertyNode()->GetArrayIndex());
-			void* DestinationData = DestinationArray->GetElementPtr(InsertIndex);
-			const FArrayProperty* ArrayProperty = CastFieldChecked<FArrayProperty>(DestinationArrayPropertyNode->GetProperty());
-			const FProperty* ElementProperty = ArrayProperty->Inner;
-			ElementProperty->CopySingleValue(DestinationData, SourceData);
-			SendPropertyChangeNotification(
-				DestinationArrayPropertyNode,
-				DestinationDetailsNode->GetDetailsView()->GetSelectedObjects()[0].Get(),
-				EPropertyChangeType::ArrayAdd);
+			ensure(SourceArrays.Num() == DestinationArrays.Num());
+			DestinationHandle->NotifyPreChange();
+			for (int32 ArrayNum = 0; ArrayNum < SourceArrays.Num(); ++ ArrayNum)
+			{
+				DestinationArrays[ArrayNum]->InsertValues(InsertIndex, 1);
+				const void* SourceData = SourceArrays[ArrayNum]->GetElementPtr(SourceDetailsNode->GetPropertyNode()->GetArrayIndex());
+				void* DestinationData = DestinationArrays[ArrayNum]->GetElementPtr(InsertIndex);
+				const FArrayProperty* ArrayProperty = CastFieldChecked<FArrayProperty>(DestinationArrayPropertyNode->GetProperty());
+				const FProperty* ElementProperty = ArrayProperty->Inner;
+				ElementProperty->CopySingleValue(DestinationData, SourceData);
+			}
+			
+			DestinationHandle->NotifyPostChange(EPropertyChangeType::ArrayAdd);
+			DestinationHandle->NotifyFinishedChangingProperties();
 		}
+		return;
 	}
 	
-	TUniquePtr<FScriptSetHelper> SourceSet;
+	// Set
+	TArray<TUniquePtr<FScriptSetHelper>> SourceSets;
 	FPropertyNode* SourceSetPropertyNode;
-	if (TryGetSourceContainer(SourceDetailsNode, SourceSetPropertyNode, SourceSet))
+	if (TryGetSourceContainer(SourceDetailsNode, SourceSetPropertyNode, SourceSets))
 	{
 		int32 InsertIndex;
-		TUniquePtr<FScriptSetHelper> DestinationSet;
+		TArray<TUniquePtr<FScriptSetHelper>> DestinationSets;
 		FPropertyNode* DestinationPropertyNode;
-		if (ensure(TryGetDestinationContainer(DestinationDetailsNode, DestinationPropertyNode, DestinationSet, InsertIndex)))
+		if (ensure(TryGetDestinationContainer(DestinationDetailsNode, DestinationPropertyNode, DestinationSets, InsertIndex)))
 		{
-			const void* SourceData = SourceSet->FindNthElementPtr(SourceDetailsNode->GetPropertyNode()->GetArrayIndex());
-			DestinationSet->AddElement(SourceData);
+			ensure(SourceSets.Num() == DestinationSets.Num());
+			DestinationHandle->NotifyPreChange();
+			for (int32 SetNum = 0; SetNum < SourceSets.Num(); ++ SetNum)
+			{
+				const void* SourceData = SourceSets[SetNum]->FindNthElementPtr(SourceDetailsNode->GetPropertyNode()->GetArrayIndex());
+				DestinationSets[SetNum]->AddElement(SourceData);
+			}
 			
-			SendPropertyChangeNotification(
-				DestinationPropertyNode,
-				DestinationDetailsNode->GetDetailsView()->GetSelectedObjects()[0].Get(),
-				EPropertyChangeType::ArrayAdd);
+			DestinationHandle->NotifyPostChange(EPropertyChangeType::ArrayAdd);
+			DestinationHandle->NotifyFinishedChangingProperties();
 		}
+		return;
 	}
 	
-	TUniquePtr<FScriptMapHelper> SourceMap;
+	// Map
+	TArray<TUniquePtr<FScriptMapHelper>> SourceMaps;
 	FPropertyNode* SourceMapPropertyNode;
-	if (TryGetSourceContainer(SourceDetailsNode, SourceMapPropertyNode, SourceMap))
+	if (TryGetSourceContainer(SourceDetailsNode, SourceMapPropertyNode, SourceMaps))
 	{
 		int32 InsertIndex;
-		TUniquePtr<FScriptMapHelper> DestinationMap;
+		TArray<TUniquePtr<FScriptMapHelper>> DestinationMaps;
 		FPropertyNode* DestinationPropertyNode;
-		if (ensure(TryGetDestinationContainer(DestinationDetailsNode, DestinationPropertyNode, DestinationMap, InsertIndex)))
+		if (ensure(TryGetDestinationContainer(DestinationDetailsNode, DestinationPropertyNode, DestinationMaps, InsertIndex)))
 		{
-			const int32 Index = SourceMap->FindInternalIndex(SourceDetailsNode->GetPropertyNode()->GetArrayIndex());
-			const void* SourceKey = SourceMap->GetKeyPtr(Index);
-			const void* SourceVal = SourceMap->GetValuePtr(Index);
-			DestinationMap->AddPair(SourceKey, SourceVal);
+			ensure(SourceMaps.Num() == DestinationMaps.Num());
+			DestinationHandle->NotifyPreChange();
+			for (int32 MapNum = 0; MapNum < SourceMaps.Num(); ++ MapNum)
+			{
+				const int32 Index = SourceMaps[MapNum]->FindInternalIndex(SourceDetailsNode->GetPropertyNode()->GetArrayIndex());
+				const void* SourceKey = SourceMaps[MapNum]->GetKeyPtr(Index);
+				const void* SourceVal = SourceMaps[MapNum]->GetValuePtr(Index);
+				DestinationMaps[MapNum]->AddPair(SourceKey, SourceVal);
+			}
 			
-			SendPropertyChangeNotification(
-				DestinationPropertyNode,
-				DestinationDetailsNode->GetDetailsView()->GetSelectedObjects()[0].Get(),
-				EPropertyChangeType::ArrayAdd);
+			DestinationHandle->NotifyPostChange(EPropertyChangeType::ArrayAdd);
+			DestinationHandle->NotifyFinishedChangingProperties();
 		}
+		return;
 	}
 }
 
@@ -206,47 +249,45 @@ static void CopyPropertyValue(const TSharedPtr<FDetailTreeNode>& SourceDetailsNo
 		default:
 			return;
 	}
-	
-	FResolvedProperty SourceResolved;
-	FResolvedProperty DestinationResolved;
-	bool bCopySingleValue = false;
-	
-	if (SourceDetailsNode)
+
+	TArray<void*> SourceData;
+	TArray<void*> DestinationData;
+	const TSharedPtr<IPropertyHandle> SourceHandle = SourceDetailsNode->CreatePropertyHandle();
+	const TSharedPtr<IPropertyHandle> DestinationHandle = DestinationDetailsNode->CreatePropertyHandle();
+	if (!ensure(SourceHandle && DestinationHandle))
 	{
-		const UObject* Object = SourceDetailsNode->GetDetailsView()->GetSelectedObjects()[0].Get();
-		const FPropertyPath PropertyPath = SourceDetailsNode->GetPropertyPath();
-		SourceResolved = FPropertySoftPath(PropertyPath).Resolve(Object);
-		bCopySingleValue = PropertyPath.GetLeafMostProperty().ArrayIndex != INDEX_NONE;
+		return;
 	}
-	if (DestinationDetailsNode)
+	if (!ensure(SourceHandle->GetProperty()->SameType(DestinationHandle->GetProperty())))
 	{
-		const UObject* Object = DestinationDetailsNode->GetDetailsView()->GetSelectedObjects()[0].Get();
-		const FPropertyPath PropertyPath = DestinationDetailsNode->GetPropertyPath();
-		DestinationResolved = FPropertySoftPath(PropertyPath).Resolve(Object);
-		bCopySingleValue = PropertyPath.GetLeafMostProperty().ArrayIndex != INDEX_NONE;
+		return;
 	}
-		
-	if (SourceResolved.Property && DestinationResolved.Property && SourceResolved.Object && DestinationResolved.Object)
+	SourceHandle->AccessRawData(SourceData);
+	DestinationHandle->AccessRawData(DestinationData);
+	if (!ensure(SourceData.Num() == DestinationData.Num()))
 	{
-		const void* SourceData = SourceResolved.Property->ContainerPtrToValuePtr<void*>(SourceResolved.Object);
-		void* DestinationData = (void*)DestinationResolved.Property->ContainerPtrToValuePtr<void*>(DestinationResolved.Object);
-		if (bCopySingleValue)
+		return;
+	}
+
+	
+	DestinationHandle->NotifyPreChange();
+	for (int32 I = 0; I < SourceData.Num(); ++I)
+	{
+		if (DestinationHandle->GetArrayIndex() != INDEX_NONE)
 		{
-			DestinationResolved.Property->CopySingleValue(DestinationData, SourceData);
+			DestinationHandle->GetProperty()->CopySingleValue(DestinationData[I], SourceData[I]);
 		}
 		else
 		{
-			DestinationResolved.Property->CopyCompleteValue(DestinationData, SourceData);
+			DestinationHandle->GetProperty()->CopyCompleteValue(DestinationData[I], SourceData[I]);
 		}
-		SendPropertyChangeNotification(
-			DestinationDetailsNode->GetPropertyNode().Get(),
-			DestinationDetailsNode->GetDetailsView()->GetSelectedObjects()[0].Get(),
-			EPropertyChangeType::ValueSet);
 	}
+	DestinationHandle->NotifyPostChange(EPropertyChangeType::ValueSet);
+	DestinationHandle->NotifyFinishedChangingProperties();
 }
 
 // note: DestinationDetailsNode is the node before the position in the tree you wish to insert
-static bool CanCopyPropertyValueForInsert(TSharedPtr<FDetailTreeNode> SourceDetailsNode, TSharedPtr<FDetailTreeNode> DestinationDetailsNode)
+static bool CanCopyPropertyValueForInsert(const TSharedPtr<FDetailTreeNode>& SourceDetailsNode, const TSharedPtr<FDetailTreeNode>& DestinationDetailsNode)
 {
 	FPropertyNode* SourceArrayPropertyNode;
 	if (TryGetSourceContainer<FArrayProperty>(SourceDetailsNode, SourceArrayPropertyNode))
@@ -288,9 +329,17 @@ static bool CanCopyPropertyValueForInsert(TSharedPtr<FDetailTreeNode> SourceDeta
 	return false;
 }
 
-static bool CanCopyPropertyValue(TSharedPtr<FDetailTreeNode> SourceDetailsNode, TSharedPtr<FDetailTreeNode> DestinationDetailsNode, ETreeDiffResult Diff)
+static bool CanCopyPropertyValue(const TSharedPtr<FDetailTreeNode>& SourceDetailsNode, const TSharedPtr<FDetailTreeNode>& DestinationDetailsNode, ETreeDiffResult Diff)
 {
-	switch(Diff)
+	// in order to copy properties there needs to be the same number of objects in each panel
+	const TArray<TWeakObjectPtr<UObject>>* SourceObjects = GetObjects(SourceDetailsNode);
+	const TArray<TWeakObjectPtr<UObject>>* DestinationObjects = GetObjects(DestinationDetailsNode);
+	if (!SourceObjects || !DestinationObjects || SourceObjects->Num() != DestinationObjects->Num())
+	{
+		return false;
+	}
+	
+	switch (Diff)
 	{
 	// traditional copy
 	case ETreeDiffResult::DifferentValues:
@@ -300,7 +349,7 @@ static bool CanCopyPropertyValue(TSharedPtr<FDetailTreeNode> SourceDetailsNode,
 	case ETreeDiffResult::MissingFromTree1:
 	case ETreeDiffResult::MissingFromTree2:
 		return CanCopyPropertyValueForInsert(SourceDetailsNode, DestinationDetailsNode);
-		
+
 	// no difference
 	case ETreeDiffResult::Invalid:
 	case ETreeDiffResult::Identical:
@@ -312,6 +361,10 @@ static bool CanCopyPropertyValue(TSharedPtr<FDetailTreeNode> SourceDetailsNode,
 void SDetailsSplitter::Construct(const FArguments& InArgs)
 {
 	Splitter = SNew(SSplitter).PhysicalSplitterHandleSize(5.f);
+	
+	GetRowHighlightColor = InArgs._RowHighlightColor.IsBound() ? InArgs._RowHighlightColor : FRowHighlightColor::CreateStatic(
+		[](const TUniquePtr<FAsyncDetailViewDiff::DiffNodeType>&) {return FLinearColor(0.f, 1.f, 1.f, .7f);});
+	
 	for(const FSlot::FSlotArguments& SlotArgs : InArgs._Slots)
 	{
 		AddSlot(SlotArgs);
@@ -342,6 +395,7 @@ void SDetailsSplitter::AddSlot(const FSlot::FSlotArguments& SlotArgs, int32 Inde
 		SlotArgs._DetailsView,
 		SlotArgs._IsReadonly,
 		SlotArgs._DifferencesWithRightPanel,
+		SlotArgs._ShouldIgnoreRow.IsBound() ? SlotArgs._ShouldIgnoreRow : FShouldIgnoreRow::CreateStatic([](const TWeakPtr<FDetailTreeNode>&){return false;})
 	}, Index);
 }
 
@@ -350,9 +404,73 @@ SDetailsSplitter::FPanel& SDetailsSplitter::GetPanel(int32 Index)
 	return Panels[Index];
 }
 
-void SDetailsSplitter::AddHighlights(const TMap<FString, TMap<FPropertySoftPath, FLinearColor>>& Highlights)
+// find the inner-most object in the property path and shorten the path relative to that object
+static void ShortenToPathFromLastObject(const UObject*& InOutObject, FPropertyPath& InOutPath)
 {
-	CustomHighlights = Highlights;
+	TArray<FPropertyInfo> PathFromSubObjectReversed;
+	PathFromSubObjectReversed.Add(InOutPath.GetLeafMostProperty());
+	InOutPath = *InOutPath.TrimPath(1);
+	while (InOutPath.IsValid())
+	{
+		// Note that if we have perf issues, this could be made faster by writing a custom resolve function that stops
+		// at the last FObjectProperty. that way we wouldn't need to resolve multiple times
+		const FResolvedProperty Resolved = FPropertySoftPath(InOutPath).Resolve(InOutObject);
+		if (const FObjectProperty* ObjectProperty = CastField<FObjectProperty>(Resolved.Property))
+		{
+			InOutObject = ObjectProperty->GetObjectPropertyValue(ObjectProperty->ContainerPtrToValuePtr<void>(Resolved.Object));
+			break;
+		}
+		PathFromSubObjectReversed.Add(InOutPath.GetLeafMostProperty());
+		InOutPath = *InOutPath.TrimPath(1);
+	}
+
+	InOutPath = {};
+	while(!PathFromSubObjectReversed.IsEmpty())
+	{
+		InOutPath.AddProperty(PathFromSubObjectReversed.Pop());
+	}
+}
+
+void SDetailsSplitter::HighlightFromMergeResults(const TMap<FString, TMap<FPropertySoftPath, ETreeDiffResult>>& CustomHighlights)
+{
+	GetRowHighlightColor = FRowHighlightColor::CreateLambda([CustomHighlights](const TUniquePtr<FAsyncDetailViewDiff::DiffNodeType>& DiffNode)
+	{
+		const TSharedPtr<FDetailTreeNode> DetailNode = DiffNode->ValueA.IsValid() ? DiffNode->ValueA.Pin() : DiffNode->ValueB.Pin();
+		
+		FPropertyPath Path = DetailNode->GetPropertyPath();
+		const UObject* OwningObject = DetailNode->GetDetailsView()->GetSelectedObjects()[0].Get();
+		if (OwningObject && Path.IsValid())
+		{
+			FPropertyPath PathFromSubObject;
+			ShortenToPathFromLastObject(OwningObject, Path);
+			if (OwningObject)
+			{
+				if (const TMap<FPropertySoftPath, ETreeDiffResult>* Highlights = CustomHighlights.Find(OwningObject->GetPathName(OwningObject->GetPackage())))
+				{
+					if (const ETreeDiffResult* DiffResult = Highlights->Find(FPropertySoftPath(Path)))
+					{
+						switch(*DiffResult)
+						{
+						case ETreeDiffResult::MissingFromTree1: // fall through
+						case ETreeDiffResult::MissingFromTree2: // fall through
+						case ETreeDiffResult::DifferentValues:
+							// color is intentionally using values greater than 1 so that it stays very saturated
+							return FLinearColor(1.5f, 0.3f, 0.3f);
+						
+						default:; // ignore identical and invalid
+						}
+					}
+				}
+			}
+		}
+		
+		return FLinearColor(0.f, 1.f, 1.f, .7f);
+	});
+}
+
+void SDetailsSplitter::SetRowHighlightColorDelegate(const FRowHighlightColor& Delegate)
+{
+	GetRowHighlightColor = Delegate;
 }
 
 SDetailsSplitter::FSlot::FSlotArguments SDetailsSplitter::Slot()
@@ -387,7 +505,7 @@ int32 SDetailsSplitter::OnPaint(const FPaintArgs& Args, const FGeometry& Allotte
 			
 			Diff->ForEachRow([&](const TUniquePtr<FAsyncDetailViewDiff::DiffNodeType>& DiffNode, int32, int32)->ETreeTraverseControl
 			{
-				const FLinearColor Color = GetHighlightColor(DiffNode);
+				const FLinearColor Color = GetRowHighlightColor.Execute(DiffNode);
 				
 				FSlateRect LeftPropertyRect;
 				if (const TSharedPtr<FDetailTreeNode> LeftDetailNode = DiffNode->ValueA.Pin())
@@ -404,7 +522,10 @@ int32 SDetailsSplitter::OnPaint(const FPaintArgs& Args, const FGeometry& Allotte
 				}
 				if (LeftPropertyRect.IsValid() && DiffNode->DiffResult != ETreeDiffResult::Identical)
 				{
-					RowHighlights.Add(LeftPropertyRect, Color);
+					if (!LeftPanel.ShouldIgnoreRow.Execute(DiffNode->ValueA))
+					{
+						RowHighlights.Add(LeftPropertyRect, Color);
+					}
 				}
 		
 				const int32 RightIndex = LeftIndex + 1;
@@ -426,27 +547,33 @@ int32 SDetailsSplitter::OnPaint(const FPaintArgs& Args, const FGeometry& Allotte
 					}
 					if (RightPropertyRect.IsValid() && DiffNode->DiffResult != ETreeDiffResult::Identical)
 					{
-						RowHighlights.Add(RightPropertyRect, Color);
+						if (!RightPanel.ShouldIgnoreRow.Execute(DiffNode->ValueB))
+						{
+							RowHighlights.Add(RightPropertyRect, Color);
+						}
 					}
 					
 					if (LeftPropertyRect.IsValid() && RightPropertyRect.IsValid() && DiffNode->DiffResult != ETreeDiffResult::Identical)
 					{
-						FLinearColor FillColor = Color.Desaturate(.3f) * FLinearColor(0.053f,0.053f,0.053f);
-						FillColor.A = 0.43f;
-			
-						const FLinearColor OutlineColor = Color;
-						PaintPropertyConnector(OutDrawElements, MaxLayerId, LeftPropertyRect, RightPropertyRect, FillColor, OutlineColor);
-						++MaxLayerId;
-
-						if (!RightPanel.IsReadonly.Get(true) && CanCopyPropertyValue(LastSeenLeftDetailsNode, LastSeenRightDetailsNode, DiffNode->DiffResult))
-						{
-							PaintCopyPropertyButton(OutDrawElements, MaxLayerId, DiffNode, LeftPropertyRect, RightPropertyRect, EPropertyCopyDirection::CopyLeftToRight);
-						}
-						if (!LeftPanel.IsReadonly.Get(true) && CanCopyPropertyValue(LastSeenRightDetailsNode, LastSeenLeftDetailsNode, DiffNode->DiffResult))
+						if (!LeftPanel.ShouldIgnoreRow.Execute(DiffNode->ValueA) && !RightPanel.ShouldIgnoreRow.Execute(DiffNode->ValueB))
 						{
-							PaintCopyPropertyButton(OutDrawElements, MaxLayerId, DiffNode, LeftPropertyRect, RightPropertyRect, EPropertyCopyDirection::CopyRightToLeft);
+							FLinearColor FillColor = Color.Desaturate(.3f) * FLinearColor(0.053f,0.053f,0.053f);
+							FillColor.A = 0.43f;
+			
+							const FLinearColor OutlineColor = Color;
+							PaintPropertyConnector(OutDrawElements, MaxLayerId, LeftPropertyRect, RightPropertyRect, FillColor, OutlineColor);
+							++MaxLayerId;
+
+							if (!RightPanel.IsReadonly.Get(true) && CanCopyPropertyValue(LastSeenLeftDetailsNode, LastSeenRightDetailsNode, DiffNode->DiffResult))
+							{
+								PaintCopyPropertyButton(OutDrawElements, MaxLayerId, DiffNode, LeftPropertyRect, RightPropertyRect, EPropertyCopyDirection::CopyLeftToRight);
+							}
+							if (!LeftPanel.IsReadonly.Get(true) && CanCopyPropertyValue(LastSeenRightDetailsNode, LastSeenLeftDetailsNode, DiffNode->DiffResult))
+							{
+								PaintCopyPropertyButton(OutDrawElements, MaxLayerId, DiffNode, LeftPropertyRect, RightPropertyRect, EPropertyCopyDirection::CopyRightToLeft);
+							}
+							++MaxLayerId;
 						}
-						++MaxLayerId;
 						
 					}
 				}
@@ -820,55 +947,3 @@ void SDetailsSplitter::PaintCopyPropertyButton(FSlateWindowElementList& OutDrawE
 		ButtonColor
 	);
 }
-
-// find the inner-most object in the property path and shorten the path relative to that object
-static void ShortenToPathFromLastObject(const UObject*& InOutObject, FPropertyPath& InOutPath)
-{
-	TArray<FPropertyInfo> PathFromSubObjectReversed;
-	PathFromSubObjectReversed.Add(InOutPath.GetLeafMostProperty());
-	InOutPath = *InOutPath.TrimPath(1);
-	while (InOutPath.IsValid())
-	{
-		// Note that if we have perf issues, this could be made faster by writing a custom resolve function that stops
-		// at the last FObjectProperty. that way we wouldn't need to resolve multiple times
-		const FResolvedProperty Resolved = FPropertySoftPath(InOutPath).Resolve(InOutObject);
-		if (const FObjectProperty* ObjectProperty = CastField<FObjectProperty>(Resolved.Property))
-		{
-			InOutObject = ObjectProperty->GetObjectPropertyValue(ObjectProperty->ContainerPtrToValuePtr<void>(Resolved.Object));
-			break;
-		}
-		PathFromSubObjectReversed.Add(InOutPath.GetLeafMostProperty());
-		InOutPath = *InOutPath.TrimPath(1);
-	}
-
-	InOutPath = {};
-	while(!PathFromSubObjectReversed.IsEmpty())
-	{
-		InOutPath.AddProperty(PathFromSubObjectReversed.Pop());
-	}
-}
-
-FLinearColor SDetailsSplitter::GetHighlightColor(const TUniquePtr<TDiffNode<TWeakPtr<FDetailTreeNode>>>& DiffNode) const
-{
-	const TSharedPtr<FDetailTreeNode> DetailNode = DiffNode->ValueA.IsValid() ? DiffNode->ValueA.Pin() : DiffNode->ValueB.Pin();
-	
-	FPropertyPath Path = DetailNode->GetPropertyPath();
-	const UObject* OwningObject = DetailNode->GetDetailsView()->GetSelectedObjects()[0].Get();
-	if (OwningObject && Path.IsValid())
-	{
-		FPropertyPath PathFromSubObject;
-		ShortenToPathFromLastObject(OwningObject, Path);
-		if (OwningObject)
-		{
-			if (const TMap<FPropertySoftPath, FLinearColor>* Highlights = CustomHighlights.Find(OwningObject->GetPathName(OwningObject->GetPackage())))
-            {
-            	if (const FLinearColor* Highlight = Highlights->Find(FPropertySoftPath(Path)))
-            	{
-            		return *Highlight;
-            	}
-            }
-		}
-	}
-	
-	return FLinearColor(0.f, 1.f, 1.f, .7f);
-}
diff --git a/Engine/Source/Editor/PropertyEditor/Public/SDetailsSplitter.h b/Engine/Source/Editor/PropertyEditor/Public/SDetailsSplitter.h
--- a/Engine/Source/Editor/PropertyEditor/Public/SDetailsSplitter.h
+++ b/Engine/Source/Editor/PropertyEditor/Public/SDetailsSplitter.h
@@ -9,17 +9,22 @@
 class PROPERTYEDITOR_API SDetailsSplitter : public SCompoundWidget
 {
 public:
+	// Callback that's called before highlighting a row. If true, the row is skipped.
+	DECLARE_DELEGATE_RetVal_OneParam(bool, FShouldIgnoreRow, const TWeakPtr<FDetailTreeNode>&)
+	DECLARE_DELEGATE_RetVal_OneParam(FLinearColor, FRowHighlightColor, const TUniquePtr<FAsyncDetailViewDiff::DiffNodeType>&)
+
 	class FSlot : public TSlotBase<FSlot>
 	{
 	public:
 		FSlot() = default;
-
+		
 		SLATE_SLOT_BEGIN_ARGS(FSlot, TSlotBase<FSlot>)
 			/** When the RuleSize is set to FractionOfParent, the size of the slot is the Value percentage of its parent size. */
 			SLATE_ATTRIBUTE(float, Value)
 			SLATE_ARGUMENT(TSharedPtr<IDetailsView>, DetailsView)
 			SLATE_ATTRIBUTE(bool, IsReadonly) // default true
 			SLATE_ATTRIBUTE(TSharedPtr<FAsyncDetailViewDiff>, DifferencesWithRightPanel)
+			SLATE_EVENT(FShouldIgnoreRow, ShouldIgnoreRow) // default false
 		SLATE_SLOT_END_ARGS()
 	};
 	struct FPanel
@@ -27,11 +32,13 @@ public:
 		TSharedPtr<IDetailsView> DetailsView;
 		TAttribute<bool> IsReadonly;
 		TAttribute<TSharedPtr<FAsyncDetailViewDiff>> DiffRight;
+		FShouldIgnoreRow ShouldIgnoreRow;
 	};
 	static FSlot::FSlotArguments Slot();
 	
 	SLATE_BEGIN_ARGS(SDetailsSplitter)
 	{}
+		SLATE_EVENT(FRowHighlightColor, RowHighlightColor) // default cyan: FLinearColor(0.f, 1.f, 1.f, .7f)
 		SLATE_SLOT_ARGUMENT(FSlot, Slots)
 	SLATE_END_ARGS()
 
@@ -48,7 +55,9 @@ public:
 	void AddSlot(const FSlot::FSlotArguments& SlotArgs, int32 Index = INDEX_NONE);
 	FPanel& GetPanel(int32 Index);
 
-	void AddHighlights(const TMap<FString, TMap<FPropertySoftPath, FLinearColor>>& Highlights);
+	void SetRowHighlightColorDelegate(const FRowHighlightColor& Delegate);
+	// generate a highlight color delegate that utilizes the merge results to determine the color
+	void HighlightFromMergeResults(const TMap<FString, TMap<FPropertySoftPath, ETreeDiffResult>>& Highlights);
 
 private:
 	
@@ -70,13 +79,11 @@ private:
 	void PaintPropertyConnector(FSlateWindowElementList& OutDrawElements, int32 LayerId, const FSlateRect& LeftPropertyRect, 
 		const FSlateRect& RightPropertyRect, const FLinearColor& FillColor, const FLinearColor& OutlineColor) const;
 	
-	void PaintCopyPropertyButton(FSlateWindowElementList& OutDrawElements, int32 LayerId, const TUniquePtr<TDiffNode<TWeakPtr<FDetailTreeNode>>>& DiffNode,
+	void PaintCopyPropertyButton(FSlateWindowElementList& OutDrawElements, int32 LayerId, const TUniquePtr<FAsyncDetailViewDiff::DiffNodeType>& DiffNode,
 		const FSlateRect& LeftPropertyRect, const FSlateRect& RightPropertyRect, EPropertyCopyDirection CopyDirection) const;
-
-	FLinearColor GetHighlightColor(const TUniquePtr<TDiffNode<TWeakPtr<FDetailTreeNode>>>& DiffNode) const;
 	
 	TSharedPtr<SSplitter> Splitter;
 	TArray<FPanel> Panels;
 	FCopyPropertyButton HoveredCopyButton;
-	TMap<FString, TMap<FPropertySoftPath, FLinearColor>> CustomHighlights;
+	FRowHighlightColor GetRowHighlightColor;
 };
diff --git a/Engine/Source/Editor/UnrealEd/Internal/AsyncTreeDifferences.h b/Engine/Source/Editor/UnrealEd/Internal/AsyncTreeDifferences.h
--- a/Engine/Source/Editor/UnrealEd/Internal/AsyncTreeDifferences.h
+++ b/Engine/Source/Editor/UnrealEd/Internal/AsyncTreeDifferences.h
@@ -69,6 +69,53 @@ namespace TreeDiffSpecification
 	template<typename InNodeType>
 	bool ShouldMatchByValue(const InNodeType& TreeNode);
 
+}
+
+template<typename InNodeType>
+class TTreeDiffSpecification
+{
+public:
+	virtual ~TTreeDiffSpecification() = default;
+	
+	/**
+	 * determine whether the values stored in two nodes are equal.
+	 * @param TreeNodeA node from the first user provided tree (guaranteed not to be null)
+	 * @param TreeNodeB node from the second user provided tree (guaranteed not to be null)
+	 */
+	virtual bool AreValuesEqual(const InNodeType& TreeNodeA, const InNodeType& TreeNodeB) const
+	{
+		return TreeNodeA == TreeNodeB;
+	}
+
+	/**
+	 * determine whether two nodes occupy the same space in their trees
+	 * for example if you have a tree key/value pairs, AreMatching should compare the keys while AreValuesEqual
+	 * should compare the values
+	 * @param TreeNodeA node from the first user provided tree (guaranteed not to be null)
+	 * @param TreeNodeB node from the second user provided tree (guaranteed not to be null)
+	 */
+	virtual bool AreMatching(const InNodeType& TreeNodeA, const InNodeType& TreeNodeB) const
+	{
+		return TreeNodeA == TreeNodeB;
+	}
+
+	/**
+	 * retrieves an array of children nodes from the parent node
+	 * @param InParent node from one of the two user provided trees (guaranteed not to be null)
+	 * @param[out] OutChildren to be filled with the children of parent
+	 */
+	virtual void GetChildren(const InNodeType& InParent, TArray<InNodeType>& OutChildren) const
+	{}
+
+	/**
+	 * return true for nodes that match using AreValuesEqual first, and pair up by position second
+	 * this is useful for arrays since we often want to keep elements with the same data paired while diffing other elements in order
+	 * @param TreeNode node from one of the two user provided trees (guaranteed not to be null)
+	*/
+	virtual bool ShouldMatchByValue(const InNodeType& TreeNode) const
+	{
+		return true;
+	}
 };
 
 
@@ -92,13 +139,13 @@ struct TDiffNode
 	TDiffNode() = default;
 	TDiffNode(TDiffNode &&Other) = default;
 	
-	TDiffNode(const ValueType& InValueA, const ValueType& InValueB, const TDiffNode* InParent)
+	TDiffNode(const ValueType& InValueA, const ValueType& InValueB, const TDiffNode* InParent, TUniquePtr<TTreeDiffSpecification<ValueType>>& Specification)
 		: ValueA(InValueA), ValueB(InValueB), Children(), Parent(InParent)
 	{
-		SetDiffType();
+		SetDiffType(Specification);
 	}
 
-	void SetDiffType();
+	void SetDiffType(TUniquePtr<TTreeDiffSpecification<ValueType>>& Specification);
 
 	bool operator==(const TDiffNode& Other) const;
 	
@@ -123,8 +170,15 @@ public:
 	// calls Method(DiffNode) on each diff node in the tree.
 	void ForEach(ETreeTraverseOrder TraversalOrder, const TFunction<ETreeTraverseControl(const TUniquePtr<DiffNodeType>&)>& Method) const;
 
+	template<typename SpecificationType, typename... TArgs>
+	void SetDiffSpecification(TArgs... Args);
+
 	// head of the main diff tree (note: not for user modification. Tick() will overwrite user changes)
 	TUniquePtr<DiffNodeType> Head;
+
+	// every time the update queue is completed this is called (max of once per Tick)
+	DECLARE_MULTICAST_DELEGATE_OneParam(FOnQueueFlushed, const TAsyncTreeDifferences&);
+	FOnQueueFlushed OnQueueFlushed;
 private:
 
 	// process one element from the top of the queue
@@ -132,10 +186,10 @@ private:
 	
 	void QueueParallelNodeLists(const TArray<ValueType>& ValuesA, const TArray<ValueType>& ValuesB, DiffNodeType* ParentNode);
 	
-	static TArray<TArray<int32>> CalculateLCSTableForMatchingValues(const TArray<ValueType>& ValuesA, const TArray<ValueType>& ValuesB);
+	TArray<TArray<int32>> CalculateLCSTableForMatchingValues(const TArray<ValueType>& ValuesA, const TArray<ValueType>& ValuesB);
 	static TArray<TArray<int32>> CalculateLCSTableForDiffNodes(const TArray<DiffNodeType>& FoundNodes, const TArray<TUniquePtr<DiffNodeType>>& ExpectedNodes);
 
-	static bool AreMatching(const ValueType& ValueA, const ValueType& ValueB);
+	bool AreMatching(const ValueType& ValueA, const ValueType& ValueB);
 	
 	static bool PreOrderRecursive(const TUniquePtr<DiffNodeType>& Node, const TFunction<ETreeTraverseControl(const TUniquePtr<DiffNodeType>&)>& Method);
 	static bool PostOrderRecursive(const TUniquePtr<DiffNodeType>& Node, const TFunction<ETreeTraverseControl(const TUniquePtr<DiffNodeType>&)>& Method);
@@ -145,14 +199,15 @@ private:
 	TArray<DiffNodeType*> UpdateQueue;
 	TAttribute<TArray<ValueType>> RootValuesA;
 	TAttribute<TArray<ValueType>> RootValuesB;
+	TUniquePtr<TTreeDiffSpecification<ValueType>> Specification;
 };
 
 template <typename InNodeType>
-void TDiffNode<InNodeType>::SetDiffType()
+void TDiffNode<InNodeType>::SetDiffType(TUniquePtr<TTreeDiffSpecification<ValueType>>& Specification)
 {
 	if (ValueA != NullValue && ValueB != NullValue)
 	{
-		if (TreeDiffSpecification::AreValuesEqual<ValueType>(ValueA, ValueB))
+		if (Specification->AreValuesEqual(ValueA, ValueB))
 		{
 			DiffResult = ETreeDiffResult::Identical;
 			return;
@@ -186,6 +241,7 @@ TAsyncTreeDifferences<InNodeType>::TAsyncTreeDifferences(const TAttribute<TArray
 	: Head(MakeUnique<DiffNodeType>())
 	, RootValuesA(InRootValuesA)
 	, RootValuesB(InRootValuesB)
+	, Specification(MakeUnique<TTreeDiffSpecification<ValueType>>())
 {
 	QueueParallelNodeLists(RootValuesA.Get(), RootValuesB.Get(), Head.Get());
 }
@@ -203,6 +259,11 @@ void TAsyncTreeDifferences<InNodeType>::Tick(float MaxAllottedTimeMs)
 	{
 		ProcessTopOfQueue();
 	}
+	
+	if (UpdateQueue.IsEmpty())
+	{
+		OnQueueFlushed.Broadcast(*this);
+	}
 }
 
 template <typename InNodeType>
@@ -212,6 +273,7 @@ void TAsyncTreeDifferences<InNodeType>::FlushQueue()
 	{
 		ProcessTopOfQueue();
 	}
+	OnQueueFlushed.Broadcast(*this);
 }
 
 template <typename InNodeType>
@@ -235,6 +297,13 @@ void TAsyncTreeDifferences<InNodeType>::ForEach(ETreeTraverseOrder TraversalOrde
 	}
 }
 
+template <typename InNodeType>
+template <typename SpecificationType, typename... TArgs>
+void TAsyncTreeDifferences<InNodeType>::SetDiffSpecification(TArgs... Args)
+{
+	Specification = MakeUnique<SpecificationType>(Args...);
+}
+
 template <typename InNodeType>
 bool TAsyncTreeDifferences<InNodeType>::PreOrderRecursive(const TUniquePtr<DiffNodeType>& Node, const TFunction<ETreeTraverseControl(const TUniquePtr<DiffNodeType>&)>& Method)
 {
@@ -336,11 +405,11 @@ void TAsyncTreeDifferences<InNodeType>::ProcessTopOfQueue()
 	TArray<ValueType> ChildrenB;
 	if (DiffNode->ValueA != DiffNodeType::NullValue)
 	{
-		TreeDiffSpecification::GetChildren<ValueType>(DiffNode->ValueA, ChildrenA);
+		Specification->GetChildren(DiffNode->ValueA, ChildrenA);
 	}
 	if (DiffNode->ValueB != DiffNodeType::NullValue)
 	{
-		TreeDiffSpecification::GetChildren<ValueType>(DiffNode->ValueB, ChildrenB);
+		Specification->GetChildren(DiffNode->ValueB, ChildrenB);
 	}
 	
 	if (!ChildrenA.IsEmpty() || !ChildrenB.IsEmpty())
@@ -373,27 +442,27 @@ void TAsyncTreeDifferences<InNodeType>::QueueParallelNodeLists(const TArray<Valu
         	if (IndexA == 0)
         	{
         		// ValueB doesn't match. push it on it's own
-        		FoundChildren.Add(DiffNodeType(DiffNodeType::NullValue, ValuesB[--IndexB], ParentNode));
+        		FoundChildren.Add(DiffNodeType(DiffNodeType::NullValue, ValuesB[--IndexB], ParentNode, Specification));
         	}
         	else if (IndexB == 0)
         	{
         		// ValueA doesn't match. push it on it's own
-        		FoundChildren.Add(DiffNodeType(ValuesA[--IndexA], DiffNodeType::NullValue, ParentNode));
+        		FoundChildren.Add(DiffNodeType(ValuesA[--IndexA], DiffNodeType::NullValue, ParentNode, Specification));
         	}
         	else if (AreMatching(ValuesA[IndexA - 1], ValuesB[IndexB - 1]))
         	{
         		// found a match between both nodes
-        		FoundChildren.Add(DiffNodeType(ValuesA[--IndexA], ValuesB[--IndexB], ParentNode));
+        		FoundChildren.Add(DiffNodeType(ValuesA[--IndexA], ValuesB[--IndexB], ParentNode, Specification));
         	}
         	else if (LCS[IndexA - 1][IndexB] <= LCS[IndexA][IndexB - 1])
         	{
         		// ValueB doesn't match. push it on it's own
-        		FoundChildren.Add(DiffNodeType(DiffNodeType::NullValue, ValuesB[--IndexB], ParentNode));
+        		FoundChildren.Add(DiffNodeType(DiffNodeType::NullValue, ValuesB[--IndexB], ParentNode, Specification));
         	}
         	else
         	{
         		// ValueA doesn't match. push it on it's own
-        		FoundChildren.Add(DiffNodeType(ValuesA[--IndexA], DiffNodeType::NullValue, ParentNode));
+        		FoundChildren.Add(DiffNodeType(ValuesA[--IndexA], DiffNodeType::NullValue, ParentNode, Specification));
         	}
         }
 
@@ -411,7 +480,7 @@ void TAsyncTreeDifferences<InNodeType>::QueueParallelNodeLists(const TArray<Valu
 	for (DiffNodeType& FoundChild : AllFoundChildren)
 	{
 		if (FoundChild.DiffResult == ETreeDiffResult::MissingFromTree1 &&
-			TreeDiffSpecification::ShouldMatchByValue<ValueType>(FoundChild.ValueB))
+			Specification->ShouldMatchByValue(FoundChild.ValueB))
 		{
 			if (CompressCount == 0)
 			{
@@ -423,10 +492,10 @@ void TAsyncTreeDifferences<InNodeType>::QueueParallelNodeLists(const TArray<Valu
 		}
 		if (CompressCount > 0 &&
             FoundChild.DiffResult == ETreeDiffResult::MissingFromTree2 &&
-			TreeDiffSpecification::ShouldMatchByValue<ValueType>(FoundChild.ValueA))
+			Specification->ShouldMatchByValue(FoundChild.ValueA))
 		{
 			FoundChildren[CompressIndex].ValueA = FoundChild.ValueA;
-			FoundChildren[CompressIndex].SetDiffType();
+			FoundChildren[CompressIndex].SetDiffType(Specification);
 			++CompressIndex;
 			--CompressCount;
 			continue;
@@ -541,15 +610,15 @@ TArray<TArray<int32>> TAsyncTreeDifferences<InNodeType>::CalculateLCSTableForDif
 template <typename InNodeType>
 bool TAsyncTreeDifferences<InNodeType>::AreMatching(const ValueType& ValueA, const ValueType& ValueB)
 {
-	const bool bMatchValueAByValue = TreeDiffSpecification::ShouldMatchByValue<ValueType>(ValueA);
-	const bool bMatchValueBByValue = TreeDiffSpecification::ShouldMatchByValue<ValueType>(ValueA);
+	const bool bMatchValueAByValue = Specification->ShouldMatchByValue(ValueA);
+	const bool bMatchValueBByValue = Specification->ShouldMatchByValue(ValueB);
 	if (bMatchValueAByValue && bMatchValueBByValue)
 	{
-		return TreeDiffSpecification::AreValuesEqual<ValueType>(ValueA, ValueB);
+		return Specification->AreValuesEqual(ValueA, ValueB);
 	}
 	if (!bMatchValueAByValue && !bMatchValueBByValue)
 	{
-		return TreeDiffSpecification::AreMatching<ValueType>(ValueA, ValueB);
+		return Specification->AreMatching(ValueA, ValueB);
 	}
 
 	// a node that should match by value will never match a node that shouldn't match by value
diff --git a/Engine/Source/Editor/UnrealEd/Private/DiffUtils.cpp b/Engine/Source/Editor/UnrealEd/Private/DiffUtils.cpp
--- a/Engine/Source/Editor/UnrealEd/Private/DiffUtils.cpp
+++ b/Engine/Source/Editor/UnrealEd/Private/DiffUtils.cpp
@@ -1074,6 +1074,96 @@ bool DiffUtils::Identical(const FResolvedProperty& AProp, const FResolvedPropert
 	return DifferingProperties.Num() == 0;
 }
 
+bool DiffUtils::Identical(const TSharedPtr<IPropertyHandle>& PropertyHandleA, const TSharedPtr<IPropertyHandle>& PropertyHandleB,
+                          const TArray<TWeakObjectPtr<UObject>>& OwningOutersA, const TArray<TWeakObjectPtr<UObject>>& OwningOutersB)
+{
+	TArray<void*> ValuesA;
+	TArray<void*> ValuesB;
+	PropertyHandleA->AccessRawData(ValuesA);
+	PropertyHandleB->AccessRawData(ValuesB);
+
+	// if OwningOuters weren't provided, fallback to using the property handles to find them
+	TArray<UObject*> HandleOutersA;
+	TArray<UObject*> HandleOutersB;
+	if (OwningOutersA.IsEmpty())
+	{
+		PropertyHandleA->GetOuterObjects(HandleOutersA);
+	}
+	if (OwningOutersB.IsEmpty())
+	{
+		PropertyHandleB->GetOuterObjects(HandleOutersB);
+	}
+	
+	if (!ensure(ValuesA.Num() == OwningOutersA.Num()))
+	{
+		// Outer count mismatch
+		return false;
+	}
+	if (!ensure(ValuesB.Num() == OwningOutersB.Num()))
+	{
+		// Outer count mismatch
+		return false;
+	}
+
+	auto IsIdenticalAtIndex = [&](int32 IndexA, int32 IndexB)
+	{
+		const void* ValueA = ValuesA[IndexA];
+		const void* ValueB = ValuesB[IndexB];
+
+		const UObject* OwningOuterA = OwningOutersA.IsEmpty() ? HandleOutersA[IndexA] : OwningOutersA[IndexA].Get();
+		const UObject* OwningOuterB = OwningOutersB.IsEmpty() ? HandleOutersB[IndexB] : OwningOutersB[IndexB].Get();;
+
+		// note that we're not directly calling FProperty::Identical because sub-object properties should be weakly compared based on
+		// their paths instead of their pointers or data
+		TArray<FPropertySoftPath> DifferingProperties;
+		IdenticalHelper(PropertyHandleA->GetProperty(), PropertyHandleB->GetProperty(), ValueA, ValueB,
+			OwningOuterA, OwningOuterB, {}, DifferingProperties, true);
+
+		return DifferingProperties.IsEmpty();
+	};
+	
+	if (ValuesA.Num() == ValuesB.Num())
+	{
+		// compare AValues[I] with BValues[I]
+		for (int32 I = 0; I < ValuesA.Num(); ++I)
+		{
+			if (!IsIdenticalAtIndex(I,I))
+			{
+				return false;
+			}
+		}
+	}
+	else if (ValuesA.Num() == 1)
+	{
+		// compare AValues[0] with BValues[0...N]
+		for (int32 I = 0; I < ValuesB.Num(); ++I)
+		{
+			if (!IsIdenticalAtIndex(0,I))
+			{
+				return false;
+			}
+		}
+	}
+	else if (ValuesB.Num() == 1)
+	{
+		// compare BValues[0] with AValues[0...N]
+		for (int32 I = 0; I < ValuesA.Num(); ++I)
+		{
+			if (!IsIdenticalAtIndex(I,0))
+			{
+				return false;
+			}
+		}
+	}
+	else
+	{
+		// number of values doesn't match... this cannot be compared
+		return ensure(false);
+	}
+	
+	return true;
+}
+
 TArray<FPropertySoftPath> DiffUtils::GetVisiblePropertiesInOrderDeclared(const UStruct* ForStruct, const FPropertySoftPath& Scope /*= TArray<FName>()*/)
 {
 	TArray<FPropertySoftPath> Ret;
diff --git a/Engine/Source/Editor/UnrealEd/Private/MergeUtils.cpp b/Engine/Source/Editor/UnrealEd/Private/MergeUtils.cpp
--- a/Engine/Source/Editor/UnrealEd/Private/MergeUtils.cpp
+++ b/Engine/Source/Editor/UnrealEd/Private/MergeUtils.cpp
@@ -384,16 +384,18 @@ struct FPropertyInstance
 };
 
 // methods that make FPropertyInstance diffable
-namespace TreeDiffSpecification
+template<>
+class TTreeDiffSpecification<FPropertyInstance>
 {
-	template<>
-	bool AreValuesEqual<FPropertyInstance>(const FPropertyInstance& TreeNodeA, const FPropertyInstance& TreeNodeB)
+public:
+	virtual ~TTreeDiffSpecification() = default;
+	
+	virtual bool AreValuesEqual(const FPropertyInstance& TreeNodeA, const FPropertyInstance& TreeNodeB) const
 	{
 		return TreeNodeA == TreeNodeB;
 	}
 	
-	template<>
-	bool AreMatching<FPropertyInstance>(const FPropertyInstance& TreeNodeA, const FPropertyInstance& TreeNodeB)
+	virtual bool AreMatching(const FPropertyInstance& TreeNodeA, const FPropertyInstance& TreeNodeB) const
 	{
 		if (TreeNodeA.KeyProperty->GetName() == TreeNodeB.KeyProperty->GetName())
 		{
@@ -411,14 +413,12 @@ namespace TreeDiffSpecification
 		return false;
 	}
 	
-	template<>
-	void GetChildren<FPropertyInstance>(const FPropertyInstance& InParent, TArray<FPropertyInstance>& OutChildren)
+	virtual void GetChildren(const FPropertyInstance& InParent, TArray<FPropertyInstance>& OutChildren) const
 	{
 		return InParent.GetChildren(OutChildren);
 	}
 
-	template<>
-	bool ShouldMatchByValue<FPropertyInstance>(const FPropertyInstance& TreeNodeA)
+	virtual bool ShouldMatchByValue(const FPropertyInstance& TreeNodeA) const
 	{
 		// array elements should match by value
 		return CastField<FArrayProperty>(TreeNodeA.KeyProperty->Owner.ToField()) != nullptr;
@@ -579,7 +579,7 @@ static void RelinkObjectProperties(const TArray<FPropertyInstance>& Props, UPack
 		}
 		
 		TArray<FPropertyInstance> Children;
-		TreeDiffSpecification::GetChildren(Instance, Children);
+		Instance.GetChildren(Children);
 		if (Children.Num())
 		{
 			RelinkObjectProperties(Children, FromPackage, ToPackage);
diff --git a/Engine/Source/Editor/UnrealEd/Public/DiffUtils.h b/Engine/Source/Editor/UnrealEd/Public/DiffUtils.h
--- a/Engine/Source/Editor/UnrealEd/Public/DiffUtils.h
+++ b/Engine/Source/Editor/UnrealEd/Public/DiffUtils.h
@@ -21,6 +21,7 @@
 #include "Widgets/SBoxPanel.h"
 #include "Widgets/SWidget.h"
 #include "Widgets/Views/STreeView.h"
+#include "PropertyHandle.h"
 
 class ISourceControlRevision;
 class SWidget;
@@ -262,6 +263,7 @@ namespace DiffUtils
 	 */
 	UNREALED_API bool Identical(const FResolvedProperty& AProp, const FResolvedProperty& BProp, const UObject* OwningOuterA, const UObject* OwningOuterB);
 	UNREALED_API bool Identical(const FResolvedProperty& AProp, const FResolvedProperty& BProp, const UObject* OwningOuterA, const UObject* OwningOuterB, const FPropertySoftPath& RootPath, TArray<FPropertySoftPath>& DifferingProperties);
+	UNREALED_API bool Identical(const TSharedPtr<IPropertyHandle>& PropertyHandleA, const TSharedPtr<IPropertyHandle>& PropertyHandleB, const TArray<TWeakObjectPtr<UObject>>& OwningOutersA = {}, const TArray<TWeakObjectPtr<UObject>>& OwningOutersB = {});
 	UNREALED_API TArray<FPropertySoftPath> GetVisiblePropertiesInOrderDeclared(const UStruct* ForStruct, const FPropertySoftPath& Scope = FPropertySoftPath());
 
 	UNREALED_API TArray<FPropertyPath> ResolveAll(const UObject* Object, const TArray<FPropertySoftPath>& InSoftProperties);
