UE-191200 - "Add UE support for [REDACTED_HANDLE] fields of optional types"

Add's support for `[REDACTED_HANDLE]` on optionals:

`FOptionalProperty` changes:
- Added custom `__INIT__` logic for `FOptionalProperty::ExportText` which previously could output no text as the optional value was initialized but returned an empty string as its export text (ie: empty arrays, maps, sets, etc). This caused the future text imports to be incorrect.

`FPropertyNode` changes:
- Added `OptionalValueNode` as private member similar in idea to existing `PropertyKeyNode`. It stores the generated `FPropertyNode` for a FOptionalValue's Value when it is set.
- Added helper function `GetOrCreateOptionalValueNode` which either returns the existing `OptionalValueNode` OR creates one if appropriate.
    - Note: This is where we bind a lambda to our Optional's Value FPropertyNode's `RebuildChildren` event which is how we update/rebuild when needed. In places where we would update/rebuild we instead call our Value to do so instead if set and then update after that has been done via this callback.

Added new widget `SPropertyEditorOptional`
- This either:
    - shows a dropdown of "Set all" or "Unset all" if we have multiple values selected
    - generates the editor for whatever our set value is if we have a single value (See `GetSingleReadAddress`)
    - shows a 'set value' button if we are not set

diff --git a/Engine/Source/Editor/PropertyEditor/Private/DetailItemNode.cpp b/Engine/Source/Editor/PropertyEditor/Private/DetailItemNode.cpp
--- a/Engine/Source/Editor/PropertyEditor/Private/DetailItemNode.cpp
+++ b/Engine/Source/Editor/PropertyEditor/Private/DetailItemNode.cpp
@@ -12,6 +12,7 @@
 #include "PropertyPermissionList.h"
 #include "SDetailCategoryTableRow.h"
 #include "SDetailSingleItemRow.h"
+#include "UObject/PropertyOptional.h"
 
 FDetailItemNode::FDetailItemNode(const FDetailLayoutCustomization& InCustomization, TSharedRef<FDetailCategoryImpl> InParentCategory, TAttribute<bool> InIsParentEnabled, TSharedPtr<IDetailGroup> InParentGroup)
 	: Customization( InCustomization )
@@ -143,7 +144,7 @@ void FDetailItemNode::InitPropertyEditor()
 {
 	FProperty* NodeProperty = Customization.GetPropertyNode()->GetProperty();
 
-	if( NodeProperty && (NodeProperty->IsA<FArrayProperty>() || NodeProperty->IsA<FSetProperty>() || NodeProperty->IsA<FMapProperty>() ))
+	if( NodeProperty && (NodeProperty->IsA<FArrayProperty>() || NodeProperty->IsA<FSetProperty>() || NodeProperty->IsA<FMapProperty>() || NodeProperty->IsA<FOptionalProperty>()))
 	{
 		const bool bUpdateFilteredNodes = false;
 		FSimpleDelegate OnRegenerateChildren = FSimpleDelegate::CreateSP( this, &FDetailItemNode::GenerateChildren, bUpdateFilteredNodes );
@@ -586,7 +587,7 @@ static bool PassesAllFilters( FDetailItemNode* ItemNode, const FDetailLayoutCust
 			if (PropertyNode.IsValid())
 			{
 				// Is it a container (array, map, set?) - if so, ignore it, we don't care about these, only their inner nodes.
-				if (CastField<FArrayProperty>(PropertyNode->GetProperty()) || CastField<FMapProperty>(PropertyNode->GetProperty()) || CastField<FSetProperty>(PropertyNode->GetProperty()))
+				if (CastField<FArrayProperty>(PropertyNode->GetProperty()) || CastField<FMapProperty>(PropertyNode->GetProperty()) || CastField<FSetProperty>(PropertyNode->GetProperty()) || CastField<FOptionalProperty>(PropertyNode->GetProperty()))
 				{
 					return FString();
 				}
diff --git a/Engine/Source/Editor/PropertyEditor/Private/ItemPropertyNode.cpp b/Engine/Source/Editor/PropertyEditor/Private/ItemPropertyNode.cpp
--- a/Engine/Source/Editor/PropertyEditor/Private/ItemPropertyNode.cpp
+++ b/Engine/Source/Editor/PropertyEditor/Private/ItemPropertyNode.cpp
@@ -9,6 +9,8 @@
 #include "Settings/EditorStyleSettings.h"
 #include "UserInterface/PropertyEditor/SPropertyEditorArrayItem.h"
 
+#include "UObject/PropertyOptional.h"
+
 #define LOCTEXT_NAMESPACE "ItemPropertyNode"
 
 FItemPropertyNode::FItemPropertyNode()
@@ -95,17 +97,80 @@ uint8* FItemPropertyNode::GetValueAddress(uint8* StartAddress, bool bIsSparseDat
 /**
  * Overridden function for special setup
  */
-void FItemPropertyNode::InitExpansionFlags (void)
+TSharedPtr<FPropertyNode>& FItemPropertyNode::GetOrCreateOptionalValueNode(void)
 {
+	if (OptionalValueNode)
+	{
+		return OptionalValueNode;
+	}
+
 	FProperty* MyProperty = GetProperty();
+	if (FOptionalProperty* OptionalProperty = CastField<FOptionalProperty>(MyProperty))
+	{
+		void* Optional = NULL;
+		FReadAddressList Addresses;
+		if (GetReadAddress(Addresses))
+		{
+			for (int i = 0; i < Addresses.Num(); i++)
+			{
+				Optional = Addresses.GetAddress(i);
+				if (OptionalProperty->IsSet(Optional))
+				{
+					OptionalValueNode = TSharedRef<FPropertyNode>(new FItemPropertyNode());
+
+					FPropertyNodeInitParams InitParams;
+					InitParams.ParentNode = SharedThis(this);
+					InitParams.Property = OptionalProperty->GetValueProperty();
+					InitParams.bAllowChildren = true;
+					InitParams.bForceHiddenPropertyVisibility = !!HasNodeFlags(EPropertyNodeFlags::ShouldShowHiddenProperties);
+					InitParams.bCreateDisableEditOnInstanceNodes = !!HasNodeFlags(EPropertyNodeFlags::ShouldShowDisableEditOnInstance);
+
+					OptionalValueNode->OnRebuildChildren().AddLambda([this]() {
+						CachedReadAddresses.Reset();
+						bool bDestroySelf = false;
+						DestroyTree(bDestroySelf);
+
+						for (int i = 0; i < OptionalValueNode->GetNumChildNodes(); i++)
+						{
+							AddChildNode(OptionalValueNode->GetChildNode(i));
+						}
+
+						// Children have been rebuilt, clear any pending rebuild requests
+						bRebuildChildrenRequested = false;
+						bChildrenRebuilt = true;
+
+						// Notify any listener that children have been rebuilt
+						OnRebuildChildrenEvent.Broadcast();
+					});
+
+					OptionalValueNode->InitNode(InitParams);
+				}
+			}
+		}
+	}
+	return OptionalValueNode;
+}
+
+/**
+ * Overridden function for special setup
+ */
+void FItemPropertyNode::InitExpansionFlags(void)
+{
+	FProperty* MyProperty = GetProperty();
+
+	if (TSharedPtr<FPropertyNode>& ValueNode = GetOrCreateOptionalValueNode())
+	{
+		// This is a set optional, so check its SetValue instead.
+		MyProperty = ValueNode->GetProperty();
+	}
 
-	bool bExpandableType = CastField<FStructProperty>(MyProperty) 
+	bool bExpandableType = CastField<FStructProperty>(MyProperty)
 		|| (CastField<FArrayProperty>(MyProperty) || CastField<FSetProperty>(MyProperty) || CastField<FMapProperty>(MyProperty));
 
-	if(	bExpandableType
+	if (bExpandableType
 		|| HasNodeFlags(EPropertyNodeFlags::EditInlineNew)
 		|| HasNodeFlags(EPropertyNodeFlags::ShowInnerObjectProperties)
-		|| ( MyProperty->ArrayDim > 1 && ArrayIndex == -1 ) )
+		|| (MyProperty->ArrayDim > 1 && ArrayIndex == -1))
 	{
 		SetNodeFlags(EPropertyNodeFlags::CanBeExpanded, true);
 	}
diff --git a/Engine/Source/Editor/PropertyEditor/Private/ItemPropertyNode.h b/Engine/Source/Editor/PropertyEditor/Private/ItemPropertyNode.h
--- a/Engine/Source/Editor/PropertyEditor/Private/ItemPropertyNode.h
+++ b/Engine/Source/Editor/PropertyEditor/Private/ItemPropertyNode.h
@@ -49,6 +49,11 @@ public:
 	 */
 	virtual FText GetToolTipText() const override;
 
+	/**
+	 * If this is node contains a FOptionalProperty, get its Value FPropertyNode node if set (may create the node if necessary).
+	 */
+	virtual TSharedPtr<FPropertyNode>& GetOrCreateOptionalValueNode() override;
+
 protected:
 	/**
 	 * Overridden function for special setup
diff --git a/Engine/Source/Editor/PropertyEditor/Private/Presentation/PropertyEditor/PropertyEditor.cpp b/Engine/Source/Editor/PropertyEditor/Private/Presentation/PropertyEditor/PropertyEditor.cpp
--- a/Engine/Source/Editor/PropertyEditor/Private/Presentation/PropertyEditor/PropertyEditor.cpp
+++ b/Engine/Source/Editor/PropertyEditor/Private/Presentation/PropertyEditor/PropertyEditor.cpp
@@ -300,6 +300,36 @@ void FPropertyEditor::OnAddGivenItem(const FString InGivenItem)
 	}
 }
 
+void FPropertyEditor::SetOptionalItem(FProperty* NewValue)
+{
+	// This action must be deferred until next tick so that we avoid accessing invalid data before we have a chance to tick
+	PropertyUtilities->EnqueueDeferredAction(FSimpleDelegate::CreateSP(this, &FPropertyEditor::OnSetOptionalValue, NewValue));
+}
+
+void FPropertyEditor::ClearOptionalItem()
+{
+	// This action must be deferred until next tick so that we avoid accessing invalid data before we have a chance to tick
+	PropertyUtilities->EnqueueDeferredAction(FSimpleDelegate::CreateSP(this, &FPropertyEditor::OnClearOptionalValue));
+}
+
+void FPropertyEditor::OnSetOptionalValue(FProperty* NewValue)
+{
+	TSharedPtr<IPropertyHandleOptional> OptionalHandle = PropertyHandle->AsOptional();
+	if (OptionalHandle.IsValid())
+	{
+		OptionalHandle->SetOptionalValue(NewValue);
+	}
+}
+
+void FPropertyEditor::OnClearOptionalValue()
+{
+	TSharedPtr<IPropertyHandleOptional> OptionalHandle = PropertyHandle->AsOptional();
+	if (OptionalHandle.IsValid())
+	{
+		OptionalHandle->ClearOptionalValue();
+	}
+}
+
 void FPropertyEditor::ClearItem()
 {
 	OnClearItem();
diff --git a/Engine/Source/Editor/PropertyEditor/Private/Presentation/PropertyEditor/PropertyEditor.h b/Engine/Source/Editor/PropertyEditor/Private/Presentation/PropertyEditor/PropertyEditor.h
--- a/Engine/Source/Editor/PropertyEditor/Private/Presentation/PropertyEditor/PropertyEditor.h
+++ b/Engine/Source/Editor/PropertyEditor/Private/Presentation/PropertyEditor/PropertyEditor.h
@@ -90,6 +90,15 @@ public:
 	void OnActorSelected( AActor* InActor );
 	void OnGetActorFiltersForSceneOutliner( TSharedPtr<FSceneOutlinerFilters>& OutFilters );
 	void EditConfigHierarchy();
+	
+	/**	
+	* If this is an editor of an optional, set that optional to the passed in value.
+	* If nullptr is passed in, the optional is set and default-initialized.
+	*/
+	void SetOptionalItem(FProperty* NewValue);
+
+	/**	If this is an editor of an optional, unset that optional. */
+	void ClearOptionalItem();
 
 	/**	In an ideal world we wouldn't expose these */
 	TSharedRef<FPropertyNode> GetPropertyNode() const;
@@ -111,6 +120,8 @@ private:
 	void OnDuplicateItem();
 	void OnBrowseTo();
 	void OnEmptyArray();
+	void OnSetOptionalValue(FProperty* NewValue);
+	void OnClearOptionalValue();
 
 private:
 
diff --git a/Engine/Source/Editor/PropertyEditor/Private/PropertyCustomizationHelpers.cpp b/Engine/Source/Editor/PropertyEditor/Private/PropertyCustomizationHelpers.cpp
--- a/Engine/Source/Editor/PropertyEditor/Private/PropertyCustomizationHelpers.cpp
+++ b/Engine/Source/Editor/PropertyEditor/Private/PropertyCustomizationHelpers.cpp
@@ -175,6 +175,37 @@ namespace PropertyCustomizationHelpers
 			.IsFocusable( false );
 	}
 
+	TSharedRef<SWidget> MakeSetOptionalButton(FSimpleDelegate OnSetOptionalClicked, TAttribute<FText> OptionalToolTipText, TAttribute<bool> IsEnabled)
+	{
+		// Custom widget for this button as it has no image and should fill a larger space
+		return SNew(SBox)
+			.HAlign(HAlign_Left)
+			.VAlign(VAlign_Center)
+			.ToolTipText(OptionalToolTipText.Get().IsEmpty() ? LOCTEXT("SetButtonToolTipText", "Set Optional to default value.") : OptionalToolTipText)
+			[
+				SNew(SButton)
+				.ButtonStyle(FAppStyle::Get(), "Button")
+				.OnClicked_Lambda([OnSetOptionalClicked](){
+					OnSetOptionalClicked.ExecuteIfBound();
+					return FReply::Handled();
+				})
+				.Text(LOCTEXT("SetButtonText", "Set to Value"))
+				.ContentPadding(0)
+				.IsFocusable(false)
+			];
+	}
+
+	TSharedRef<SWidget> MakeClearOptionalButton(FSimpleDelegate OnClearOptionalClicked, TAttribute<FText> OptionalToolTipText, TAttribute<bool> IsEnabled)
+	{
+		return
+			SNew(SPropertyEditorButton)
+			.Text(OptionalToolTipText.Get().IsEmpty() ? LOCTEXT("ClearButtonToolTipText", "Clear Optional") : OptionalToolTipText)
+			.Image(FAppStyle::Get().GetBrush("Icons.X"))
+			.OnClickAction(OnClearOptionalClicked)
+			.IsEnabled(IsEnabled)
+			.IsFocusable(false);
+	}
+
 	FText GetVisibilityDisplay(TAttribute<bool> bEnabled)
 	{
 		return bEnabled.Get() ? FEditorFontGlyphs::Eye : FEditorFontGlyphs::Eye_Slash;
diff --git a/Engine/Source/Editor/PropertyEditor/Private/PropertyEditorHelpers.cpp b/Engine/Source/Editor/PropertyEditor/Private/PropertyEditorHelpers.cpp
--- a/Engine/Source/Editor/PropertyEditor/Private/PropertyEditorHelpers.cpp
+++ b/Engine/Source/Editor/PropertyEditor/Private/PropertyEditorHelpers.cpp
@@ -23,11 +23,14 @@
 #include "UserInterface/PropertyEditor/SPropertyEditorStruct.h"
 #include "UserInterface/PropertyEditor/SPropertyEditorSet.h"
 #include "UserInterface/PropertyEditor/SPropertyEditorMap.h"
+#include "UserInterface/PropertyEditor/SPropertyEditorOptional.h"
 
 #include "Kismet2/KismetEditorUtilities.h"
 #include "EditorClassUtils.h"
 #include "Engine/Selection.h"
 
+#include "UObject/PropertyOptional.h"
+
 #define LOCTEXT_NAMESPACE "PropertyEditor"
 
 void SPropertyNameWidget::Construct( const FArguments& InArgs, TSharedPtr<FPropertyEditor> InPropertyEditor )
@@ -157,6 +160,14 @@ TSharedRef<SWidget> SPropertyValueWidget::ConstructPropertyEditorWidget( TShared
 
 			MapWidget->GetDesiredWidth( MinDesiredWidth, MaxDesiredWidth );
 		}
+		else if (SPropertyEditorOptional::Supports(PropertyEditorRef))
+		{
+			TSharedRef<SPropertyEditorOptional> OptionalWidget =
+				SAssignNew(PropertyWidget, SPropertyEditorOptional, PropertyEditorRef, InPropertyUtilities.ToSharedRef())
+				.Font(FontStyle);
+
+			OptionalWidget->GetDesiredWidth(MinDesiredWidth, MaxDesiredWidth);
+		}
 		else if (SPropertyEditorClass::Supports(PropertyEditorRef))
 		{
 			static TArray<TSharedRef<class IClassViewerFilter>> NullFilters;
@@ -438,6 +449,12 @@ namespace PropertyEditorHelpers
 		return NodeProperty && CastField<const FArrayProperty>(NodeProperty) != NULL;
 	}
 
+	bool IsOptionalProperty(const FPropertyNode& InPropertyNode)
+	{
+		const FProperty* NodeProperty = InPropertyNode.GetProperty();
+		return NodeProperty && CastField<const FOptionalProperty>(NodeProperty) != NULL;
+	}
+
 	const FProperty* GetArrayParent( const FPropertyNode& InPropertyNode )
 	{
 		const FProperty* ParentProperty = InPropertyNode.GetParentNode() != NULL ? InPropertyNode.GetParentNode()->GetProperty() : NULL;
@@ -626,6 +643,10 @@ namespace PropertyEditorHelpers
 		{
 			PropertyHandle = MakeShareable(new FPropertyHandleFieldPath(PropertyNode, NotifyHook, PropertyUtilities));
 		}
+		else if (FPropertyHandleOptional::Supports(PropertyNode))
+		{
+			PropertyHandle = MakeShareable(new FPropertyHandleOptional(PropertyNode, NotifyHook, PropertyUtilities));
+		}
 		// struct should be checked last as there are several specializations of it above
 		else if (FPropertyHandleStruct::Supports(PropertyNode))
 		{
@@ -683,6 +704,17 @@ namespace PropertyEditorHelpers
 			}
 		}
 
+		//////////////////////////////
+		// Handle an optional property.
+		if (const FOptionalProperty* OptionalProp = CastField<FOptionalProperty>(NodeProperty))
+		{
+			uint8* ValueAddress = nullptr;
+			if (PropertyNode->GetSingleReadAddress(ValueAddress) == FPropertyAccess::Success && OptionalProp->IsSet(ValueAddress))
+			{
+				OutRequiredButtons.Add(EPropertyButton::OptionalClear);
+			}
+		}
+
 		//////////////////////////////
 		// Handle a class property.
 		FClassProperty* ClassProp = CastField<FClassProperty>(NodeProperty);
@@ -1072,6 +1104,14 @@ namespace PropertyEditorHelpers
 			NewButton = PropertyCustomizationHelpers::MakeDocumentationButton(PropertyEditor);
 			break;
 
+		case EPropertyButton::OptionalSet:
+			NewButton = PropertyCustomizationHelpers::MakeSetOptionalButton(FSimpleDelegate::CreateSP(PropertyEditor, &FPropertyEditor::SetOptionalItem, (FProperty*)nullptr), FText(), IsEnabledAttribute);
+			break;
+
+		case EPropertyButton::OptionalClear:
+			NewButton = PropertyCustomizationHelpers::MakeClearOptionalButton(FSimpleDelegate::CreateSP(PropertyEditor, &FPropertyEditor::ClearOptionalItem), FText(), IsEnabledAttribute);
+			break;
+
 		default:
 			checkf( 0, TEXT( "Unknown button type" ) );
 			break;
diff --git a/Engine/Source/Editor/PropertyEditor/Private/PropertyEditorHelpers.h b/Engine/Source/Editor/PropertyEditor/Private/PropertyEditorHelpers.h
--- a/Engine/Source/Editor/PropertyEditor/Private/PropertyEditorHelpers.h
+++ b/Engine/Source/Editor/PropertyEditor/Private/PropertyEditorHelpers.h
@@ -37,6 +37,8 @@ namespace EPropertyButton
 		NewBlueprint,
 		EditConfigHierarchy,
 		Documentation,
+		OptionalSet,
+		OptionalClear,
 	};
 }
 
@@ -152,6 +154,13 @@ namespace PropertyEditorHelpers
 	 */
 	bool IsDynamicArray( const FPropertyNode& InPropertyNode );
 
+	/**
+	 * Returns whether or not a property is an optional wrapper
+	 *
+	 * @param InPropertyNode	The property node containing the property to check
+	 */
+	bool IsOptionalProperty(const FPropertyNode& InPropertyNode);
+
 	/**
 	 * Returns true if this is an FSoftObjectPath and should be treated like a TSoftObjectPtr
 	 */
diff --git a/Engine/Source/Editor/PropertyEditor/Private/PropertyHandleImpl.cpp b/Engine/Source/Editor/PropertyEditor/Private/PropertyHandleImpl.cpp
--- a/Engine/Source/Editor/PropertyEditor/Private/PropertyHandleImpl.cpp
+++ b/Engine/Source/Editor/PropertyEditor/Private/PropertyHandleImpl.cpp
@@ -32,6 +32,7 @@
 #include "PropertyPathHelpers.h"
 #include "PropertyTextUtilities.h"
 #include "HAL/PlatformApplicationMisc.h"
+#include "UObject/PropertyOptional.h"
 
 #define LOCTEXT_NAMESPACE "PropertyHandleImplementation"
 
@@ -489,10 +490,25 @@ FPropertyAccess::Result FPropertyValueImpl::ImportText( const TArray<FObjectBase
 				bInteractiveChangeInProgress = (Flags & EPropertyValueSetFlags::InteractiveChange) != 0;
 			}
 
+			const bool bDidOptionalHaveValueNode = !!(InPropertyNode->GetOptionalValueNode());
+
 			// Set the new value.
 			EPropertyPortFlags PortFlags = (Flags & EPropertyValueSetFlags::InstanceObjects) != 0 ? PPF_InstanceSubobjects : PPF_None;
 			FPropertyTextUtilities::TextToPropertyHelper(*NewValue, InPropertyNode, NodeProperty, Cur, PortFlags);
 
+			// If we are an FOptionalProperty, this check tells us whether we have changed between set/unset and 
+			// thus require a re-draw of our node in the details panel UI which is accomplished by having our parent rebuild us.
+			// 
+			// NOTE: It seems that the property propagation below relies entierly on the FProperty and not the FPropertyNode.
+			// So, we don't worry about updating `InPropertyNode->OptionalValueNode` as we are going to rebuild anyways which will handle this for us.
+			if (FOptionalProperty* OptionalProperty = CastField<FOptionalProperty>(NodeProperty))
+			{
+				if (ParentNode && bDidOptionalHaveValueNode != OptionalProperty->IsSet(Cur.BaseAddress))
+				{
+					ParentNode->RequestRebuildChildren();
+				}
+			}
+
 			// Cache the value of the property after having modified it.
 			FString ValueAfterImport;
 			FPropertyTextUtilities::PropertyToTextHelper(ValueAfterImport, InPropertyNode, NodeProperty, Cur, PPF_ForDiff);
@@ -3451,6 +3467,7 @@ IMPLEMENT_PROPERTY_VALUE( FPropertyHandleByte )
 IMPLEMENT_PROPERTY_VALUE( FPropertyHandleString )
 IMPLEMENT_PROPERTY_VALUE( FPropertyHandleObject )
 IMPLEMENT_PROPERTY_VALUE( FPropertyHandleArray )
+IMPLEMENT_PROPERTY_VALUE(FPropertyHandleOptional)
 IMPLEMENT_PROPERTY_VALUE( FPropertyHandleText )
 IMPLEMENT_PROPERTY_VALUE( FPropertyHandleSet )
 IMPLEMENT_PROPERTY_VALUE( FPropertyHandleMap )
@@ -4963,6 +4980,116 @@ bool FPropertyHandleArray::IsEditable() const
 	return Implementation->HasValidPropertyNode() && !Implementation->IsEditConst() && Implementation->IsPropertyTypeOf(FArrayProperty::StaticClass());
 }
 
+bool FPropertyHandleOptional::Supports(TSharedRef<FPropertyNode> PropertyNode)
+{
+	FProperty* Property = PropertyNode->GetProperty();
+
+	return !!CastField<FOptionalProperty>(Property);
+}
+
+/** IPropertyHandleOptional interface */
+FPropertyAccess::Result FPropertyHandleOptional::GetOptionalValue(FProperty* OutValue)
+{
+	const TSharedPtr<FPropertyNode>& PropertyNode = Implementation->GetPropertyNode();
+	FOptionalProperty* OptionalProperty = CastField<FOptionalProperty>(PropertyNode->GetProperty());
+	
+	uint8* ValueAddress = nullptr;
+	FPropertyAccess::Result Result = PropertyNode->GetSingleReadAddress(ValueAddress);
+
+	if (Result != FPropertyAccess::Success)
+	{
+		return Result;
+	}
+
+	if (OptionalProperty->IsSet(ValueAddress))
+	{
+		OutValue = OptionalProperty->GetValueProperty();
+	}
+
+	return Result;
+}
+
+FPropertyAccess::Result FPropertyHandleOptional::SetOptionalValue(FProperty* NewValue)
+{
+	const TSharedPtr<FPropertyNode>& PropertyNode = Implementation->GetPropertyNode();
+	FOptionalProperty* OptionalProperty = CastField<FOptionalProperty>(PropertyNode->GetProperty());
+	
+	FReadAddressList Addresses;
+	if (!PropertyNode->GetReadAddress(Addresses))
+	{
+		return FPropertyAccess::Fail;
+	}
+
+	for (int i = 0; i < Addresses.Num(); i++)
+	{
+		void* Optional = Addresses.GetAddress(i);
+		OptionalProperty->MarkSetAndGetInitializedValuePointerToReplace(Optional);
+
+		if (NewValue)
+		{
+			OptionalProperty->SetValueProperty(NewValue);
+		}
+	}
+
+	// Rebuild our parent as we require a de-draw in the details panel
+	if (FPropertyNode* ParentNode = PropertyNode->GetParentNode())
+	{
+		ParentNode->RequestRebuildChildren();
+	}
+
+	return FPropertyAccess::Success;
+}
+
+FPropertyAccess::Result FPropertyHandleOptional::ClearOptionalValue()
+{
+	if (!IsEditable())
+	{
+		return FPropertyAccess::Fail;
+	}
+
+	const TSharedPtr<FPropertyNode>& PropertyNode = Implementation->GetPropertyNode();
+	FOptionalProperty* OptionalProperty = CastField<FOptionalProperty>(PropertyNode->GetProperty());
+
+	if (!OptionalProperty)
+	{
+		return FPropertyAccess::Fail;
+	}
+
+	FReadAddressList Addresses;
+	if (!PropertyNode->GetReadAddress(Addresses))
+	{
+		return FPropertyAccess::Fail;
+	}
+	
+	for (int i = 0; i < Addresses.Num(); i++)
+	{
+		void* Optional = Addresses.GetAddress(i);
+		OptionalProperty->MarkUnset(Optional);
+	}
+
+	// Could be removed as unecessary (rebuild will do this for us... but removing now makes any future debugging clearer)
+	PropertyNode->GetOptionalValueNode().Reset();
+
+	// Rebuild our parent as we require a re-draw in the details panel
+	if (FPropertyNode* ParentNode = PropertyNode->GetParentNode())
+	{
+		ParentNode->RequestRebuildChildren();
+	}
+
+	return FPropertyAccess::Success;
+}
+
+TSharedPtr<IPropertyHandleOptional> FPropertyHandleOptional::AsOptional()
+{
+	return SharedThis(this);
+}
+
+bool FPropertyHandleOptional::IsEditable() const
+{
+	// Property is editable if its a non-const dynamic array
+	return Implementation->HasValidPropertyNode() && !Implementation->IsEditConst() && Implementation->IsPropertyTypeOf(FOptionalProperty::StaticClass());
+}
+
 // Localized Text
 bool FPropertyHandleText::Supports(TSharedRef<FPropertyNode> PropertyNode)
 {
diff --git a/Engine/Source/Editor/PropertyEditor/Private/PropertyHandleImpl.h b/Engine/Source/Editor/PropertyEditor/Private/PropertyHandleImpl.h
--- a/Engine/Source/Editor/PropertyEditor/Private/PropertyHandleImpl.h
+++ b/Engine/Source/Editor/PropertyEditor/Private/PropertyHandleImpl.h
@@ -465,6 +465,7 @@ public:
 	virtual TSharedPtr<IPropertyHandleArray> AsArray() override { return nullptr; }
 	virtual TSharedPtr<IPropertyHandleSet> AsSet() override { return nullptr; }
 	virtual TSharedPtr<IPropertyHandleMap> AsMap() override { return nullptr; }
+	virtual TSharedPtr<IPropertyHandleOptional> AsOptional() override { return nullptr; }
 	virtual TSharedPtr<IPropertyHandleStruct> AsStruct() override { return nullptr; }
 	virtual const FFieldClass* GetPropertyClass() const override;
 	virtual FProperty* GetProperty() const override;
@@ -712,6 +713,19 @@ public:
 	virtual bool IsEditable() const override;
 };
 
+class FPropertyHandleOptional : public FPropertyHandleBase, public IPropertyHandleOptional
+{
+public:
+	FPropertyHandleOptional(TSharedRef<FPropertyNode> PropertyNode, FNotifyHook* NotifyHook, TSharedPtr<IPropertyUtilities> PropertyUtilities);
+	static bool Supports(TSharedRef<FPropertyNode> PropertyNode);
+	/** IPropertyHandleOptional interface */
+	virtual FPropertyAccess::Result GetOptionalValue(FProperty* OutValue) override;
+	virtual FPropertyAccess::Result SetOptionalValue(FProperty* NewValue) override;
+	virtual FPropertyAccess::Result ClearOptionalValue() override;
+	virtual TSharedPtr<IPropertyHandleOptional> AsOptional() override;
+	virtual bool IsEditable() const override;
+};
+
 class FPropertyHandleText : public FPropertyHandleBase
 {
 public:
diff --git a/Engine/Source/Editor/PropertyEditor/Private/PropertyNode.cpp b/Engine/Source/Editor/PropertyEditor/Private/PropertyNode.cpp
--- a/Engine/Source/Editor/PropertyEditor/Private/PropertyNode.cpp
+++ b/Engine/Source/Editor/PropertyEditor/Private/PropertyNode.cpp
@@ -30,6 +30,8 @@
 #include "UObject/EnumProperty.h"
 #include "UObject/UnrealType.h"
 
+#include "UObject/PropertyOptional.h"
+
 #define LOCTEXT_NAMESPACE "PropertyNode"
 
 FEditConditionParser FPropertyNode::EditConditionParser;
@@ -173,7 +175,7 @@ void FPropertyNode::InitNode(const FPropertyNodeInitParams& InitParams)
 		const FProperty* OwnerProperty = MyProperty->GetOwnerProperty();
 
 		const bool bIsObjectOrInterface = CastField<FObjectPropertyBase>(MyProperty) || CastField<FInterfaceProperty>(MyProperty);
-		bool bIsInsideContainer = CastField<FArrayProperty>(OwnerProperty) || CastField<FSetProperty>(OwnerProperty) || CastField<FMapProperty>(OwnerProperty);
+		bool bIsInsideContainer = CastField<FArrayProperty>(OwnerProperty) || CastField<FSetProperty>(OwnerProperty) || CastField<FMapProperty>(OwnerProperty) || CastField<FOptionalProperty>(OwnerProperty);
 
 		// Don't consider the container's inline status if the key is a class property that is not inline
 		if (const FMapProperty* MapProperty = CastField<FMapProperty>(OwnerProperty))
@@ -243,7 +245,7 @@ void FPropertyNode::InitNode(const FPropertyNodeInitParams& InitParams)
 		bool bRequiresValidation = bIsEditInlineNew || bShowInnerObjectProperties;
 	
 		// We require validation if we are in a container.
-		bRequiresValidation |= MyProperty->IsA<FArrayProperty>() || MyProperty->IsA<FSetProperty>() || MyProperty->IsA<FMapProperty>();
+		bRequiresValidation |= MyProperty->IsA<FArrayProperty>() || MyProperty->IsA<FSetProperty>() || MyProperty->IsA<FMapProperty>() || MyProperty->IsA<FOptionalProperty>();
 
 		// We require validation if our parent also needs validation (if an array parent was resized all the addresses of children are invalid)
 		bRequiresValidation |= (GetParentNode() && GetParentNode()->HasNodeFlags(EPropertyNodeFlags::RequiresValidation));
@@ -270,6 +272,11 @@ void FPropertyNode::InitNode(const FPropertyNodeInitParams& InitParams)
  */
 void FPropertyNode::RebuildChildren()
 {
+	if (TSharedPtr<FPropertyNode>& ValueNode = GetOrCreateOptionalValueNode())
+	{
+		return ValueNode->RebuildChildren();
+	}
+
 	CachedReadAddresses.Reset();
 
 	bool bDestroySelf = false;
@@ -495,6 +502,11 @@ void FPropertyNode::MarkChildrenAsRebuilt()
  */
 EPropertyDataValidationResult FPropertyNode::EnsureDataIsValid()
 {
+	if (TSharedPtr<FPropertyNode>& ValueNode = GetOrCreateOptionalValueNode())
+	{
+		return ValueNode->EnsureDataIsValid();
+	}
+
 	bool bValidateChildren = !HasNodeFlags(EPropertyNodeFlags::SkipChildValidation);
 	bool bValidateChildrenKeyNodes = false;		// by default, we don't check this, since it's just for Map properties
 
@@ -1651,8 +1663,8 @@ public:
 			PropertyValueRoot.OwnerObject = Cast<UClass>(PropertyValueRoot.OwnerObject)->GetDefaultObject();
 		}
 
-		const bool bIsContainerProperty = CastField<FArrayProperty>(Property) || CastField<FSetProperty>(Property) || CastField<FMapProperty>(Property);
-		const bool bIsInsideContainerProperty = Property->GetOwner<FArrayProperty>() || Property->GetOwner<FSetProperty>() || Property->GetOwner<FMapProperty>();
+		const bool bIsContainerProperty = CastField<FArrayProperty>(Property) || CastField<FSetProperty>(Property) || CastField<FMapProperty>(Property) || CastField<FOptionalProperty>(Property);
+		const bool bIsInsideContainerProperty = Property->GetOwner<FArrayProperty>() || Property->GetOwner<FSetProperty>() || Property->GetOwner<FMapProperty>() || Property->GetOwner<FOptionalProperty>();
 
 		FPropertyNode* Node = bIsInsideContainerProperty ? ParentNode : PropertyNode;
 
@@ -2349,8 +2361,8 @@ bool FPropertyNode::GetDiffersFromDefault()
 			StructNode->GetAllStructureData(Structs);
 			
 			const bool bIsSparse = HasNodeFlags(EPropertyNodeFlags::IsSparseClassData);
-			const bool bIsContainer = CastField<FArrayProperty>(Prop) || CastField<FSetProperty>(Prop) || CastField<FMapProperty>(Prop);
-			const bool bIsInsideContainerProperty = Property->GetOwner<FArrayProperty>() || Property->GetOwner<FSetProperty>() || Property->GetOwner<FMapProperty>();
+			const bool bIsContainer = CastField<FArrayProperty>(Prop) || CastField<FSetProperty>(Prop) || CastField<FMapProperty>(Prop) || CastField<FOptionalProperty>(Prop);
+			const bool bIsInsideContainerProperty = Property->GetOwner<FArrayProperty>() || Property->GetOwner<FSetProperty>() || Property->GetOwner<FMapProperty>() || Property->GetOwner<FOptionalProperty>();
 			const FPropertyNode* BaseNode = bIsInsideContainerProperty ? GetParentNode() : this;
 
 			FStructOnScope DefaultStruct;
@@ -2502,8 +2514,8 @@ FString FPropertyNode::GetDefaultValueAsString(bool bUseDisplayName)
 		StructNode->GetAllStructureData(Structs);
 
 		const bool bIsSparse = HasNodeFlags(EPropertyNodeFlags::IsSparseClassData);
-		const bool bIsContainer = CastField<FArrayProperty>(Prop) || CastField<FSetProperty>(Prop) || CastField<FMapProperty>(Prop);
-		const bool bIsInsideContainerProperty = Property->GetOwner<FArrayProperty>() || Property->GetOwner<FSetProperty>() || Property->GetOwner<FMapProperty>();
+		const bool bIsContainer = CastField<FArrayProperty>(Prop) || CastField<FSetProperty>(Prop) || CastField<FMapProperty>(Prop) || CastField<FOptionalProperty>(Prop);
+		const bool bIsInsideContainerProperty = Property->GetOwner<FArrayProperty>() || Property->GetOwner<FSetProperty>() || Property->GetOwner<FMapProperty>() || Property->GetOwner<FOptionalProperty>();
 		const FPropertyNode* BaseNode = bIsInsideContainerProperty ? GetParentNode() : this;
 
 		FStructOnScope DefaultStruct;
diff --git a/Engine/Source/Editor/PropertyEditor/Private/PropertyNode.h b/Engine/Source/Editor/PropertyEditor/Private/PropertyNode.h
--- a/Engine/Source/Editor/PropertyEditor/Private/PropertyNode.h
+++ b/Engine/Source/Editor/PropertyEditor/Private/PropertyNode.h
@@ -1078,6 +1078,16 @@ public:
 	 */
 	void MarkDestroyedRecursive();
 
+	TSharedPtr<FPropertyNode>& GetOptionalValueNode() { return OptionalValueNode; }
+
+	/**
+	 * Interface function for getting a FOptionalProperty's Node's OptionalValueNode (May construct it if needed).
+	 * @return The optionals value OR null if this is not an optional or is an unset optional.
+	 * 
+	 * Note: This is currently the only method by which an optional's Value FPropertyNode is created (If you change this please update this documentation).
+	 */
+	virtual TSharedPtr<FPropertyNode>& GetOrCreateOptionalValueNode() { return OptionalValueNode; }
+
 protected:
 	TSharedRef<FEditPropertyChain> BuildPropertyChain( FProperty* PropertyAboutToChange ) const;
 	TSharedRef<FEditPropertyChain> BuildPropertyChain( FProperty* PropertyAboutToChange, const TSet<UObject*>& InAffectedArchetypeInstances ) const;
@@ -1189,6 +1199,9 @@ protected:
 	/**	The property node, if any, that serves as the key value for this node */
 	TSharedPtr<FPropertyNode> PropertyKeyNode;
 
+	/**	The property node, if any, is this nodes optional value */
+	TSharedPtr<FPropertyNode> OptionalValueNode;
+
 	/** Cached read addresses for this property node */
 	mutable FReadAddressListData CachedReadAddresses;
 
diff --git a/Engine/Source/Editor/PropertyEditor/Private/UserInterface/PropertyEditor/SPropertyEditorOptional.cpp b/Engine/Source/Editor/PropertyEditor/Private/UserInterface/PropertyEditor/SPropertyEditorOptional.cpp
new file mode 100644
--- /dev/null
+++ b/Engine/Source/Editor/PropertyEditor/Private/UserInterface/PropertyEditor/SPropertyEditorOptional.cpp
@@ -0,0 +1,136 @@
+// Copyright Epic Games, Inc. All Rights Reserved.
+
+#include "UserInterface/PropertyEditor/SPropertyEditorOptional.h"
+
+#include "Delegates/Delegate.h"
+#include "Internationalization/Internationalization.h"
+#include "Layout/Children.h"
+#include "Presentation/PropertyEditor/PropertyEditor.h"
+#include "PropertyNode.h"
+#include "UObject/PropertyOptional.h"
+#include "Widgets/Input/SButton.h"
+#include "Widgets/Text/STextBlock.h"
+#include "Widgets/Input/SComboBox.h"
+#include "PropertyEditorHelpers.h"
+
+#define LOCTEXT_NAMESPACE "PropertyEditor"
+
+void SPropertyEditorOptional::Construct(const FArguments& InArgs, const TSharedRef<FPropertyEditor>& InPropertyEditor, TSharedRef<IPropertyUtilities> InPropertyUtilities)
+{
+	MinDesiredWidth = 0.0f;
+	MaxDesiredWidth = 0.0f;
+	PropertyEditor = InPropertyEditor;
+
+	const TSharedRef<FPropertyNode> PropertyNode = InPropertyEditor->GetPropertyNode();
+
+	uint8* ValueAddress = nullptr;
+	if (PropertyNode->GetSingleReadAddress(ValueAddress) == FPropertyAccess::MultipleValues)
+	{	
+		MultiselectOptions.Add(MakeShareable( new FText(LOCTEXT("OptionalMultiselect_Multiple", "MultipleStates"))));
+		MultiselectOptions.Add(MakeShareable( new FText(LOCTEXT("OptionalMultiselect_Value", "Set all to Value"))));
+		MultiselectOptions.Add(MakeShareable( new FText(LOCTEXT("OptionalMultiselect_None", "Set all to None"))));
+
+		ValueEditorWidget = SNew(SComboBox<TSharedPtr<FText>>)
+			.OptionsSource(&MultiselectOptions)
+			.OnSelectionChanged(this, &SPropertyEditorOptional::OnOptionChanged)
+			.OnGenerateWidget(this, &SPropertyEditorOptional::MakeWidgetForOption)
+			.InitiallySelectedItem(MultiselectOptions[0])
+			[
+				SNew(STextBlock)
+					.Text(*MultiselectOptions[0])
+			];
+	}
+	else if (const TSharedPtr<FPropertyNode> OptionalValueNode = PropertyNode->GetOrCreateOptionalValueNode())
+	{
+		// create the OptionalValue's PropertyEditor
+		TSharedRef<FPropertyEditor> ValuePropertyEditor = FPropertyEditor::Create(OptionalValueNode.ToSharedRef(), InPropertyUtilities);
+		
+		TSharedPtr<SPropertyValueWidget> Widget = SNew(SPropertyValueWidget, ValuePropertyEditor, InPropertyUtilities)
+			.IsEnabled(InArgs._IsEnabled)
+			.ShowPropertyButtons(true);
+
+		 MinDesiredWidth = Widget->GetMinDesiredWidth();
+		 MaxDesiredWidth = Widget->GetMaxDesiredWidth();
+
+		 ValueEditorWidget = Widget;
+	}
+	else
+	{
+		ValueEditorWidget = PropertyEditorHelpers::MakePropertyButton(EPropertyButton::OptionalSet, InPropertyEditor);
+	}
+
+	ChildSlot
+	[
+		ValueEditorWidget.ToSharedRef()
+	];
+
+	SetToolTipText(GetOptionalTooltipText());
+
+	SetEnabled(TAttribute<bool>(this, &SPropertyEditorOptional::CanEdit));
+}
+
+bool SPropertyEditorOptional::Supports(const TSharedRef< FPropertyEditor >& InPropertyEditor)
+{
+	const TSharedRef<FPropertyNode> PropertyNode = InPropertyEditor->GetPropertyNode();
+	const FProperty* Property = InPropertyEditor->GetProperty();
+
+	if (Property && !PropertyNode->HasNodeFlags(EPropertyNodeFlags::EditInlineNew)
+		&& Property->IsA<FOptionalProperty>())
+	{
+		return true;
+	}
+
+	return false;
+}
+
+void SPropertyEditorOptional::GetDesiredWidth(float& OutMinDesiredWidth, float& OutMaxDesiredWidth)
+{
+	// Did we get desired width's from our set value? If not just set to a default value of 170
+	if (MinDesiredWidth && MaxDesiredWidth)
+	{
+		OutMinDesiredWidth = MinDesiredWidth;
+		OutMaxDesiredWidth = MaxDesiredWidth;
+	}
+	else
+	{
+		OutMinDesiredWidth = 170.0f;
+		OutMaxDesiredWidth = 170.0f;
+	}
+}
+
+FText SPropertyEditorOptional::GetOptionalTooltipText() const
+{
+	return LOCTEXT("RichOptionalTooltipText", "This is an optional property. It can either be `set`, which means it contains a `value` property, or it can be unset which means it contains nothing.");
+}
+
+bool SPropertyEditorOptional::CanEdit() const
+{
+	return PropertyEditor.IsValid() ? !PropertyEditor->IsEditConst() : true;
+}
+
+TSharedRef<SWidget> SPropertyEditorOptional::MakeWidgetForOption(TSharedPtr<FText> InOption)
+{
+	return SNew(STextBlock).Text(*InOption);
+}
+
+void SPropertyEditorOptional::OnOptionChanged(TSharedPtr<FText> NewOption, ESelectInfo::Type)
+{
+	TSharedPtr<IPropertyHandleOptional> OptionalHandle = PropertyEditor->GetPropertyHandle()->AsOptional();
+	if (!OptionalHandle.IsValid())
+	{
+		// We cannot do anything if we don't have a valid handle
+		return;
+	}
+
+	if (NewOption->EqualTo(*MultiselectOptions[1]) /*"Set all to Value"*/)
+	{
+		OptionalHandle->ClearOptionalValue(); // Clear all values
+		OptionalHandle->SetOptionalValue(nullptr); // Set values to default-initialized
+	}
+	else if (NewOption->EqualTo(*MultiselectOptions[2]) /*"Set all to None"*/)
+	{
+		OptionalHandle->ClearOptionalValue();
+	}
+}
+
+#undef LOCTEXT_NAMESPACE
diff --git a/Engine/Source/Editor/PropertyEditor/Private/UserInterface/PropertyEditor/SPropertyEditorOptional.h b/Engine/Source/Editor/PropertyEditor/Private/UserInterface/PropertyEditor/SPropertyEditorOptional.h
new file mode 100644
--- /dev/null
+++ b/Engine/Source/Editor/PropertyEditor/Private/UserInterface/PropertyEditor/SPropertyEditorOptional.h
@@ -0,0 +1,62 @@
+// Copyright Epic Games, Inc. All Rights Reserved.
+#pragma once
+
+#include "Fonts/SlateFontInfo.h"
+#include "Internationalization/Text.h"
+#include "Misc/Attribute.h"
+#include "Styling/AppStyle.h"
+#include "Templates/SharedPointer.h"
+#include "Templates/TypeHash.h"
+#include "UserInterface/PropertyEditor/PropertyEditorConstants.h"
+#include "Widgets/DeclarativeSyntaxSupport.h"
+#include "Widgets/SCompoundWidget.h"
+
+#include "IPropertyUtilities.h"
+
+class FPropertyEditor;
+
+/*
+* This class represents a `value editor widget` for optional types in the details panel.
+* 
+* It has three possible states which it can display:
+*	1. (Unset Single-Select) 
+*		The optional has no set value. A button `Set to Value` is shown and
+*		pressing this button will set the optional's value to a default initialized state.
+*	2. (Set Single-Select OR Multi-Select) 
+*		The optional has a set value. The `value editor widget` for the set value is shown
+*		(example: SPropertyEditorArray) along with a `Clear Optional` button which un-sets the optional.
+*	3. (Multi-Select with mix of set + unset)
+*		A combo-box with `Multiple States` is shown which when clicked displays a dropdown menu offering:
+*		`Set all to Value` (sets all values to default initialized)
+*		`Set all to None` (unsets all values)
+*/
+class SPropertyEditorOptional : public SCompoundWidget
+{
+public:
+	SLATE_BEGIN_ARGS(SPropertyEditorOptional)
+		: _Font(FAppStyle::GetFontStyle(PropertyEditorConstants::PropertyFontStyle))
+		{}
+		SLATE_ATTRIBUTE(FSlateFontInfo, Font)
+	SLATE_END_ARGS()
+
+	static bool Supports(const TSharedRef<FPropertyEditor>& InPropertyEditor);
+
+	void Construct(const FArguments& InArgs, const TSharedRef<FPropertyEditor>& InPropertyEditor, TSharedRef<IPropertyUtilities> InPropertyUtilities);
+
+	void GetDesiredWidth(float& OutMinDesiredWidth, float& OutMaxDesiredWidth);
+
+private:
+	FText GetOptionalTooltipText() const;
+
+	/** @return True if the property can be edited */
+	bool CanEdit() const;
+
+	TSharedRef<SWidget> MakeWidgetForOption(TSharedPtr<FText> InOption);
+	void OnOptionChanged(TSharedPtr<FText> NewOption, ESelectInfo::Type);
+
+	TSharedPtr< FPropertyEditor > PropertyEditor;
+	TSharedPtr<SWidget> ValueEditorWidget;
+	TArray<TSharedPtr<FText>> MultiselectOptions;
+	float MinDesiredWidth;
+	float MaxDesiredWidth;
+};
diff --git a/Engine/Source/Editor/PropertyEditor/Public/PropertyCustomizationHelpers.h b/Engine/Source/Editor/PropertyEditor/Public/PropertyCustomizationHelpers.h
--- a/Engine/Source/Editor/PropertyEditor/Public/PropertyCustomizationHelpers.h
+++ b/Engine/Source/Editor/PropertyEditor/Public/PropertyCustomizationHelpers.h
@@ -116,6 +116,8 @@ namespace PropertyCustomizationHelpers
 	PROPERTYEDITOR_API TSharedRef<SWidget> MakeEditConfigHierarchyButton(FSimpleDelegate OnEditConfigClicked, TAttribute<FText> OptionalToolTipText = FText(), TAttribute<bool> IsEnabled = true);
 	PROPERTYEDITOR_API TSharedRef<SWidget> MakeDocumentationButton(const TSharedRef<FPropertyEditor>& InPropertyEditor);
 	PROPERTYEDITOR_API TSharedRef<SWidget> MakeSaveButton(FSimpleDelegate OnSaveClicked, TAttribute<FText> OptionalToolTipText = FText(), TAttribute<bool> IsEnabled = true);
+	PROPERTYEDITOR_API TSharedRef<SWidget> MakeSetOptionalButton(FSimpleDelegate OnSetOptionalClicked, TAttribute<FText> OptionalToolTipText = FText(), TAttribute<bool> IsEnabled = true);
+	PROPERTYEDITOR_API TSharedRef<SWidget> MakeClearOptionalButton(FSimpleDelegate OnClearOptionalClicked, TAttribute<FText> OptionalToolTipText = FText(), TAttribute<bool> IsEnabled = true);
 
 	/** @return the FBoolProperty edit condition property if one exists. */
 	PROPERTYEDITOR_API FBoolProperty* GetEditConditionProperty(const FProperty* InProperty, bool& bNegate);
diff --git a/Engine/Source/Editor/PropertyEditor/Public/PropertyHandle.h b/Engine/Source/Editor/PropertyEditor/Public/PropertyHandle.h
--- a/Engine/Source/Editor/PropertyEditor/Public/PropertyHandle.h
+++ b/Engine/Source/Editor/PropertyEditor/Public/PropertyHandle.h
@@ -14,6 +14,7 @@ class FResetToDefaultOverride;
 class IPropertyHandleArray;
 class IPropertyHandleMap;
 class IPropertyHandleSet;
+class IPropertyHandleOptional;
 class IPropertyHandleStruct;
 class IStructureDataProvider;
 
@@ -521,6 +522,11 @@ public:
 	 */
 	virtual TSharedPtr<IPropertyHandleMap> AsMap() = 0;
 
+	/**
+	 * @return This handle as an optional if possible
+	 */
+	virtual TSharedPtr<IPropertyHandleOptional> AsOptional() = 0;
+
 	/**
 	 * @return This handle as struct if possible
 	 */
@@ -911,4 +917,32 @@ public:
 	virtual ~IPropertyHandleStruct() {}
 
 	virtual TSharedPtr<FStructOnScope> GetStructData() const = 0;
-};
\ No newline at end of file
+};
+
+/**
+ * A handle to an optional property which allows you to manipulate the optional
+ */
+class IPropertyHandleOptional
+{
+public:
+	virtual ~IPropertyHandleOptional() {}
+
+	/**
+	 * Get the item
+	 * @param OutValue	The value of the optional if it is set.
+	 */
+	virtual FPropertyAccess::Result GetOptionalValue(FProperty* OutValue) = 0;
+
+	/**
+	 * Set the item
+	 * @param NewValue	The value to set the optional to. Pass in nullptr to default-initialize the value.
+	 * @return			Whether or not this was successful
+	 */
+	virtual FPropertyAccess::Result SetOptionalValue(FProperty* NewValue) = 0;
+
+	/**
+	 * Clear the item
+	 * @return Whether or not this was successful
+	 */
+	virtual FPropertyAccess::Result ClearOptionalValue() = 0;
+};
diff --git a/Engine/Source/Runtime/CoreUObject/Private/UObject/PropertyOptional.cpp b/Engine/Source/Runtime/CoreUObject/Private/UObject/PropertyOptional.cpp
--- a/Engine/Source/Runtime/CoreUObject/Private/UObject/PropertyOptional.cpp
+++ b/Engine/Source/Runtime/CoreUObject/Private/UObject/PropertyOptional.cpp
@@ -33,6 +33,7 @@ struct FLegacyFOptionPropertyCustomVersion
 	{}
 };
 static FLegacyFOptionPropertyCustomVersion LegacyFOptionPropertyCustomVersion;
+static const FString InitString = TEXT("__INIT__");
 
 FOptionalProperty::FOptionalProperty(FFieldVariant InOwner, const FName& InName, EObjectFlags InObjectFlags)
 	: FProperty(InOwner, InName, InObjectFlags)
@@ -301,6 +302,14 @@ void FOptionalProperty::ExportText_Internal(FString& ValueStr, const void* Conta
 	if (const void* ValuePointer = GetValuePointerForReadIfSet(OptionalValuePointer))
 	{
 		ValueProperty->ExportTextItem_Direct(ValueStr, ValuePointer, DefaultValue ? GetValuePointerForReadIfSet(DefaultValue) : nullptr, Parent, PortFlags, ExportRootScope);
+
+		// If we got no value back from our ValueProperty's text export but we are SET (ie: an empty array, map, set, etc), 
+		// we set `__init__` to the exported text to so we know to do an initialize on importing (See ImportText_Internal).
+		if (ValueStr == TEXT("("))
+		{
+			ValueStr = InitString;
+			return;
+		}
 	}
 	ValueStr += TEXT(")");
 }
@@ -310,7 +319,21 @@ const TCHAR* FOptionalProperty::ImportText_Internal(const TCHAR* Buffer, void* C
 	checkSlow(ValueProperty);
 	void* Data = PointerToValuePtr(ContainerOrPropertyPtr, PropertyPointerType);
 
-	if (!Buffer || *Buffer++ != TCHAR('('))
+	if (!Buffer)
+	{
+		return nullptr;
+	}
+
+	// INIT
+	if (FString(Buffer) == InitString)
+	{
+		// We are set but have no text to import for our value property.
+		MarkUnset(Data);
+		MarkSetAndGetInitializedValuePointerToReplace(Data);
+		return Buffer;
+	}
+
+	if (*Buffer++ != TCHAR('('))
 	{
 		return nullptr;
 	}
