[CIS Fix] namespaced static functions to avoid unity compile errors

diff --git a/Engine/Source/Editor/PropertyEditor/Private/AsyncDetailViewDiff.cpp b/Engine/Source/Editor/PropertyEditor/Private/AsyncDetailViewDiff.cpp
--- a/Engine/Source/Editor/PropertyEditor/Private/AsyncDetailViewDiff.cpp
+++ b/Engine/Source/Editor/PropertyEditor/Private/AsyncDetailViewDiff.cpp
@@ -6,30 +6,85 @@
 #include "DiffUtils.h"
 #include "PropertyHandleImpl.h"
 
-static TArray<TWeakObjectPtr<UObject>> GetObjects(const TSharedPtr<FDetailTreeNode>& TreeNode)
+namespace AsyncDetailViewDiffHelpers
 {
-	TArray<UObject*> Result;
-	if (const IDetailsViewPrivate* DetailsView = TreeNode->GetDetailsView())
+	TArray<TWeakObjectPtr<UObject>> GetObjects(const TSharedPtr<FDetailTreeNode>& TreeNode)
 	{
-		return DetailsView->GetSelectedObjects();
+		TArray<UObject*> Result;
+		if (const IDetailsViewPrivate* DetailsView = TreeNode->GetDetailsView())
+		{
+			return DetailsView->GetSelectedObjects();
+		}
+		return {};
 	}
-	return {};
-}
 
-static FResolvedProperty GetResolvedProperty(const TSharedPtr<FPropertyNode>& PropertyNode, const UObject* Object)
-{
-	if (PropertyNode && Object)
+	FResolvedProperty GetResolvedProperty(const TSharedPtr<FPropertyNode>& PropertyNode, const UObject* Object)
+	{
+		if (PropertyNode && Object)
+		{
+			const TSharedRef<FPropertyPath> PropertyPath = FPropertyNode::CreatePropertyPath(PropertyNode.ToSharedRef());
+			if (PropertyPath->IsValid())
+			{
+				return FPropertySoftPath(*PropertyPath).Resolve(Object);
+			}
+		}
+		return FResolvedProperty();
+	}
+	
+	bool MapKeysMatch(const TSharedRef<FPropertyNode>& MapPropertyNodeA, const TSharedRef<FPropertyNode>& MapPropertyNodeB, int32 KeyIndexA, int32 KeyIndexB,
+		const UObject* OwningObjectA, const UObject* OwningObjectB)
 	{
-		const TSharedRef<FPropertyPath> PropertyPath = FPropertyNode::CreatePropertyPath(PropertyNode.ToSharedRef());
-		if (PropertyPath->IsValid())
+		const FMapProperty* MapPropertyA = CastField<FMapProperty>(MapPropertyNodeA->GetProperty());
+		const FMapProperty* MapPropertyB = CastField<FMapProperty>(MapPropertyNodeB->GetProperty());
+		if (!MapPropertyA || !MapPropertyB)
+		{
+			return false;
+		}
+		
+		const FResolvedProperty ResolvedMapA = AsyncDetailViewDiffHelpers::GetResolvedProperty(MapPropertyNodeA, OwningObjectA);
+		const FResolvedProperty ResolvedMapB = AsyncDetailViewDiffHelpers::GetResolvedProperty(MapPropertyNodeB, OwningObjectB);
+		FScriptMapHelper MapHelperA(MapPropertyA, MapPropertyA->ContainerPtrToValuePtr<UObject*>(ResolvedMapA.Object));
+		FScriptMapHelper MapHelperB(MapPropertyB, MapPropertyB->ContainerPtrToValuePtr<UObject*>(ResolvedMapB.Object));
+		
+		const void* KeyA = MapHelperA.FindNthKeyPtr(KeyIndexA);
+		const void* KeyB = MapHelperB.FindNthKeyPtr(KeyIndexB);
+		
+		if (MapPropertyA->KeyProp->SameType(MapPropertyB->KeyProp))
+		{
+			return MapPropertyA->KeyProp->Identical(KeyA, KeyB, PPF_DeepComparison);
+		}
+		return false;
+	}
+
+	bool SetKeysMatch(const TSharedRef<FPropertyNode>& SetPropertyNodeA, const TSharedRef<FPropertyNode>& SetPropertyNodeB, int32 KeyIndexA, int32 KeyIndexB,
+		const UObject* OwningObjectA, const UObject* OwningObjectB)
+	{
+		const FSetProperty* SetPropertyA = CastField<FSetProperty>(SetPropertyNodeA->GetProperty());
+		const FSetProperty* SetPropertyB = CastField<FSetProperty>(SetPropertyNodeB->GetProperty());
+		if (!SetPropertyA || !SetPropertyB)
+		{
+			return false;
+		}
+		
+		const FResolvedProperty ResolvedSetA = AsyncDetailViewDiffHelpers::GetResolvedProperty(SetPropertyNodeA, OwningObjectA);
+		const FResolvedProperty ResolvedSetB = AsyncDetailViewDiffHelpers::GetResolvedProperty(SetPropertyNodeB, OwningObjectB);
+		FScriptSetHelper SetHelperA(SetPropertyA, SetPropertyA->ContainerPtrToValuePtr<UObject*>(ResolvedSetA.Object));
+		FScriptSetHelper SetHelperB(SetPropertyB, SetPropertyB->ContainerPtrToValuePtr<UObject*>(ResolvedSetB.Object));
+		
+		const void* KeyA = SetHelperA.FindNthElementPtr(KeyIndexA);
+		const void* KeyB = SetHelperB.FindNthElementPtr(KeyIndexB);
+		
+		if (SetPropertyA->ElementProp->SameType(SetPropertyB->ElementProp))
 		{
-			return FPropertySoftPath(*PropertyPath).Resolve(Object);
+			return SetPropertyA->ElementProp->Identical(KeyA, KeyB, PPF_DeepComparison);
 		}
+		return false;
 	}
-	return FResolvedProperty();
 }
 
 
+
+
 bool TTreeDiffSpecification<TWeakPtr<FDetailTreeNode>>::AreValuesEqual(const TWeakPtr<FDetailTreeNode>& TreeNodeA, const TWeakPtr<FDetailTreeNode>& TreeNodeB) const
 {
 	const TSharedPtr<FDetailTreeNode> PinnedTreeNodeA = TreeNodeA.Pin();
@@ -63,8 +118,8 @@ bool TTreeDiffSpecification<TWeakPtr<FDetailTreeNode>>::AreValuesEqual(const TWe
 		return true;
 	}
 	
-	const TArray<TWeakObjectPtr<UObject>> OwningObjectsA = GetObjects(PinnedTreeNodeA);
-	const TArray<TWeakObjectPtr<UObject>> OwningObjectsB = GetObjects(PinnedTreeNodeB);
+	const TArray<TWeakObjectPtr<UObject>> OwningObjectsA = AsyncDetailViewDiffHelpers::GetObjects(PinnedTreeNodeA);
+	const TArray<TWeakObjectPtr<UObject>> OwningObjectsB = AsyncDetailViewDiffHelpers::GetObjects(PinnedTreeNodeB);
 
 	if (!ensure(OwningObjectsA.Num() == DataValuesA.Num() && OwningObjectsB.Num() == DataValuesB.Num()))
 	{
@@ -73,56 +128,6 @@ bool TTreeDiffSpecification<TWeakPtr<FDetailTreeNode>>::AreValuesEqual(const TWe
 	return DiffUtils::Identical(PropertyHandleA, PropertyHandleB, OwningObjectsA, OwningObjectsB);
 }
 
-static bool MapKeysMatch(const TSharedRef<FPropertyNode>& MapPropertyNodeA, const TSharedRef<FPropertyNode>& MapPropertyNodeB, int32 KeyIndexA, int32 KeyIndexB,
-	const UObject* OwningObjectA, const UObject* OwningObjectB)
-{
-	const FMapProperty* MapPropertyA = CastField<FMapProperty>(MapPropertyNodeA->GetProperty());
-	const FMapProperty* MapPropertyB = CastField<FMapProperty>(MapPropertyNodeB->GetProperty());
-	if (!MapPropertyA || !MapPropertyB)
-	{
-		return false;
-	}
-	
-	const FResolvedProperty ResolvedMapA = GetResolvedProperty(MapPropertyNodeA, OwningObjectA);
-	const FResolvedProperty ResolvedMapB = GetResolvedProperty(MapPropertyNodeB, OwningObjectB);
-	FScriptMapHelper MapHelperA(MapPropertyA, MapPropertyA->ContainerPtrToValuePtr<UObject*>(ResolvedMapA.Object));
-	FScriptMapHelper MapHelperB(MapPropertyB, MapPropertyB->ContainerPtrToValuePtr<UObject*>(ResolvedMapB.Object));
-	
-	const void* KeyA = MapHelperA.FindNthKeyPtr(KeyIndexA);
-	const void* KeyB = MapHelperB.FindNthKeyPtr(KeyIndexB);
-	
-	if (MapPropertyA->KeyProp->SameType(MapPropertyB->KeyProp))
-	{
-		return MapPropertyA->KeyProp->Identical(KeyA, KeyB, PPF_DeepComparison);
-	}
-	return false;
-}
-
-static bool SetKeysMatch(const TSharedRef<FPropertyNode>& SetPropertyNodeA, const TSharedRef<FPropertyNode>& SetPropertyNodeB, int32 KeyIndexA, int32 KeyIndexB,
-	const UObject* OwningObjectA, const UObject* OwningObjectB)
-{
-	const FSetProperty* SetPropertyA = CastField<FSetProperty>(SetPropertyNodeA->GetProperty());
-	const FSetProperty* SetPropertyB = CastField<FSetProperty>(SetPropertyNodeB->GetProperty());
-	if (!SetPropertyA || !SetPropertyB)
-	{
-		return false;
-	}
-	
-	const FResolvedProperty ResolvedSetA = GetResolvedProperty(SetPropertyNodeA, OwningObjectA);
-	const FResolvedProperty ResolvedSetB = GetResolvedProperty(SetPropertyNodeB, OwningObjectB);
-	FScriptSetHelper SetHelperA(SetPropertyA, SetPropertyA->ContainerPtrToValuePtr<UObject*>(ResolvedSetA.Object));
-	FScriptSetHelper SetHelperB(SetPropertyB, SetPropertyB->ContainerPtrToValuePtr<UObject*>(ResolvedSetB.Object));
-	
-	const void* KeyA = SetHelperA.FindNthElementPtr(KeyIndexA);
-	const void* KeyB = SetHelperB.FindNthElementPtr(KeyIndexB);
-	
-	if (SetPropertyA->ElementProp->SameType(SetPropertyB->ElementProp))
-	{
-		return SetPropertyA->ElementProp->Identical(KeyA, KeyB, PPF_DeepComparison);
-	}
-	return false;
-}
-
 bool TTreeDiffSpecification<TWeakPtr<FDetailTreeNode>>::AreMatching(const TWeakPtr<FDetailTreeNode>& TreeNodeA, const TWeakPtr<FDetailTreeNode>& TreeNodeB) const
 {
 	const TSharedPtr<FDetailTreeNode> PinnedTreeNodeA = TreeNodeA.Pin();
@@ -147,8 +152,8 @@ bool TTreeDiffSpecification<TWeakPtr<FDetailTreeNode>>::AreMatching(const TWeakP
 		return PinnedTreeNodeA->GetNodeName() == PinnedTreeNodeB->GetNodeName();
 	}
 
-	const TArray<TWeakObjectPtr<UObject>> OwningObjectsA = GetObjects(PinnedTreeNodeA);
-	const TArray<TWeakObjectPtr<UObject>> OwningObjectsB = GetObjects(PinnedTreeNodeB);
+	const TArray<TWeakObjectPtr<UObject>> OwningObjectsA = AsyncDetailViewDiffHelpers::GetObjects(PinnedTreeNodeA);
+	const TArray<TWeakObjectPtr<UObject>> OwningObjectsB = AsyncDetailViewDiffHelpers::GetObjects(PinnedTreeNodeB);
 	if (OwningObjectsA.IsEmpty() || OwningObjectsB.IsEmpty())
 	{
 		return PinnedTreeNodeA->GetNodeName() == PinnedTreeNodeB->GetNodeName();
@@ -181,14 +186,14 @@ bool TTreeDiffSpecification<TWeakPtr<FDetailTreeNode>>::AreMatching(const TWeakP
 			// and compare those instead
 			if (ParentPropertyA->IsA<FMapProperty>() || ParentPropertyB->IsA<FMapProperty>())
 			{
-				if (!MapKeysMatch(ParentPropertyNodeA, ParentPropertyNodeB, ArrayIndexA, ArrayIndexB, OwningObjectA, OwningObjectB))
+				if (!AsyncDetailViewDiffHelpers::MapKeysMatch(ParentPropertyNodeA, ParentPropertyNodeB, ArrayIndexA, ArrayIndexB, OwningObjectA, OwningObjectB))
 				{
 					return false;
 				}
 			}
 			if (ParentPropertyA->IsA<FSetProperty>() || ParentPropertyB->IsA<FSetProperty>())
 			{
-				if (SetKeysMatch(ParentPropertyNodeA, ParentPropertyNodeB, ArrayIndexA, ArrayIndexB, OwningObjectA, OwningObjectB))
+				if (AsyncDetailViewDiffHelpers::SetKeysMatch(ParentPropertyNodeA, ParentPropertyNodeB, ArrayIndexA, ArrayIndexB, OwningObjectA, OwningObjectB))
 				{
 					return false;
 				}
diff --git a/Engine/Source/Editor/PropertyEditor/Private/SDetailsSplitter.cpp b/Engine/Source/Editor/PropertyEditor/Private/SDetailsSplitter.cpp
--- a/Engine/Source/Editor/PropertyEditor/Private/SDetailsSplitter.cpp
+++ b/Engine/Source/Editor/PropertyEditor/Private/SDetailsSplitter.cpp
@@ -8,354 +8,327 @@
 #include "Styling/StyleColors.h"
 #include "Framework/Application/SlateApplication.h"
 
-static const TArray<TWeakObjectPtr<UObject>>* GetObjects(const TSharedPtr<FDetailTreeNode>& TreeNode)
+namespace DetailsSplitterHelpers
 {
-	if (const IDetailsViewPrivate* DetailsView = TreeNode->GetDetailsView())
+	const TArray<TWeakObjectPtr<UObject>>* GetObjects(const TSharedPtr<FDetailTreeNode>& TreeNode)
 	{
-		return &DetailsView->GetSelectedObjects();
+		if (const IDetailsViewPrivate* DetailsView = TreeNode->GetDetailsView())
+		{
+			return &DetailsView->GetSelectedObjects();
+		}
+		return nullptr;
 	}
-	return nullptr;
-}
 
-// converts from a container helper like FScriptArrayHelper to a property like FArrayProperty
-template<typename HelperType>
-using TContainerPropertyType =
-	std::conditional_t<std::is_same_v<HelperType, FScriptArrayHelper>, FArrayProperty,
-	std::conditional_t<std::is_same_v<HelperType, FScriptMapHelper>, FMapProperty,
-	std::conditional_t<std::is_same_v<HelperType, FScriptSetHelper>, FSetProperty,
-	void>>>;
+	// converts from a container helper like FScriptArrayHelper to a property like FArrayProperty
+	template<typename HelperType>
+	using TContainerPropertyType =
+		std::conditional_t<std::is_same_v<HelperType, FScriptArrayHelper>, FArrayProperty,
+		std::conditional_t<std::is_same_v<HelperType, FScriptMapHelper>, FMapProperty,
+		std::conditional_t<std::is_same_v<HelperType, FScriptSetHelper>, FSetProperty,
+		void>>>;
 
-template<typename HelperType>
-static bool TryGetSourceContainer(TSharedPtr<FDetailTreeNode> DetailsNode, FPropertyNode*& OutPropNode, TArray<TUniquePtr<HelperType>>& OutContainerHelper)
-{
-	using ContainerPropertyType = TContainerPropertyType<HelperType>;
-	if (TryGetSourceContainer<ContainerPropertyType>(DetailsNode, OutPropNode))
+	template<typename ContainerPropertyType>
+	bool TryGetSourceContainer(TSharedPtr<FDetailTreeNode> DetailsNode, FPropertyNode*& OutPropNode)
 	{
-		const FPropertySoftPath SoftPropertyPath(FPropertyNode::CreatePropertyPath(OutPropNode->AsShared()).Get());
-		if (const TArray<TWeakObjectPtr<UObject>>* Objects = GetObjects(DetailsNode))
+		if ((OutPropNode = DetailsNode->GetPropertyNode().Get()) == nullptr)
 		{
-			for (const TWeakObjectPtr<UObject> WeakObject : *Objects)
-			{
-				if (const UObject* Object = WeakObject.Get())
-				{
-					const FResolvedProperty Resolved = SoftPropertyPath.Resolve(Object);
-					const ContainerPropertyType* ContainerProperty = CastFieldChecked<ContainerPropertyType>(Resolved.Property);
-					OutContainerHelper.Add(MakeUnique<HelperType>(ContainerProperty, ContainerProperty->template ContainerPtrToValuePtr<UObject*>(Resolved.Object)));
-				}
-			}
+			return false;
+		}
+		if ((OutPropNode = OutPropNode->GetParentNode()) == nullptr)
+		{
+			return false;
+		}
+		if (CastField<ContainerPropertyType>(OutPropNode->GetProperty()) != nullptr)
+		{
+			return true;
 		}
-		return true;
+		return false;
 	}
-	return false;
-}
 
-template<typename HelperType>
-static bool TryGetDestinationContainer(TSharedPtr<FDetailTreeNode> DetailsNode, FPropertyNode*& OutPropNode, TArray<TUniquePtr<HelperType>>& OutContainerHelper, int32& OutInsertIndex)
-{
-	using ContainerPropertyType = TContainerPropertyType<HelperType>;
-	if (TryGetDestinationContainer<ContainerPropertyType>(DetailsNode, OutPropNode, OutInsertIndex))
+	template<typename ContainerPropertyType>
+	bool TryGetDestinationContainer(TSharedPtr<FDetailTreeNode> DetailsNode, FPropertyNode*& OutPropNode, int32& OutInsertIndex)
 	{
-		const FPropertySoftPath SoftPropertyPath(FPropertyNode::CreatePropertyPath(OutPropNode->AsShared()).Get());
-		if (const TArray<TWeakObjectPtr<UObject>>* Objects = GetObjects(DetailsNode))
+		if ((OutPropNode = DetailsNode->GetPropertyNode().Get()) == nullptr)
+		{
+			return false;
+		}
+		if (CastField<ContainerPropertyType>(OutPropNode->GetProperty()) != nullptr)
+		{
+			OutInsertIndex = 0;
+			return true;
+		}
+		OutInsertIndex = OutPropNode->GetArrayIndex() + 1;
+		while((OutPropNode = OutPropNode->GetParentNode()) != nullptr)
 		{
-			for (const TWeakObjectPtr<UObject> WeakObject : *Objects)
+			if (CastField<ContainerPropertyType>(OutPropNode->GetProperty()) != nullptr)
 			{
-				if (const UObject* Object = WeakObject.Get())
-				{
-					const FResolvedProperty Resolved = SoftPropertyPath.Resolve(Object);
-					const ContainerPropertyType* ContainerProperty = CastFieldChecked<ContainerPropertyType>(Resolved.Property);
-					OutContainerHelper.Add(MakeUnique<HelperType>(ContainerProperty, ContainerProperty->template ContainerPtrToValuePtr<UObject*>(Resolved.Object)));
-				}
+				return true;
 			}
+			OutInsertIndex = OutPropNode->GetArrayIndex() + 1;
 		}
-		return true;
-	}
-	return false;
-}
-
-template<typename ContainerPropertyType>
-static bool TryGetSourceContainer(TSharedPtr<FDetailTreeNode> DetailsNode, FPropertyNode*& OutPropNode)
-{
-	if ((OutPropNode = DetailsNode->GetPropertyNode().Get()) == nullptr)
-	{
-		return false;
-	}
-	if ((OutPropNode = OutPropNode->GetParentNode()) == nullptr)
-	{
 		return false;
 	}
-	if (CastField<ContainerPropertyType>(OutPropNode->GetProperty()) != nullptr)
-	{
-		return true;
-	}
-	return false;
-}
 
-template<typename ContainerPropertyType>
-static bool TryGetDestinationContainer(TSharedPtr<FDetailTreeNode> DetailsNode, FPropertyNode*& OutPropNode, int32& OutInsertIndex)
-{
-	if ((OutPropNode = DetailsNode->GetPropertyNode().Get()) == nullptr)
-	{
-		return false;
-	}
-	if (CastField<ContainerPropertyType>(OutPropNode->GetProperty()) != nullptr)
+	template<typename HelperType>
+	bool TryGetSourceContainer(TSharedPtr<FDetailTreeNode> DetailsNode, FPropertyNode*& OutPropNode, TArray<TUniquePtr<HelperType>>& OutContainerHelper)
 	{
-		OutInsertIndex = 0;
-		return true;
-	}
-	OutInsertIndex = OutPropNode->GetArrayIndex() + 1;
-	while((OutPropNode = OutPropNode->GetParentNode()) != nullptr)
-	{
-		if (CastField<ContainerPropertyType>(OutPropNode->GetProperty()) != nullptr)
+		using ContainerPropertyType = TContainerPropertyType<HelperType>;
+		if (TryGetSourceContainer<ContainerPropertyType>(DetailsNode, OutPropNode))
 		{
+			const FPropertySoftPath SoftPropertyPath(FPropertyNode::CreatePropertyPath(OutPropNode->AsShared()).Get());
+			if (const TArray<TWeakObjectPtr<UObject>>* Objects = DetailsSplitterHelpers::GetObjects(DetailsNode))
+			{
+				for (const TWeakObjectPtr<UObject> WeakObject : *Objects)
+				{
+					if (const UObject* Object = WeakObject.Get())
+					{
+						const FResolvedProperty Resolved = SoftPropertyPath.Resolve(Object);
+						const ContainerPropertyType* ContainerProperty = CastFieldChecked<ContainerPropertyType>(Resolved.Property);
+						OutContainerHelper.Add(MakeUnique<HelperType>(ContainerProperty, ContainerProperty->template ContainerPtrToValuePtr<UObject*>(Resolved.Object)));
+					}
+				}
+			}
 			return true;
 		}
-		OutInsertIndex = OutPropNode->GetArrayIndex() + 1;
+		return false;
 	}
-	return false;
-}
 
-static void SendPropertyChangeNotification(FPropertyNode* PropertyNode, UObject* OuterObject, EPropertyChangeType::Type ChangeType)
-{
-	// find an object to send the FPropertyChangedEvent to
-	FProperty* ChangedProperty = PropertyNode->GetProperty();
-	PropertyNode = PropertyNode->GetParentNode();
-	while(PropertyNode && !CastField<FObjectProperty>(PropertyNode->GetProperty()))
+	template<typename HelperType>
+	bool TryGetDestinationContainer(TSharedPtr<FDetailTreeNode> DetailsNode, FPropertyNode*& OutPropNode, TArray<TUniquePtr<HelperType>>& OutContainerHelper, int32& OutInsertIndex)
 	{
-		if (FProperty* Property = PropertyNode->GetProperty())
+		using ContainerPropertyType = TContainerPropertyType<HelperType>;
+		if (TryGetDestinationContainer<ContainerPropertyType>(DetailsNode, OutPropNode, OutInsertIndex))
 		{
-			ChangedProperty = Property;
+			const FPropertySoftPath SoftPropertyPath(FPropertyNode::CreatePropertyPath(OutPropNode->AsShared()).Get());
+			if (const TArray<TWeakObjectPtr<UObject>>* Objects = DetailsSplitterHelpers::GetObjects(DetailsNode))
+			{
+				for (const TWeakObjectPtr<UObject> WeakObject : *Objects)
+				{
+					if (const UObject* Object = WeakObject.Get())
+					{
+						const FResolvedProperty Resolved = SoftPropertyPath.Resolve(Object);
+						const ContainerPropertyType* ContainerProperty = CastFieldChecked<ContainerPropertyType>(Resolved.Property);
+						OutContainerHelper.Add(MakeUnique<HelperType>(ContainerProperty, ContainerProperty->template ContainerPtrToValuePtr<UObject*>(Resolved.Object)));
+					}
+				}
+			}
+			return true;
 		}
-		PropertyNode = PropertyNode->GetParentNode();
+		return false;
 	}
 
-	// send the notification
-	FPropertyChangedEvent ChangeEvent(ChangedProperty, ChangeType);
-	if (PropertyNode)
+	void CopyPropertyValueForInsert(const TSharedPtr<FDetailTreeNode>& SourceDetailsNode, const TSharedPtr<FDetailTreeNode>& DestinationDetailsNode)
 	{
-		const FObjectProperty* AsObjectProperty = CastFieldChecked<FObjectProperty>(PropertyNode->GetProperty());
-		const FPropertySoftPath SoftPropertyPath(FPropertyNode::CreatePropertyPath(PropertyNode->AsShared()).Get());
-		const FResolvedProperty Resolved = SoftPropertyPath.Resolve(OuterObject);
-		UObject* InnerObject = AsObjectProperty->GetObjectPropertyValue(Resolved.Property->ContainerPtrToValuePtr<void>(Resolved.Object));
-		InnerObject->PostEditChangeProperty(ChangeEvent);
-	}
-	else
-	{
-		// if an FObjectProperty wasn't found, this property must be owned by OuterObject directly
-		OuterObject->PostEditChangeProperty(ChangeEvent);
-	}
-}
-
-static void CopyPropertyValueForInsert(const TSharedPtr<FDetailTreeNode>& SourceDetailsNode, const TSharedPtr<FDetailTreeNode>& DestinationDetailsNode)
-{
-	const TSharedPtr<IPropertyHandle> DestinationHandle = DestinationDetailsNode->CreatePropertyHandle();
+		const TSharedPtr<IPropertyHandle> DestinationHandle = DestinationDetailsNode->CreatePropertyHandle();
 
-	// Array
-	TArray<TUniquePtr<FScriptArrayHelper>> SourceArrays;
-	FPropertyNode* SourceArrayPropertyNode;
-	if (TryGetSourceContainer(SourceDetailsNode, SourceArrayPropertyNode, SourceArrays))
-	{
-		int32 InsertIndex;
-		TArray<TUniquePtr<FScriptArrayHelper>> DestinationArrays;
-		FPropertyNode* DestinationArrayPropertyNode;
-		if (ensure(TryGetDestinationContainer(DestinationDetailsNode, DestinationArrayPropertyNode, DestinationArrays, InsertIndex)))
-		{
-			ensure(SourceArrays.Num() == DestinationArrays.Num());
-			DestinationHandle->NotifyPreChange();
-			for (int32 ArrayNum = 0; ArrayNum < SourceArrays.Num(); ++ ArrayNum)
+		// Array
+		TArray<TUniquePtr<FScriptArrayHelper>> SourceArrays;
+		FPropertyNode* SourceArrayPropertyNode;
+		if (TryGetSourceContainer(SourceDetailsNode, SourceArrayPropertyNode, SourceArrays))
+		{
+			int32 InsertIndex;
+			TArray<TUniquePtr<FScriptArrayHelper>> DestinationArrays;
+			FPropertyNode* DestinationArrayPropertyNode;
+			if (ensure(TryGetDestinationContainer(DestinationDetailsNode, DestinationArrayPropertyNode, DestinationArrays, InsertIndex)))
 			{
-				DestinationArrays[ArrayNum]->InsertValues(InsertIndex, 1);
-				const void* SourceData = SourceArrays[ArrayNum]->GetElementPtr(SourceDetailsNode->GetPropertyNode()->GetArrayIndex());
-				void* DestinationData = DestinationArrays[ArrayNum]->GetElementPtr(InsertIndex);
-				const FArrayProperty* ArrayProperty = CastFieldChecked<FArrayProperty>(DestinationArrayPropertyNode->GetProperty());
-				const FProperty* ElementProperty = ArrayProperty->Inner;
-				ElementProperty->CopySingleValue(DestinationData, SourceData);
+				ensure(SourceArrays.Num() == DestinationArrays.Num());
+				DestinationHandle->NotifyPreChange();
+				for (int32 ArrayNum = 0; ArrayNum < SourceArrays.Num(); ++ ArrayNum)
+				{
+					DestinationArrays[ArrayNum]->InsertValues(InsertIndex, 1);
+					const void* SourceData = SourceArrays[ArrayNum]->GetElementPtr(SourceDetailsNode->GetPropertyNode()->GetArrayIndex());
+					void* DestinationData = DestinationArrays[ArrayNum]->GetElementPtr(InsertIndex);
+					const FArrayProperty* ArrayProperty = CastFieldChecked<FArrayProperty>(DestinationArrayPropertyNode->GetProperty());
+					const FProperty* ElementProperty = ArrayProperty->Inner;
+					ElementProperty->CopySingleValue(DestinationData, SourceData);
+				}
+				
+				DestinationHandle->NotifyPostChange(EPropertyChangeType::ArrayAdd);
+				DestinationHandle->NotifyFinishedChangingProperties();
 			}
-			
-			DestinationHandle->NotifyPostChange(EPropertyChangeType::ArrayAdd);
-			DestinationHandle->NotifyFinishedChangingProperties();
+			return;
 		}
-		return;
-	}
-	
-	// Set
-	TArray<TUniquePtr<FScriptSetHelper>> SourceSets;
-	FPropertyNode* SourceSetPropertyNode;
-	if (TryGetSourceContainer(SourceDetailsNode, SourceSetPropertyNode, SourceSets))
-	{
-		int32 InsertIndex;
-		TArray<TUniquePtr<FScriptSetHelper>> DestinationSets;
-		FPropertyNode* DestinationPropertyNode;
-		if (ensure(TryGetDestinationContainer(DestinationDetailsNode, DestinationPropertyNode, DestinationSets, InsertIndex)))
-		{
-			ensure(SourceSets.Num() == DestinationSets.Num());
-			DestinationHandle->NotifyPreChange();
-			for (int32 SetNum = 0; SetNum < SourceSets.Num(); ++ SetNum)
+		
+		// Set
+		TArray<TUniquePtr<FScriptSetHelper>> SourceSets;
+		FPropertyNode* SourceSetPropertyNode;
+		if (TryGetSourceContainer(SourceDetailsNode, SourceSetPropertyNode, SourceSets))
+		{
+			int32 InsertIndex;
+			TArray<TUniquePtr<FScriptSetHelper>> DestinationSets;
+			FPropertyNode* DestinationPropertyNode;
+			if (ensure(TryGetDestinationContainer(DestinationDetailsNode, DestinationPropertyNode, DestinationSets, InsertIndex)))
 			{
-				const void* SourceData = SourceSets[SetNum]->FindNthElementPtr(SourceDetailsNode->GetPropertyNode()->GetArrayIndex());
-				DestinationSets[SetNum]->AddElement(SourceData);
+				ensure(SourceSets.Num() == DestinationSets.Num());
+				DestinationHandle->NotifyPreChange();
+				for (int32 SetNum = 0; SetNum < SourceSets.Num(); ++ SetNum)
+				{
+					const void* SourceData = SourceSets[SetNum]->FindNthElementPtr(SourceDetailsNode->GetPropertyNode()->GetArrayIndex());
+					DestinationSets[SetNum]->AddElement(SourceData);
+				}
+				
+				DestinationHandle->NotifyPostChange(EPropertyChangeType::ArrayAdd);
+				DestinationHandle->NotifyFinishedChangingProperties();
 			}
-			
-			DestinationHandle->NotifyPostChange(EPropertyChangeType::ArrayAdd);
-			DestinationHandle->NotifyFinishedChangingProperties();
+			return;
 		}
-		return;
-	}
-	
-	// Map
-	TArray<TUniquePtr<FScriptMapHelper>> SourceMaps;
-	FPropertyNode* SourceMapPropertyNode;
-	if (TryGetSourceContainer(SourceDetailsNode, SourceMapPropertyNode, SourceMaps))
-	{
-		int32 InsertIndex;
-		TArray<TUniquePtr<FScriptMapHelper>> DestinationMaps;
-		FPropertyNode* DestinationPropertyNode;
-		if (ensure(TryGetDestinationContainer(DestinationDetailsNode, DestinationPropertyNode, DestinationMaps, InsertIndex)))
-		{
-			ensure(SourceMaps.Num() == DestinationMaps.Num());
-			DestinationHandle->NotifyPreChange();
-			for (int32 MapNum = 0; MapNum < SourceMaps.Num(); ++ MapNum)
+		
+		// Map
+		TArray<TUniquePtr<FScriptMapHelper>> SourceMaps;
+		FPropertyNode* SourceMapPropertyNode;
+		if (TryGetSourceContainer(SourceDetailsNode, SourceMapPropertyNode, SourceMaps))
+		{
+			int32 InsertIndex;
+			TArray<TUniquePtr<FScriptMapHelper>> DestinationMaps;
+			FPropertyNode* DestinationPropertyNode;
+			if (ensure(TryGetDestinationContainer(DestinationDetailsNode, DestinationPropertyNode, DestinationMaps, InsertIndex)))
 			{
-				const int32 Index = SourceMaps[MapNum]->FindInternalIndex(SourceDetailsNode->GetPropertyNode()->GetArrayIndex());
-				const void* SourceKey = SourceMaps[MapNum]->GetKeyPtr(Index);
-				const void* SourceVal = SourceMaps[MapNum]->GetValuePtr(Index);
-				DestinationMaps[MapNum]->AddPair(SourceKey, SourceVal);
+				ensure(SourceMaps.Num() == DestinationMaps.Num());
+				DestinationHandle->NotifyPreChange();
+				for (int32 MapNum = 0; MapNum < SourceMaps.Num(); ++ MapNum)
+				{
+					const int32 Index = SourceMaps[MapNum]->FindInternalIndex(SourceDetailsNode->GetPropertyNode()->GetArrayIndex());
+					const void* SourceKey = SourceMaps[MapNum]->GetKeyPtr(Index);
+					const void* SourceVal = SourceMaps[MapNum]->GetValuePtr(Index);
+					DestinationMaps[MapNum]->AddPair(SourceKey, SourceVal);
+				}
+				
+				DestinationHandle->NotifyPostChange(EPropertyChangeType::ArrayAdd);
+				DestinationHandle->NotifyFinishedChangingProperties();
 			}
-			
-			DestinationHandle->NotifyPostChange(EPropertyChangeType::ArrayAdd);
-			DestinationHandle->NotifyFinishedChangingProperties();
+			return;
 		}
-		return;
 	}
-}
 
-static void CopyPropertyValue(const TSharedPtr<FDetailTreeNode>& SourceDetailsNode, const TSharedPtr<FDetailTreeNode>& DestinationDetailsNode, ETreeDiffResult Diff)
-{
-	switch(Diff)
+	void CopyPropertyValue(const TSharedPtr<FDetailTreeNode>& SourceDetailsNode, const TSharedPtr<FDetailTreeNode>& DestinationDetailsNode, ETreeDiffResult Diff)
 	{
-		// traditional copy
-		case ETreeDiffResult::DifferentValues:
-			break;
+		switch(Diff)
+		{
+			// traditional copy
+			case ETreeDiffResult::DifferentValues:
+				break;
 
-		// insert
-		case ETreeDiffResult::MissingFromTree1:
-		case ETreeDiffResult::MissingFromTree2:
-			CopyPropertyValueForInsert(SourceDetailsNode, DestinationDetailsNode);
+			// insert
+			case ETreeDiffResult::MissingFromTree1:
+			case ETreeDiffResult::MissingFromTree2:
+				CopyPropertyValueForInsert(SourceDetailsNode, DestinationDetailsNode);
+				return;
+			
+			// no difference
+			case ETreeDiffResult::Invalid:
+			case ETreeDiffResult::Identical:
+			default:
+				return;
+		}
+
+		TArray<void*> SourceData;
+		TArray<void*> DestinationData;
+		const TSharedPtr<IPropertyHandle> SourceHandle = SourceDetailsNode->CreatePropertyHandle();
+		const TSharedPtr<IPropertyHandle> DestinationHandle = DestinationDetailsNode->CreatePropertyHandle();
+		if (!ensure(SourceHandle && DestinationHandle))
+		{
 			return;
-		
-		// no difference
-		case ETreeDiffResult::Invalid:
-		case ETreeDiffResult::Identical:
-		default:
+		}
+		if (!ensure(SourceHandle->GetProperty()->SameType(DestinationHandle->GetProperty())))
+		{
 			return;
-	}
+		}
+		SourceHandle->AccessRawData(SourceData);
+		DestinationHandle->AccessRawData(DestinationData);
+		if (!ensure(SourceData.Num() == DestinationData.Num()))
+		{
+			return;
+		}
 
-	TArray<void*> SourceData;
-	TArray<void*> DestinationData;
-	const TSharedPtr<IPropertyHandle> SourceHandle = SourceDetailsNode->CreatePropertyHandle();
-	const TSharedPtr<IPropertyHandle> DestinationHandle = DestinationDetailsNode->CreatePropertyHandle();
-	if (!ensure(SourceHandle && DestinationHandle))
-	{
-		return;
-	}
-	if (!ensure(SourceHandle->GetProperty()->SameType(DestinationHandle->GetProperty())))
-	{
-		return;
-	}
-	SourceHandle->AccessRawData(SourceData);
-	DestinationHandle->AccessRawData(DestinationData);
-	if (!ensure(SourceData.Num() == DestinationData.Num()))
-	{
-		return;
+		
+		DestinationHandle->NotifyPreChange();
+		for (int32 I = 0; I < SourceData.Num(); ++I)
+		{
+			if (DestinationHandle->GetArrayIndex() != INDEX_NONE)
+			{
+				DestinationHandle->GetProperty()->CopySingleValue(DestinationData[I], SourceData[I]);
+			}
+			else
+			{
+				DestinationHandle->GetProperty()->CopyCompleteValue(DestinationData[I], SourceData[I]);
+			}
+		}
+		DestinationHandle->NotifyPostChange(EPropertyChangeType::ValueSet);
+		DestinationHandle->NotifyFinishedChangingProperties();
 	}
 
-	
-	DestinationHandle->NotifyPreChange();
-	for (int32 I = 0; I < SourceData.Num(); ++I)
+	// note: DestinationDetailsNode is the node before the position in the tree you wish to insert
+	bool CanCopyPropertyValueForInsert(const TSharedPtr<FDetailTreeNode>& SourceDetailsNode, const TSharedPtr<FDetailTreeNode>& DestinationDetailsNode)
 	{
-		if (DestinationHandle->GetArrayIndex() != INDEX_NONE)
+		FPropertyNode* SourceArrayPropertyNode;
+		if (TryGetSourceContainer<FArrayProperty>(SourceDetailsNode, SourceArrayPropertyNode))
 		{
-			DestinationHandle->GetProperty()->CopySingleValue(DestinationData[I], SourceData[I]);
+			FPropertyNode* DestinationArrayPropertyNode;
+			int32 InsertIndex;
+			if (TryGetDestinationContainer<FArrayProperty>(DestinationDetailsNode, DestinationArrayPropertyNode, InsertIndex))
+			{
+				return true;
+			}
+			return false;
 		}
-		else
+		
+		FPropertyNode* SourceSetPropertyNode;
+		if (TryGetSourceContainer<FSetProperty>(SourceDetailsNode, SourceSetPropertyNode))
 		{
-			DestinationHandle->GetProperty()->CopyCompleteValue(DestinationData[I], SourceData[I]);
+			FPropertyNode* DestinationSetPropertyNode;
+			int32 InsertIndex;
+			if (TryGetDestinationContainer<FSetProperty>(DestinationDetailsNode, DestinationSetPropertyNode, InsertIndex))
+			{
+				return true;
+			}
+			return false;
 		}
-	}
-	DestinationHandle->NotifyPostChange(EPropertyChangeType::ValueSet);
-	DestinationHandle->NotifyFinishedChangingProperties();
-}
-
-// note: DestinationDetailsNode is the node before the position in the tree you wish to insert
-static bool CanCopyPropertyValueForInsert(const TSharedPtr<FDetailTreeNode>& SourceDetailsNode, const TSharedPtr<FDetailTreeNode>& DestinationDetailsNode)
-{
-	FPropertyNode* SourceArrayPropertyNode;
-	if (TryGetSourceContainer<FArrayProperty>(SourceDetailsNode, SourceArrayPropertyNode))
-	{
-		FPropertyNode* DestinationArrayPropertyNode;
-		int32 InsertIndex;
-		if (TryGetDestinationContainer<FArrayProperty>(DestinationDetailsNode, DestinationArrayPropertyNode, InsertIndex))
+		
+		FPropertyNode* SourceMapPropertyNode;
+		if (TryGetSourceContainer<FMapProperty>(SourceDetailsNode, SourceMapPropertyNode))
 		{
-			return true;
+			FPropertyNode* DestinationMapPropertyNode;
+			int32 InsertIndex;
+			if (TryGetDestinationContainer<FMapProperty>(DestinationDetailsNode, DestinationMapPropertyNode, InsertIndex))
+			{
+				return true;
+			}
+			return false;
 		}
+
+		// you can only insert into containers
 		return false;
 	}
-	
-	FPropertyNode* SourceSetPropertyNode;
-	if (TryGetSourceContainer<FSetProperty>(SourceDetailsNode, SourceSetPropertyNode))
+
+	bool CanCopyPropertyValue(const TSharedPtr<FDetailTreeNode>& SourceDetailsNode, const TSharedPtr<FDetailTreeNode>& DestinationDetailsNode, ETreeDiffResult Diff)
 	{
-		FPropertyNode* DestinationSetPropertyNode;
-		int32 InsertIndex;
-		if (TryGetDestinationContainer<FSetProperty>(DestinationDetailsNode, DestinationSetPropertyNode, InsertIndex))
+		// in order to copy properties there needs to be the same number of objects in each panel
+		const TArray<TWeakObjectPtr<UObject>>* SourceObjects = DetailsSplitterHelpers::GetObjects(SourceDetailsNode);
+		const TArray<TWeakObjectPtr<UObject>>* DestinationObjects = DetailsSplitterHelpers::GetObjects(DestinationDetailsNode);
+		if (!SourceObjects || !DestinationObjects || SourceObjects->Num() != DestinationObjects->Num())
 		{
-			return true;
+			return false;
 		}
-		return false;
-	}
-	
-	FPropertyNode* SourceMapPropertyNode;
-	if (TryGetSourceContainer<FMapProperty>(SourceDetailsNode, SourceMapPropertyNode))
-	{
-		FPropertyNode* DestinationMapPropertyNode;
-		int32 InsertIndex;
-		if (TryGetDestinationContainer<FMapProperty>(DestinationDetailsNode, DestinationMapPropertyNode, InsertIndex))
+		
+		switch (Diff)
 		{
+		// traditional copy
+		case ETreeDiffResult::DifferentValues:
 			return true;
-		}
-		return false;
-	}
 
-	// you can only insert into containers
-	return false;
-}
+		// insert
+		case ETreeDiffResult::MissingFromTree1:
+		case ETreeDiffResult::MissingFromTree2:
+			return CanCopyPropertyValueForInsert(SourceDetailsNode, DestinationDetailsNode);
 
-static bool CanCopyPropertyValue(const TSharedPtr<FDetailTreeNode>& SourceDetailsNode, const TSharedPtr<FDetailTreeNode>& DestinationDetailsNode, ETreeDiffResult Diff)
-{
-	// in order to copy properties there needs to be the same number of objects in each panel
-	const TArray<TWeakObjectPtr<UObject>>* SourceObjects = GetObjects(SourceDetailsNode);
-	const TArray<TWeakObjectPtr<UObject>>* DestinationObjects = GetObjects(DestinationDetailsNode);
-	if (!SourceObjects || !DestinationObjects || SourceObjects->Num() != DestinationObjects->Num())
-	{
-		return false;
+		// no difference
+		case ETreeDiffResult::Invalid:
+		case ETreeDiffResult::Identical:
+		default:
+			return false;
+		}
 	}
 	
-	switch (Diff)
-	{
-	// traditional copy
-	case ETreeDiffResult::DifferentValues:
-		return true;
-
-	// insert
-	case ETreeDiffResult::MissingFromTree1:
-	case ETreeDiffResult::MissingFromTree2:
-		return CanCopyPropertyValueForInsert(SourceDetailsNode, DestinationDetailsNode);
-
-	// no difference
-	case ETreeDiffResult::Invalid:
-	case ETreeDiffResult::Identical:
-	default:
-		return false;
-	}
 }
 
 void SDetailsSplitter::Construct(const FArguments& InArgs)
@@ -564,11 +537,11 @@ int32 SDetailsSplitter::OnPaint(const FPaintArgs& Args, const FGeometry& Allotte
 							PaintPropertyConnector(OutDrawElements, MaxLayerId, LeftPropertyRect, RightPropertyRect, FillColor, OutlineColor);
 							++MaxLayerId;
 
-							if (!RightPanel.IsReadonly.Get(true) && CanCopyPropertyValue(LastSeenLeftDetailsNode, LastSeenRightDetailsNode, DiffNode->DiffResult))
+							if (!RightPanel.IsReadonly.Get(true) && DetailsSplitterHelpers::CanCopyPropertyValue(LastSeenLeftDetailsNode, LastSeenRightDetailsNode, DiffNode->DiffResult))
 							{
 								PaintCopyPropertyButton(OutDrawElements, MaxLayerId, DiffNode, LeftPropertyRect, RightPropertyRect, EPropertyCopyDirection::CopyLeftToRight);
 							}
-							if (!LeftPanel.IsReadonly.Get(true) && CanCopyPropertyValue(LastSeenRightDetailsNode, LastSeenLeftDetailsNode, DiffNode->DiffResult))
+							if (!LeftPanel.IsReadonly.Get(true) && DetailsSplitterHelpers::CanCopyPropertyValue(LastSeenRightDetailsNode, LastSeenLeftDetailsNode, DiffNode->DiffResult))
 							{
 								PaintCopyPropertyButton(OutDrawElements, MaxLayerId, DiffNode, LeftPropertyRect, RightPropertyRect, EPropertyCopyDirection::CopyRightToLeft);
 							}
@@ -677,7 +650,7 @@ FReply SDetailsSplitter::OnMouseMove(const FGeometry& MyGeometry, const FPointer
 
 					if (CopyButtonZoneLeftToRight.ContainsPoint(MousePosition))
 					{
-						if (CanCopyPropertyValue(LastSeenLeftDetailsNode, LastSeenRightDetailsNode, DiffNode->DiffResult))
+						if (DetailsSplitterHelpers::CanCopyPropertyValue(LastSeenLeftDetailsNode, LastSeenRightDetailsNode, DiffNode->DiffResult))
 						{
 							HoveredCopyButton = {
 								LastSeenLeftDetailsNode,
@@ -701,7 +674,7 @@ FReply SDetailsSplitter::OnMouseMove(const FGeometry& MyGeometry, const FPointer
 
 					if (CopyButtonZoneRightToLeft.ContainsPoint(MousePosition))
 					{
-						if (CanCopyPropertyValue(LastSeenRightDetailsNode, LastSeenLeftDetailsNode, DiffNode->DiffResult))
+						if (DetailsSplitterHelpers::CanCopyPropertyValue(LastSeenRightDetailsNode, LastSeenLeftDetailsNode, DiffNode->DiffResult))
 						{
 							HoveredCopyButton = {
 								LastSeenRightDetailsNode,
@@ -744,7 +717,7 @@ FReply SDetailsSplitter::OnMouseButtonDown(const FGeometry& MyGeometry, const FP
 		return FReply::Unhandled();
 	}
 
-	CopyPropertyValue(
+	DetailsSplitterHelpers::CopyPropertyValue(
 		HoveredCopyButton.SourceDetailsNode.Pin(),
 		HoveredCopyButton.DestinationDetailsNode.Pin(),
 		HoveredCopyButton.DiffResult);
