The editor crashes once the user opens the "Editor preferences" window on Linux OS.

The code in DetailLayoutBuilderImpl.cpp that is attempting to create a new category is relying on a value from a FProperty that is not necessarily initialized and should not be relied upon. Instead, use information from the UObject to initialize the Category and it's UProperty

diff --git a/Engine/Source/Editor/PropertyEditor/Private/DetailLayoutBuilderImpl.cpp b/Engine/Source/Editor/PropertyEditor/Private/DetailLayoutBuilderImpl.cpp
--- a/Engine/Source/Editor/PropertyEditor/Private/DetailLayoutBuilderImpl.cpp
+++ b/Engine/Source/Editor/PropertyEditor/Private/DetailLayoutBuilderImpl.cpp
@@ -1059,28 +1059,15 @@ bool FDetailLayoutBuilderImpl::IsPropertyPathAllowed(const FString& InPath) cons
 
 bool FDetailLayoutBuilderImpl::AddEmptyCategoryIfNeeded(TSharedPtr<FComplexPropertyNode> Node)
 {
-	if (DefaultCategoryMap.IsEmpty())
+	const bool bHasNoValidCategories = DefaultCategoryMap.IsEmpty();
+	const bool bHasValidDisplayManager = DetailsView && DetailsView->GetDisplayManager().IsValid();
+	const bool bHasValidPropertyNode = Node.IsValid();
+	
+	if ( bHasNoValidCategories &&
+		 bHasValidDisplayManager &&
+		 bHasValidPropertyNode )
 	{
-		if (Node.IsValid() &&
-			Node->AsObjectNode() &&
-			Node->AsObjectNode()->GetNumObjects() > 0)
-		{
-			const FName PropertyCategoryName = FObjectEditorUtils::GetCategoryFName(Node->GetProperty());
-			FDetailCategoryImpl& Category = DefaultCategory(PropertyCategoryName);
-			const FName InstanceName = Node->GetProperty()->GetFName();
-			
-			FObjectPropertyNode* ObjectPropertyNode =  Node->AsObjectNode();
-			UObject* OutObject = ObjectPropertyNode->GetUObject(0);
-
-			if (DetailsView &&
-				DetailsView->GetDisplayManager().IsValid() &&
-				DetailsView->GetDisplayManager()->ShowEmptyCategoryIfRootUObjectHasNoPropertyData(OutObject))
-			{
-				Category.AddPropertyNode(Node.ToSharedRef(), InstanceName);
-				Category.SetIsEmpty(true);
-				return true;
-			}
-		}
+ 		return DetailsView->GetDisplayManager()->AddEmptyCategoryToDetailLayoutIfNeeded(Node.ToSharedRef(), SharedThis(this));
 	}
 	return false;
 }
diff --git a/Engine/Source/Editor/PropertyEditor/Private/DetailsDisplayManager.cpp b/Engine/Source/Editor/PropertyEditor/Private/DetailsDisplayManager.cpp
--- a/Engine/Source/Editor/PropertyEditor/Private/DetailsDisplayManager.cpp
+++ b/Engine/Source/Editor/PropertyEditor/Private/DetailsDisplayManager.cpp
@@ -4,6 +4,11 @@
 
 #include "DetailsViewStyle.h"
 #include "SDetailsView.h"
+#include "DetailLayoutBuilderImpl.h"
+#include "DetailCategoryBuilderImpl.h"
+#include "ObjectEditorUtils.h"
+#include "ObjectPropertyNode.h"
+
 
 static TAutoConsoleVariable<bool> CVarForceShowComponentEditor(
 	TEXT("CoreEntity.UI.ForceShowComponentEditor"),
@@ -73,6 +78,45 @@ bool FDetailsDisplayManager::ShowEmptyCategoryIfRootUObjectHasNoPropertyData(UOb
 	return false;
 }
 
+bool FDetailsDisplayManager::AddEmptyCategoryToDetailLayoutIfNeeded(TSharedRef<FComplexPropertyNode> Node,
+                                                      TSharedRef<FDetailLayoutBuilderImpl> DetailLayoutBuilder)
+{
+	if (FObjectPropertyNode* ObjectPropertyNode = Node->AsObjectNode())
+	{
+		/* a pointer to the UObject for this Category, if we have one */
+		UObject* UObjectForCategory = nullptr;
+
+		const bool bPropertyNodeHasOneUObject =  ObjectPropertyNode && ObjectPropertyNode->GetNumObjects() == 1;
+		if ( bPropertyNodeHasOneUObject )
+		{
+			UObjectForCategory = ObjectPropertyNode->GetUObject(0);
+		}
+
+		if (UObjectForCategory &&
+			/* we're supposed to show an empty category for this particular object if it has no UProperties */
+			ShowEmptyCategoryIfRootUObjectHasNoPropertyData(UObjectForCategory) &&
+			UObjectForCategory->GetClass())
+		{
+			/* Base empty category display text off the display text of the UObject class */
+			const FText UObjectClassNameDisplayText = UObjectForCategory->GetClass()->GetDisplayNameText();
+			
+			/* Base empty category name off the name of the UObject class */
+			const FName UObjectClassName = UObjectForCategory->GetClass()->GetFName();
+			
+			FDetailCategoryImpl& EmptyCategory = DetailLayoutBuilder->DefaultCategory( UObjectClassName );
+			EmptyCategory.SetDisplayName(UObjectClassName, UObjectClassNameDisplayText);
+
+			/* property node for now will always show empty ~ eventually we will want to add hooks to optionally show
+			 * a "No properties here" or "This empty component does XYZ" string of some kind that can be overridden
+			 * based on object type and usage */
+			EmptyCategory.AddPropertyNode(Node, UObjectClassName);
+			EmptyCategory.SetIsEmpty(true);
+			return true;
+		}
+	}
+	return false;
+}
+
 bool FDetailsDisplayManager::GetIsScrollBarNeeded() const
 {
 	return bIsScrollBarNeeded;
diff --git a/Engine/Source/Editor/PropertyEditor/Public/DetailsDisplayManager.h b/Engine/Source/Editor/PropertyEditor/Public/DetailsDisplayManager.h
--- a/Engine/Source/Editor/PropertyEditor/Public/DetailsDisplayManager.h
+++ b/Engine/Source/Editor/PropertyEditor/Public/DetailsDisplayManager.h
@@ -8,6 +8,8 @@
 DECLARE_DELEGATE(FOnDetailsNeedsUpdate)
 
 class FDetailsViewStyle;
+class FComplexPropertyNode;
+class FDetailLayoutBuilderImpl;
 
 /** An @code FDetailsDisplayManager @endcode provides an API to tweak various settings of your details view, and
  * provides some utility methods to work with Details.  */
@@ -90,6 +92,18 @@ public:
 	 */
 	FMargin GetTablePadding() const;
 
+	/**
+	 * If TSharedRef<FComplexPropertyNode> Node has a valid UObject instance, add an empty properties
+	 * Category to the DetailLayoutBuilder
+	 *
+	 * @param Node the FComplexPropertyNode that wil be added as a property to the new Category
+	 * @param DetailLayoutBuilder The FDetailLayoutBuilderImpl that will have a stub category added to it
+	 *
+	 * @return true if the Category was successfully added to FDetailLayoutBuilderImpl& DetailLayoutBuilder,
+	 * else it returns false
+	 */	
+	PROPERTYEDITOR_API virtual bool AddEmptyCategoryToDetailLayoutIfNeeded(TSharedRef<FComplexPropertyNode> Node, TSharedRef<FDetailLayoutBuilderImpl> DetailLayoutBuilder);
+
 	/**
 	* Returns true if the specified UObject is a Root Node Object and should show an empty Category stub even if it
 	* has no UProperty Data, else it returns false
@@ -121,6 +135,5 @@ protected:
 	* here means that in this value the work has been done to figure out if the scrollbar should show 
 	 */
 	bool bIsScrollBarNeeded = false;
-
 };
 
