For horde localization warning: https://horde.devtools.epicgames.com/job/6513a92fa561d77c19c617bb?step=90cb&issue=371644

diff --git a/Engine/Source/Editor/PropertyEditor/Private/PropertyCustomizationHelpers.cpp b/Engine/Source/Editor/PropertyEditor/Private/PropertyCustomizationHelpers.cpp
--- a/Engine/Source/Editor/PropertyEditor/Private/PropertyCustomizationHelpers.cpp
+++ b/Engine/Source/Editor/PropertyEditor/Private/PropertyCustomizationHelpers.cpp
@@ -181,7 +181,7 @@ namespace PropertyCustomizationHelpers
 		return SNew(SBox)
 			.HAlign(HAlign_Left)
 			.VAlign(VAlign_Center)
-			.ToolTipText(OptionalToolTipText.Get().IsEmpty() ? LOCTEXT("SetButtonToolTipText", "Set Optional to default value.") : OptionalToolTipText)
+			.ToolTipText(OptionalToolTipText.Get().IsEmpty() ? LOCTEXT("SetOptionalButtonToolTipText", "Set Optional to default value.") : OptionalToolTipText)
 			[
 				SNew(SButton)
 				.ButtonStyle(FAppStyle::Get(), "Button")
@@ -199,7 +199,7 @@ namespace PropertyCustomizationHelpers
 	{
 		return
 			SNew(SPropertyEditorButton)
-			.Text(OptionalToolTipText.Get().IsEmpty() ? LOCTEXT("ClearButtonToolTipText", "Clear Optional") : OptionalToolTipText)
+			.Text(OptionalToolTipText.Get().IsEmpty() ? LOCTEXT("ClearOptionalButtonToolTipText", "Clear Optional") : OptionalToolTipText)
 			.Image(FAppStyle::Get().GetBrush("Icons.X"))
 			.OnClickAction(OnClearOptionalClicked)
 			.IsEnabled(IsEnabled)
