Task: Ensure Cvars are set to the correct default values once features are all in for card layout

Changing the default CVar values to hide the subobject\component editor by default, and to show the components in card layout. This has no effect when an actor is selected

diff --git a/Engine/Source/Editor/PropertyEditor/Private/DetailsDisplayManager.cpp b/Engine/Source/Editor/PropertyEditor/Private/DetailsDisplayManager.cpp
--- a/Engine/Source/Editor/PropertyEditor/Private/DetailsDisplayManager.cpp
+++ b/Engine/Source/Editor/PropertyEditor/Private/DetailsDisplayManager.cpp
@@ -11,9 +11,9 @@
 
 
 static TAutoConsoleVariable<bool> CVarForceShowComponentEditor(
-	TEXT("CoreEntity.UI.ForceShowComponentEditor"),
-	true,
-	TEXT("Force the component editor to show in the main details tab."));
+	TEXT("DetailsPanel.UI.ForceShowComponentEditor"),
+	false,
+	TEXT("If true, forces the component editor to show in the main viewport and blueprint details panel for UObjects which normally have it hidden."));
 
 
 FDetailsDisplayManager::FDetailsDisplayManager(): bIsOuterCategory(false)
