[Feature] added override for IDetailCategoryBuilder::AddAllExternalStructureProperties that takes an IStructureDataProvider so it can work with multi-edit properties

diff --git a/Engine/Source/Editor/PropertyEditor/Private/DetailCategoryBuilderImpl.cpp b/Engine/Source/Editor/PropertyEditor/Private/DetailCategoryBuilderImpl.cpp
--- a/Engine/Source/Editor/PropertyEditor/Private/DetailCategoryBuilderImpl.cpp
+++ b/Engine/Source/Editor/PropertyEditor/Private/DetailCategoryBuilderImpl.cpp
@@ -502,9 +502,16 @@ IDetailPropertyRow* FDetailCategoryImpl::AddExternalStructureProperty(TSharedPtr
 }
 
 TArray<TSharedPtr<IPropertyHandle>> FDetailCategoryImpl::AddAllExternalStructureProperties(TSharedRef<FStructOnScope> StructData, EPropertyLocation::Type Location, TArray<IDetailPropertyRow*>* OutPropertiesRow)
+{
+	return AddAllExternalStructureProperties(MakeShared<FStructOnScopeStructureDataProvider>(StructData), Location, OutPropertiesRow);
+}
+
+TArray<TSharedPtr<IPropertyHandle>> FDetailCategoryImpl::AddAllExternalStructureProperties(
+	TSharedPtr<IStructureDataProvider> StructProvider, EPropertyLocation::Type Location,
+	TArray<IDetailPropertyRow*>* OutPropertiesRow)
 {
 	TSharedPtr<FStructurePropertyNode> RootPropertyNode(new FStructurePropertyNode);
-	RootPropertyNode->SetStructure(StructData);
+	RootPropertyNode->SetStructure(StructProvider);
 
 	FPropertyNodeInitParams InitParams;
 	InitParams.ParentNode = nullptr;
diff --git a/Engine/Source/Editor/PropertyEditor/Private/DetailCategoryBuilderImpl.h b/Engine/Source/Editor/PropertyEditor/Private/DetailCategoryBuilderImpl.h
--- a/Engine/Source/Editor/PropertyEditor/Private/DetailCategoryBuilderImpl.h
+++ b/Engine/Source/Editor/PropertyEditor/Private/DetailCategoryBuilderImpl.h
@@ -175,6 +175,7 @@ public:
 	virtual IDetailPropertyRow* AddExternalStructure(TSharedPtr<FStructOnScope> StructData, EPropertyLocation::Type Location = EPropertyLocation::Default) override;
 	virtual IDetailPropertyRow* AddExternalStructureProperty(TSharedPtr<FStructOnScope> StructData, FName PropertyName, EPropertyLocation::Type Location = EPropertyLocation::Default, const FAddPropertyParams& Params = FAddPropertyParams()) override;
 	virtual TArray<TSharedPtr<IPropertyHandle>> AddAllExternalStructureProperties(TSharedRef<FStructOnScope> StructData, EPropertyLocation::Type Location = EPropertyLocation::Default, TArray<IDetailPropertyRow*>* OutPropertiesRow = nullptr) override;
+	virtual TArray<TSharedPtr<IPropertyHandle>> AddAllExternalStructureProperties(TSharedPtr<IStructureDataProvider> StructProvider, EPropertyLocation::Type Location, TArray<IDetailPropertyRow*>* OutPropertiesRow) override;
 	virtual bool IsParentLayoutValid() const override { return DetailLayoutBuilder.IsValid(); }
 	virtual IDetailLayoutBuilder& GetParentLayout() const override { return *DetailLayoutBuilder.Pin(); }
 	virtual FDetailWidgetRow& AddCustomRow(const FText& FilterString, bool bForAdvanced = false) override;
diff --git a/Engine/Source/Editor/PropertyEditor/Public/DetailCategoryBuilder.h b/Engine/Source/Editor/PropertyEditor/Public/DetailCategoryBuilder.h
--- a/Engine/Source/Editor/PropertyEditor/Public/DetailCategoryBuilder.h
+++ b/Engine/Source/Editor/PropertyEditor/Public/DetailCategoryBuilder.h
@@ -142,6 +142,7 @@ public:
 	 * @return Array of all properties added to this category
 	 */
 	virtual TArray<TSharedPtr<IPropertyHandle>> AddAllExternalStructureProperties(TSharedRef<FStructOnScope> StructData, EPropertyLocation::Type Location = EPropertyLocation::Default, TArray<IDetailPropertyRow*>* OutPropertiesRow = nullptr) = 0;
+	virtual TArray<TSharedPtr<IPropertyHandle>> AddAllExternalStructureProperties(TSharedPtr<IStructureDataProvider> StructProvider, EPropertyLocation::Type Location = EPropertyLocation::Default, TArray<IDetailPropertyRow*>* OutPropertiesRow = nullptr) = 0;
 
 	/**
 	 * Adds a custom widget row to the category
