Allow the details panel to distinguish between templates and instances when comparing property values

#rnx

diff --git a/Engine/Source/Editor/PropertyEditor/Private/PropertyNode.cpp b/Engine/Source/Editor/PropertyEditor/Private/PropertyNode.cpp
--- a/Engine/Source/Editor/PropertyEditor/Private/PropertyNode.cpp
+++ b/Engine/Source/Editor/PropertyEditor/Private/PropertyNode.cpp
@@ -2249,7 +2249,7 @@ private:
 };
 
 
-bool FPropertyNode::GetDiffersFromDefault(const uint8* PropertyValueAddress, const uint8* PropertyDefaultAddress, const uint8* DefaultPropertyValueBaseAddress, const FProperty* InProperty) const
+bool FPropertyNode::GetDiffersFromDefault(const uint8* PropertyValueAddress, const uint8* PropertyDefaultAddress, const uint8* DefaultPropertyValueBaseAddress, const FProperty* InProperty, const UObject* TopLevelObject) const
 {
 	bool bDiffersFromDefaultValue = false;
 
@@ -2313,8 +2313,8 @@ bool FPropertyNode::GetDiffersFromDefault(const uint8* PropertyValueAddress, con
 			// try to compare the values at the current and default property addresses
 			if( PropertyValueAddress != nullptr && PropertyDefaultAddress != nullptr )
 			{
-				FString DefaultValue = GetDefaultValueAsString(PropertyDefaultAddress, InProperty, EValueAsStringMode::ForDiff);
-				FString CurrentValue = GetDefaultValueAsString(PropertyValueAddress, InProperty, EValueAsStringMode::ForDiff);
+				FString DefaultValue = GetDefaultValueAsString(PropertyDefaultAddress, InProperty, EValueAsStringMode::ForDiff, TopLevelObject);
+				FString CurrentValue = GetDefaultValueAsString(PropertyValueAddress, InProperty, EValueAsStringMode::ForDiff, TopLevelObject);
 				bDiffersFromDefaultValue = !(DefaultValue.Equals(CurrentValue, ESearchCase::CaseSensitive));
 			}
 		}
@@ -2333,7 +2333,7 @@ bool FPropertyNode::GetDiffersFromDefaultForObject( FPropertyItemValueDataTracke
 
 	if (bIsValidTracker && bHasDefaultValue && bHasParent)
 	{
-		return GetDiffersFromDefault(ValueTracker.GetPropertyValueAddress(), ValueTracker.GetPropertyDefaultAddress(), ValueTracker.GetPropertyDefaultBaseAddress(), InProperty);
+		return GetDiffersFromDefault(ValueTracker.GetPropertyValueAddress(), ValueTracker.GetPropertyDefaultAddress(), ValueTracker.GetPropertyDefaultBaseAddress(), InProperty, ValueTracker.GetTopLevelObject());
 	}
 
 	return false;
@@ -2367,6 +2367,7 @@ bool FPropertyNode::GetDiffersFromDefault()
 
 			FStructOnScope DefaultStruct;
 
+			const FObjectPropertyNode* TopLevelObjectNode = StructNode->FindObjectItemParent();
 			for (int32 Index = 0; !bDiffersFromDefault && Index < Structs.Num(); Index++)
 			{
 				const TSharedPtr<FStructOnScope>& StructData = Structs[Index];
@@ -2403,7 +2404,8 @@ bool FPropertyNode::GetDiffersFromDefault()
 					PropertyDefaultAddress = PropertyDefaultBaseAddress;
 				}
 
-				bDiffersFromDefault = GetDiffersFromDefault(PropertyValueAddress, PropertyDefaultAddress, PropertyDefaultBaseAddress, Prop);
+				const UObject* TopLevelObject = (TopLevelObjectNode && Index < TopLevelObjectNode->GetNumObjects()) ? TopLevelObjectNode->GetUObject(Index) : nullptr;
+				bDiffersFromDefault = GetDiffersFromDefault(PropertyValueAddress, PropertyDefaultAddress, PropertyDefaultBaseAddress, Prop, TopLevelObject);
 			}
 		}
 		else if (FObjectPropertyNode* ObjectNode = FindObjectItemParent())
@@ -2429,7 +2431,7 @@ bool FPropertyNode::GetDiffersFromDefault()
 }
 
 
-FString FPropertyNode::GetDefaultValueAsString(const uint8* PropertyDefaultAddress, const FProperty* InProperty, EValueAsStringMode Mode) const
+FString FPropertyNode::GetDefaultValueAsString(const uint8* PropertyDefaultAddress, const FProperty* InProperty, EValueAsStringMode Mode, const UObject* TopLevelObject) const
 {
 	const bool bUseDisplayName = (Mode == EValueAsStringMode::UseDisplayName);
 	FString DefaultValue;
@@ -2442,6 +2444,10 @@ FString FPropertyNode::GetDefaultValueAsString(const uint8* PropertyDefaultAddre
 	else if (Mode == EValueAsStringMode::ForDiff)
 	{
 		PortFlags |= PPF_ForDiff;
+		if (TopLevelObject && !TopLevelObject->IsTemplate())
+		{
+			PortFlags |= PPF_ForDiffInstanceOnly;
+		}
 	}
 
 	if (InProperty->ContainsInstancedObjectProperty())
@@ -2460,7 +2466,7 @@ FString FPropertyNode::GetDefaultValueAsString(const uint8* PropertyDefaultAddre
 			FMemory::Free(TempComplexPropAddr);
 		};
 				
-		InProperty->ExportText_Direct(DefaultValue, TempComplexPropAddr, TempComplexPropAddr, nullptr, PPF_None);
+		InProperty->ExportText_Direct(DefaultValue, TempComplexPropAddr, TempComplexPropAddr, nullptr, PortFlags);
 	}
 	else if ( GetArrayIndex() == INDEX_NONE && InProperty->ArrayDim > 1 )
 	{
@@ -2470,7 +2476,7 @@ FString FPropertyNode::GetDefaultValueAsString(const uint8* PropertyDefaultAddre
 	else
 	{
 		// Port flags will cause enums to display correctly
-		InProperty->ExportTextItem_Direct(DefaultValue, PropertyDefaultAddress, PropertyDefaultAddress, nullptr, PortFlags, nullptr);
+		InProperty->ExportTextItem_Direct(DefaultValue, PropertyDefaultAddress, PropertyDefaultAddress, nullptr, PortFlags);
 	}
 
 	return DefaultValue;
@@ -2488,7 +2494,7 @@ FString FPropertyNode::GetDefaultValueAsStringForObject( FPropertyItemValueDataT
 	{
 		if ( ValueTracker.IsValidTracker() && ValueTracker.HasDefaultValue() )
 		{
-			DefaultValue = GetDefaultValueAsString(ValueTracker.GetPropertyDefaultAddress(), InProperty, Mode);
+			DefaultValue = GetDefaultValueAsString(ValueTracker.GetPropertyDefaultAddress(), InProperty, Mode, InObject);
 		}
 	}
 
@@ -2521,8 +2527,10 @@ FString FPropertyNode::GetDefaultValueAsString(bool bUseDisplayName)
 		FStructOnScope DefaultStruct;
 		FString NodeDefaultValue;
 
-		for (const TSharedPtr<FStructOnScope>& StructData : Structs)
+		const FObjectPropertyNode* TopLevelObjectNode = StructNode->FindObjectItemParent();
+		for (int32 StructIndex = 0; StructIndex < Structs.Num(); ++StructIndex)
 		{
+			const TSharedPtr<FStructOnScope>& StructData = Structs[StructIndex];
 			if (!StructData.IsValid())
 			{
 				continue;
@@ -2550,7 +2558,8 @@ FString FPropertyNode::GetDefaultValueAsString(bool bUseDisplayName)
 					PropertyDefaultAddress = PropertyDefaultBaseAddress;
 				}
 
-				NodeDefaultValue = GetDefaultValueAsString(PropertyDefaultAddress, Prop, Mode);
+				const UObject* TopLevelObject = (TopLevelObjectNode && StructIndex < TopLevelObjectNode->GetNumObjects()) ? TopLevelObjectNode->GetUObject(StructIndex) : nullptr;
+				NodeDefaultValue = GetDefaultValueAsString(PropertyDefaultAddress, Prop, Mode, TopLevelObject);
 			}
 			
 			if (DefaultValue.IsEmpty())
diff --git a/Engine/Source/Editor/PropertyEditor/Private/PropertyNode.h b/Engine/Source/Editor/PropertyEditor/Private/PropertyNode.h
--- a/Engine/Source/Editor/PropertyEditor/Private/PropertyNode.h
+++ b/Engine/Source/Editor/PropertyEditor/Private/PropertyNode.h
@@ -1127,7 +1127,7 @@ protected:
 	/** @return		The property stored at this node, to be passed to Pre/PostEditChange. */
 	FProperty*		GetStoredProperty()		{ return nullptr; }
 
-	bool GetDiffersFromDefault(const uint8* PropertyValueAddress, const uint8* PropertyDefaultAddress, const uint8* DefaultPropertyValueBaseAddress, const FProperty* InProperty) const;
+	bool GetDiffersFromDefault(const uint8* PropertyValueAddress, const uint8* PropertyDefaultAddress, const uint8* DefaultPropertyValueBaseAddress, const FProperty* InProperty, const UObject* TopLevelObject) const;
 	bool GetDiffersFromDefaultForObject( FPropertyItemValueDataTrackerSlate& ValueTracker, FProperty* InProperty );
 
 	enum class EValueAsStringMode
@@ -1136,7 +1136,7 @@ protected:
 		UseDisplayName,
 		ForDiff,
 	};
-	FString GetDefaultValueAsString(const uint8* PropertyDefaultAddress, const FProperty* InProperty, EValueAsStringMode Mode) const;
+	FString GetDefaultValueAsString(const uint8* PropertyDefaultAddress, const FProperty* InProperty, EValueAsStringMode Mode, const UObject* TopLevelObject) const;
 	FString GetDefaultValueAsStringForObject( FPropertyItemValueDataTrackerSlate& ValueTracker, UObject* InObject, FProperty* InProperty, EValueAsStringMode Mode);
 	
 	/**
diff --git a/Engine/Source/Runtime/Core/Public/UObject/PropertyPortFlags.h b/Engine/Source/Runtime/Core/Public/UObject/PropertyPortFlags.h
--- a/Engine/Source/Runtime/Core/Public/UObject/PropertyPortFlags.h
+++ b/Engine/Source/Runtime/Core/Public/UObject/PropertyPortFlags.h
@@ -71,12 +71,17 @@ enum EPropertyPortFlags
 	/** Indicates that object property values should be exported without the package or class information */
 	PPF_SimpleObjectText			= 0x00002000,
 
+	/** Set if this export is going to be used for diffing and may want to omit certain data that doesn't make sense for comparison */
 	PPF_ForDiff						= 0x00004000,
 	
-	/** parsing default properties - allow text for transient properties to be imported - also modifies ObjectProperty importing slightly for subobjects */
-	PPF_ParsingDefaultProperties	= 0x00008000,
+	/**
+	 * Set if this export is going to be used for diffing and may want to omit certain data related to CDOs/archetypes that doesn't make sense when comparing an instance of an object/property/value.
+	 * Requires PPF_ForDiff to also be set.
+	 */
+	PPF_ForDiffInstanceOnly			= 0x00008000,
 
-	//								= 0x00010000,
+	/** parsing default properties - allow text for transient properties to be imported - also modifies ObjectProperty importing slightly for subobjects */
+	PPF_ParsingDefaultProperties	= 0x00010000,
 	
 	/** indicates that non-categorized transient properties should be exported (by default, they would not be) */
 	PPF_IncludeTransient			= 0x00020000,
