Updated the fixed-size array case to also compare against the default via ExportText

#rnx

diff --git a/Engine/Source/Editor/PropertyEditor/Private/PropertyNode.cpp b/Engine/Source/Editor/PropertyEditor/Private/PropertyNode.cpp
--- a/Engine/Source/Editor/PropertyEditor/Private/PropertyNode.cpp
+++ b/Engine/Source/Editor/PropertyEditor/Private/PropertyNode.cpp
@@ -2284,39 +2284,17 @@ bool FPropertyNode::GetDiffersFromDefault(const uint8* PropertyValueAddress, con
 
 	if (!bDiffersFromDefaultValue)
 	{
-		uint32 PortFlags = 0;
-		if (InProperty->ContainsInstancedObjectProperty())
-		{
-			PortFlags |= PPF_DeepComparison;
-		}
-	
 		if (PropertyValueAddress == nullptr || PropertyDefaultAddress == nullptr)
 		{
 			// if either are NULL, we had a dynamic array somewhere in our parent chain and the array doesn't
 			// have enough elements in either the default or the object
 			bDiffersFromDefaultValue = true;
 		}
-		else if (GetArrayIndex() == INDEX_NONE && InProperty->ArrayDim > 1)
-		{
-			// this is a container; loop through all of its elements and see if any of them differ from the default
-			for (int32 Idx = 0; !bDiffersFromDefault && Idx < InProperty->ArrayDim; Idx++)
-			{
-				bDiffersFromDefaultValue = !InProperty->Identical(
-					PropertyValueAddress + Idx * InProperty->ElementSize,
-					PropertyDefaultAddress + Idx * InProperty->ElementSize,
-					PortFlags
-					);
-			}
-		}
 		else
 		{
-			// try to compare the values at the current and default property addresses
-			if( PropertyValueAddress != nullptr && PropertyDefaultAddress != nullptr )
-			{
-				FString DefaultValue = GetDefaultValueAsString(PropertyDefaultAddress, InProperty, EValueAsStringMode::ForDiff, TopLevelObject);
-				FString CurrentValue = GetDefaultValueAsString(PropertyValueAddress, InProperty, EValueAsStringMode::ForDiff, TopLevelObject);
-				bDiffersFromDefaultValue = !(DefaultValue.Equals(CurrentValue, ESearchCase::CaseSensitive));
-			}
+			FString DefaultValue = GetDefaultValueAsString(PropertyDefaultAddress, InProperty, EValueAsStringMode::ForDiff, TopLevelObject);
+			FString CurrentValue = GetDefaultValueAsString(PropertyValueAddress, InProperty, EValueAsStringMode::ForDiff, TopLevelObject);
+			bDiffersFromDefaultValue = !(DefaultValue.Equals(CurrentValue, ESearchCase::CaseSensitive));
 		}
 	}
 
