Details Panel - Restrict Add Component buttons

* Added an argument for IClassViewerFilter for the SVerseClassCombo to take in the IClassViewerFilter
* Had FCoreEntityDetailsViewObjectFilter extend FClassNameListFilter to provide the filter for the combobox
* Had FEntityDetailsDisplayManager fill out the list properly when adding and removing components
* Hooked up a OnComboBoxOpened event to the add component dropdown which will refresh it based on the currently specified filter list
* In CoreEntityDetails.cpp removed the Add Entity button from the level editor widget for its customization
* In SVerseClassCombo, added GenerateFilteredVerseClassList which takes no parameters and uses the class memeber specified CurrentFilterString to filter
* Updated the existing GenerateFilteredVerseClassList(const FString& InSearchText) method in SVerseClassCombo to remove any classes specified by the user in the new IClassViewerFilter argument

diff --git a/Engine/Source/Editor/PropertyEditor/Private/DetailsDisplayManager.cpp b/Engine/Source/Editor/PropertyEditor/Private/DetailsDisplayManager.cpp
--- a/Engine/Source/Editor/PropertyEditor/Private/DetailsDisplayManager.cpp
+++ b/Engine/Source/Editor/PropertyEditor/Private/DetailsDisplayManager.cpp
@@ -6,9 +6,9 @@
 #include "SDetailsView.h"
 #include "DetailLayoutBuilderImpl.h"
 #include "DetailCategoryBuilderImpl.h"
-#include "ObjectEditorUtils.h"
 #include "ObjectPropertyNode.h"
 
+#define LOCTEXT_NAMESPACE "DetailsDisplayManager"
 
 static TAutoConsoleVariable<bool> CVarForceShowComponentEditor(
 	TEXT("DetailsPanel.UI.ForceShowComponentEditor"),
@@ -78,8 +78,9 @@ bool FDetailsDisplayManager::ShowEmptyCategoryIfRootUObjectHasNoPropertyData(UOb
 	return false;
 }
 
+
 bool FDetailsDisplayManager::AddEmptyCategoryToDetailLayoutIfNeeded(TSharedRef<FComplexPropertyNode> Node,
-                                                      TSharedRef<FDetailLayoutBuilderImpl> DetailLayoutBuilder)
+																	TSharedRef<FDetailLayoutBuilderImpl> DetailLayoutBuilder)
 {
 	if (FObjectPropertyNode* ObjectPropertyNode = Node->AsObjectNode())
 	{
@@ -132,3 +133,4 @@ bool FDetailsDisplayManager::GetForceShowSubObjectEditor()
 	return CVarForceShowComponentEditor.GetValueOnAnyThread();
 }
 
+#undef LOCTEXT_NAMESPACE
\ No newline at end of file
