Properly assign the allow list ruleset to an explicit entry with an empty allow list

diff --git a/Engine/Source/Editor/PropertyEditor/Private/PropertyPermissionList.cpp b/Engine/Source/Editor/PropertyEditor/Private/PropertyPermissionList.cpp
--- a/Engine/Source/Editor/PropertyEditor/Private/PropertyPermissionList.cpp
+++ b/Engine/Source/Editor/PropertyEditor/Private/PropertyPermissionList.cpp
@@ -238,9 +238,14 @@ const FNamePermissionList& FPropertyPermissionList::GetCachedPermissionListForSt
 {
 	check(Struct);
 
+	auto IsAllowListAllProps = [](const FPropertyPermissionListEntry* InEntry)
+		{
+			return InEntry ? (InEntry->PermissionList.GetAllowList().Num() == 0 || InEntry->Rules == EPropertyPermissionListRules::AllowListAllProperties) : true;
+		};
+
 	const FPropertyPermissionListEntry* Entry = RawPropertyPermissionList.Find(Struct);
 	// If an entry is set to AllowListAll, then treat it as if it has no manually-defined properties.
-	const bool bIsThisAllowListEmpty = Entry ? (Entry->PermissionList.GetAllowList().Num() == 0 || Entry->Rules == EPropertyPermissionListRules::AllowListAllProperties) : true;
+	const bool bIsThisAllowListEmpty = IsAllowListAllProps(Entry);
 
 	// Normally this case would be caught in GetCachedPermissionListForStruct, but when being called recursively from a subclass
 	// we still need to update bInOutShouldAllowListAllProperties so that new PermissionLists cache properly.
@@ -294,7 +299,7 @@ const FNamePermissionList& FPropertyPermissionList::GetCachedPermissionListForSt
 				NewPermissionList.Append(Entry->PermissionList);
 			}
 
-			bInOutShouldAllowListAllProperties = Entry->Rules == EPropertyPermissionListRules::AllowListAllProperties;
+			bInOutShouldAllowListAllProperties = IsAllowListAllProps(Entry);
 		}
 
 		// PermissionList all properties if the flag is set, the parent Struct has a PermissionList, and this Struct has no PermissionList
@@ -318,7 +323,7 @@ const FNamePermissionList& FPropertyPermissionList::GetCachedPermissionListForSt
 	// In this case, simply add a dummy entry to the Struct's PermissionList that (likely) won't ever collide with a real property name
 	if (!bIsThisAllowListEmpty)
 	{
-		bInOutShouldAllowListAllProperties = Entry->Rules == EPropertyPermissionListRules::AllowListAllSubclassProperties;
+		bInOutShouldAllowListAllProperties = Entry && Entry->Rules == EPropertyPermissionListRules::AllowListAllSubclassProperties;
 	}
 
 	return PermissionListEntry->PermissionList;
