Added an option to require an explicit definition in the hidden list to consider any property modification as allowed on a specific type
Modified error message to report the type of the object the validation is found on

diff --git a/Engine/Source/Editor/PropertyEditor/Private/PropertyPermissionList.cpp b/Engine/Source/Editor/PropertyEditor/Private/PropertyPermissionList.cpp
--- a/Engine/Source/Editor/PropertyEditor/Private/PropertyPermissionList.cpp
+++ b/Engine/Source/Editor/PropertyEditor/Private/PropertyPermissionList.cpp
@@ -329,6 +329,11 @@ const FNamePermissionList& FPropertyPermissionList::GetCachedPermissionListForSt
 	return PermissionListEntry->PermissionList;
 }
 
+bool FPropertyPermissionList::HasSpecificList(const UStruct* ObjectStruct) const
+{
+	return RawPropertyPermissionList.Find(ObjectStruct) != nullptr;
+}
+
 bool FPropertyPermissionList::IsSpecificPropertyAllowListed(const UStruct* ObjectStruct, FName PropertyName) const
 {
 	const FPropertyPermissionListEntry* Entry = RawPropertyPermissionList.Find(ObjectStruct);
diff --git a/Engine/Source/Editor/PropertyEditor/Public/PropertyPermissionList.h b/Engine/Source/Editor/PropertyEditor/Public/PropertyPermissionList.h
--- a/Engine/Source/Editor/PropertyEditor/Public/PropertyPermissionList.h
+++ b/Engine/Source/Editor/PropertyEditor/Public/PropertyPermissionList.h
@@ -92,6 +92,9 @@ public:
 	 */
 	bool DoesPropertyPassFilter(const UStruct* ObjectStruct, FName PropertyName) const;
 
+	/** */
+	bool HasSpecificList(const UStruct* ObjectStruct) const;
+
 	/** Check whether a property exists on the PermissionList for a specific Struct - this will return false if the property is AllowListed on a parent Struct */
 	bool IsSpecificPropertyAllowListed(const UStruct* ObjectStruct, FName PropertyName) const;
 	/** Check whether a property exists on the DenyList for a specific Struct - this will return false if the property is DenyListed on a parent Struct */
