Instanced Struct: Fix crash when struct type was changed during UI update.

diff --git a/Engine/Plugins/Experimental/StructUtils/Source/StructUtilsEditor/Private/Customizations/InstancedStructDetails.cpp b/Engine/Plugins/Experimental/StructUtils/Source/StructUtilsEditor/Private/Customizations/InstancedStructDetails.cpp
--- a/Engine/Plugins/Experimental/StructUtils/Source/StructUtilsEditor/Private/Customizations/InstancedStructDetails.cpp
+++ b/Engine/Plugins/Experimental/StructUtils/Source/StructUtilsEditor/Private/Customizations/InstancedStructDetails.cpp
@@ -83,16 +83,15 @@ public:
 		return CommonStruct;
 	}
 
-	virtual void GetInstances(TArray<TSharedPtr<FStructOnScope>>& OutInstances) const override
+	virtual void GetInstances(TArray<TSharedPtr<FStructOnScope>>& OutInstances, const UStruct* ExpectedBaseStructure) const override
 	{
 		// The returned instances need to be compatible with base structure.
 		// This function returns empty instances in case they are not compatible, with the idea that we have as many instances as we have outer objects.
-		const UScriptStruct* CommonStruct = Cast<UScriptStruct>(GetBaseStructure());
-		EnumerateInstances([&OutInstances, CommonStruct](const UScriptStruct* ScriptStruct, uint8* Memory, UPackage* Package)
+		EnumerateInstances([&OutInstances, ExpectedBaseStructure](const UScriptStruct* ScriptStruct, uint8* Memory, UPackage* Package)
 		{
 			TSharedPtr<FStructOnScope> Result;
 			
-			if (CommonStruct && ScriptStruct && ScriptStruct->IsChildOf(CommonStruct))
+			if (ExpectedBaseStructure && ScriptStruct && ScriptStruct->IsChildOf(ExpectedBaseStructure))
 			{
 				Result = MakeShared<FStructOnScope>(ScriptStruct, Memory);
 				Result->SetPackage(Package);
@@ -109,7 +108,7 @@ public:
 		return true;
 	}
 
-	virtual uint8* GetValueBaseAddress(uint8* ParentValueAddress, const UStruct* ExpectedType) const override
+	virtual uint8* GetValueBaseAddress(uint8* ParentValueAddress, const UStruct* ExpectedBaseStructure) const override
 	{
 		if (!ParentValueAddress)
 		{
@@ -117,11 +116,11 @@ public:
 		}
 
 		FInstancedStruct& InstancedStruct = *reinterpret_cast<FInstancedStruct*>(ParentValueAddress);
-		if (ExpectedType && InstancedStruct.GetScriptStruct() && InstancedStruct.GetScriptStruct()->IsChildOf(ExpectedType))
+		if (ExpectedBaseStructure && InstancedStruct.GetScriptStruct() && InstancedStruct.GetScriptStruct()->IsChildOf(ExpectedBaseStructure))
 		{
 			return InstancedStruct.GetMutableMemory();
 		}
-
+		
 		return nullptr;
 	}
 	
diff --git a/Engine/Source/Editor/PropertyEditor/Private/StructurePropertyNode.cpp b/Engine/Source/Editor/PropertyEditor/Private/StructurePropertyNode.cpp
--- a/Engine/Source/Editor/PropertyEditor/Private/StructurePropertyNode.cpp
+++ b/Engine/Source/Editor/PropertyEditor/Private/StructurePropertyNode.cpp
@@ -42,6 +42,9 @@ void FStructurePropertyNode::InitChildNodes()
 		NewItemNode->InitNode(InitParams);
 		AddChildNode(NewItemNode);
 	}
+
+	// Cache the init time base struct so that we can determine of the struct has changed.
+	WeakCachedBaseStruct = Struct;
 }
 
 uint8* FStructurePropertyNode::GetValueBaseAddress(uint8* StartAddress, bool bIsSparseData, bool bIsStruct) const
@@ -68,11 +71,11 @@ uint8* FStructurePropertyNode::GetValueBaseAddress(uint8* StartAddress, bool bIs
 			const TSharedPtr<FPropertyNode> ParentNode = ParentNodeWeakPtr.Pin();
 			if (!ensureMsgf(ParentNode, TEXT("Expecting valid parent node when indirection structure provider is called with Object data.")))
 			{
-				return nullptr;				
+				return nullptr;
 			}
 			// Resolve from parent nodes data.
 			uint8* ParentValueAddress = ParentNode->GetValueAddress(StartAddress, bIsSparseData);
-			uint8* ValueAddress = StructProvider->GetValueBaseAddress(ParentValueAddress, GetBaseStructure());
+			uint8* ValueAddress = StructProvider->GetValueBaseAddress(ParentValueAddress, WeakCachedBaseStruct.Get());
 			return ValueAddress;
 		}
 
@@ -80,7 +83,7 @@ uint8* FStructurePropertyNode::GetValueBaseAddress(uint8* StartAddress, bool bIs
 		// In that case we can only support one instance, since we cannot discern them.
 		// Note: Multiple standalone structure instances are supported when bIsStruct is true (e.g. when the structure property is root node).
 		TArray<TSharedPtr<FStructOnScope>> Instances;
-		StructProvider->GetInstances(Instances);
+		StructProvider->GetInstances(Instances, WeakCachedBaseStruct.Get());
 		ensureMsgf(Instances.Num() <= 1, TEXT("Expecting max one instance on standalone structure provider."));
 		if (Instances.Num() == 1 && Instances[0].IsValid())
 		{
@@ -95,26 +98,12 @@ EPropertyDataValidationResult FStructurePropertyNode::EnsureDataIsValid()
 {
 	CachedReadAddresses.Reset();
 
-	const UStruct* BaseStruct = GetBaseStructure();
-	
-	// Check that struct node's children still belong to the current base struct.
-	for (const TSharedPtr<FPropertyNode>& ChildNode : ChildNodes)
+	// If the struct has changed, rebuild children.
+	const UStruct* CachedBaseStruct = WeakCachedBaseStruct.Get();
+	if (GetBaseStructure() != CachedBaseStruct)
 	{
-		if (ChildNode.IsValid())
-		{
-			if (const FProperty* ChildProperty = ChildNode->GetProperty())
-			{
-				const UStruct* OwnerStruct = ChildProperty->GetOwnerStruct();
-				if (!OwnerStruct
-					|| OwnerStruct->IsStructTrashed()
-					|| !BaseStruct
-					|| !BaseStruct->IsChildOf(OwnerStruct)) // OwnerStruct can be BaseStruct or one of structs BaseStruct is derived from.
-				{
-					RebuildChildren();
-					return EPropertyDataValidationResult::ChildrenRebuilt;
-				}
-			}
-		}
+		RebuildChildren();
+		return EPropertyDataValidationResult::ChildrenRebuilt;
 	}
 	
 	return FPropertyNode::EnsureDataIsValid();
diff --git a/Engine/Source/Editor/PropertyEditor/Private/StructurePropertyNode.h b/Engine/Source/Editor/PropertyEditor/Private/StructurePropertyNode.h
--- a/Engine/Source/Editor/PropertyEditor/Private/StructurePropertyNode.h
+++ b/Engine/Source/Editor/PropertyEditor/Private/StructurePropertyNode.h
@@ -25,6 +25,7 @@ public:
 		ClearCachedReadAddresses(true);
 		DestroyTree();
 		StructProvider = nullptr;
+		WeakCachedBaseStruct.Reset();
 	}
 
 	void SetStructure(TSharedPtr<FStructOnScope> InStructData)
@@ -53,7 +54,7 @@ public:
 		if (StructProvider)
 		{
 			TArray<TSharedPtr<FStructOnScope>> Instances;
-			StructProvider->GetInstances(Instances);
+			StructProvider->GetInstances(Instances, WeakCachedBaseStruct.Get());
 			
 			if (Instances.Num() > 0)
 			{
@@ -67,7 +68,7 @@ public:
 	{
 		if (StructProvider)
 		{
-			StructProvider->GetInstances(OutStructs);
+			StructProvider->GetInstances(OutStructs, WeakCachedBaseStruct.Get());
 		}
 	}
 
@@ -98,7 +99,7 @@ public:
 		}
 
 		TArray<TSharedPtr<FStructOnScope>> Instances;
-		StructProvider->GetInstances(Instances);
+		StructProvider->GetInstances(Instances, WeakCachedBaseStruct.Get());
 		bool bHasData = false;
 
 		for (TSharedPtr<FStructOnScope>& Instance : Instances)
@@ -142,7 +143,7 @@ public:
 		bool bAllTheSame = true;
 
 		TArray<TSharedPtr<FStructOnScope>> Instances;
-		StructProvider->GetInstances(Instances);
+		StructProvider->GetInstances(Instances, WeakCachedBaseStruct.Get());
 		
 		if (Instances.IsEmpty())
 		{
@@ -201,6 +202,7 @@ public:
 			}
 		}
 
+		bool bHasData = false;
 		if (bAllTheSame && OutAddresses)
 		{
 			for (TSharedPtr<FStructOnScope>& Instance : Instances)
@@ -209,11 +211,12 @@ public:
 				if (ReadAddress)
 				{
 					OutAddresses->Add(nullptr, InPropertyNode.GetValueBaseAddress(ReadAddress, InPropertyNode.HasNodeFlags(EPropertyNodeFlags::IsSparseClassData) != 0, /*bIsStruct=*/true), /*bIsStruct=*/true);
+					bHasData = true;
 				}
 			}
 		}
 
-		return bAllTheSame;
+		return bAllTheSame && bHasData;
 	}
 
 	void GetOwnerPackages(TArray<UPackage*>& OutPackages) const
@@ -221,7 +224,7 @@ public:
 		if (StructProvider)
 		{
 			TArray<TSharedPtr<FStructOnScope>> Instances;
-			StructProvider->GetInstances(Instances);
+			StructProvider->GetInstances(Instances, WeakCachedBaseStruct.Get());
 
 			for (TSharedPtr<FStructOnScope>& Instance : Instances)
 			{
@@ -254,7 +257,7 @@ public:
 		if (StructProvider)
 		{
 			TArray<TSharedPtr<FStructOnScope>> Instances;
-			StructProvider->GetInstances(Instances);
+			StructProvider->GetInstances(Instances, WeakCachedBaseStruct.Get());
 			
 			for (TSharedPtr<FStructOnScope>& Instance : Instances)
 			{
@@ -274,7 +277,7 @@ public:
 		if (StructProvider)
 		{
 			TArray<TSharedPtr<FStructOnScope>> Instances;
-			StructProvider->GetInstances(Instances);
+			StructProvider->GetInstances(Instances, WeakCachedBaseStruct.Get());
 			
 			for (TSharedPtr<FStructOnScope>& Instance : Instances)
 			{
@@ -292,7 +295,7 @@ public:
 		if (StructProvider)
 		{
 			TArray<TSharedPtr<FStructOnScope>> Instances;
-			StructProvider->GetInstances(Instances);
+			StructProvider->GetInstances(Instances, WeakCachedBaseStruct.Get());
 			
 			return Instances.Num();
 		}
@@ -304,7 +307,7 @@ public:
 		if (StructProvider)
 		{
 			TArray<TSharedPtr<FStructOnScope>> Instances;
-			StructProvider->GetInstances(Instances);
+			StructProvider->GetInstances(Instances, WeakCachedBaseStruct.Get());
 			if (Instances.IsValidIndex(Index) && Instances[Index].IsValid())
 			{
 				return Instances[Index]->GetStructMemory();
@@ -381,4 +384,7 @@ protected:
 
 private:
 	TSharedPtr<IStructureDataProvider> StructProvider;
+
+	/** The base struct at the time InitChildNodes() was called. */
+	TWeakObjectPtr<const UStruct> WeakCachedBaseStruct = nullptr;
 };
diff --git a/Engine/Source/Editor/PropertyEditor/Public/IStructureDataProvider.h b/Engine/Source/Editor/PropertyEditor/Public/IStructureDataProvider.h
--- a/Engine/Source/Editor/PropertyEditor/Public/IStructureDataProvider.h
+++ b/Engine/Source/Editor/PropertyEditor/Public/IStructureDataProvider.h
@@ -46,7 +46,11 @@
 // If IsPropertyIndirection() is false, there structures are expected to be standalone, and GetInstances() is used instead to retried the instance values.
 // When indirection is not used and the struct provider is part of an another instance, the provider should return only one instance, since there is no to associate the provided value
 // with an foreign object.
-// 
+//
+// ExpectedType parameter: This is the type that e.g. the UI is currently initialized for (cached via GetBaseStructure()).
+// It is possible that the provided data changes due to some property change callbacks, and the ExpectedType does not match the provided data type.
+// In that case empty data should be returned so that the UI does not try to access incompatible data.
+//
 // Multiple instances are supported when the provider is the root node (e.g. on Structure details view).
 // 
 //-----------------------------------------------------------------------------
@@ -61,8 +65,13 @@ public:
 	/** @return the most common struct of the instance data. This struct will be used to build the UI for the instances. */
 	virtual const UStruct* GetBaseStructure() const = 0;
 
-	/** @return instances to edit. Each provided struct should be compatible with base struct. */
-	virtual void GetInstances(TArray<TSharedPtr<FStructOnScope>>& OutInstances) const = 0;
+	
+	/**
+	 * Returns instances to edit. Each provided struct should be compatible with ExpectedType.
+	 * @param OutInstances returned instances
+	 * @param ExpectedBaseStructure the base type of struct that is expected to be returned (see the documentation above about ExpectedType) 
+	 */
+	virtual void GetInstances(TArray<TSharedPtr<FStructOnScope>>& OutInstances, const UStruct* ExpectedBaseStructure) const = 0;
 
 	/** @return true, if the provider supports handling indirections using GetValueBaseAddress(). */
 	virtual bool IsPropertyIndirection() const { return false; }
@@ -70,10 +79,16 @@ public:
 	/**
 	 * Returns base address of provided struct based on parent property nodes value.
 	 * @param ParentValueAddress Value of the parent 
-	 * @param ExpectedType the type of struct that is expected to be returned (generally the same type as returned by GetBaseStructure())
+	 * @param ExpectedBaseStructure the base type of struct that is expected to be returned (see the documentation above about ExpectedType)
 	 * @return value base address based on parent base address.
 	 */
-	virtual uint8* GetValueBaseAddress(uint8* ParentValueAddress, const UStruct* ExpectedType) const { return ParentValueAddress; }
+	virtual uint8* GetValueBaseAddress(uint8* ParentValueAddress, const UStruct* ExpectedBaseStructure) const { return ParentValueAddress; }
+
+	UE_DEPRECATED(5.4, "Please use and override GetInstances() with ExpectedType instead.")
+	virtual void GetInstances(TArray<TSharedPtr<FStructOnScope>>& OutInstances) const final
+	{
+		return GetInstances(OutInstances, nullptr);
+	}
 };
 
 
@@ -104,9 +119,18 @@ public:
 		return StructData.IsValid() ? StructData->GetStruct() : nullptr;
 	}
 	
-	virtual void GetInstances(TArray<TSharedPtr<FStructOnScope>>& OutInstances) const override
+	virtual void GetInstances(TArray<TSharedPtr<FStructOnScope>>& OutInstances, const UStruct* ExpectedBaseStructure) const override
 	{
-		OutInstances.Add(StructData);
+		if (StructData.IsValid())
+		{
+			const UStruct* Struct = StructData->GetStruct();
+			if (ExpectedBaseStructure
+				&& Struct
+				&& Struct->IsChildOf(ExpectedBaseStructure))
+			{
+				OutInstances.Add(StructData);				
+			}
+		}
 	}
 
 protected:
