GetAssetFilter: add meta tag to be able to specify a UFunction to filter assets
#rnx
#tests Tested in Editor

diff --git a/Engine/Source/Editor/PropertyEditor/Private/UserInterface/PropertyEditor/SPropertyEditorAsset.cpp b/Engine/Source/Editor/PropertyEditor/Private/UserInterface/PropertyEditor/SPropertyEditorAsset.cpp
--- a/Engine/Source/Editor/PropertyEditor/Private/UserInterface/PropertyEditor/SPropertyEditorAsset.cpp
+++ b/Engine/Source/Editor/PropertyEditor/Private/UserInterface/PropertyEditor/SPropertyEditorAsset.cpp
@@ -306,7 +306,6 @@ void SPropertyEditorAsset::Construct(const FArguments& InArgs, const TSharedPtr<
 	OwnerAssetDataArray = InArgs._OwnerAssetDataArray;
 	OnIsEnabled = InArgs._IsEnabled;
 	OnSetObject = InArgs._OnSetObject;
-	OnShouldFilterAsset = InArgs._OnShouldFilterAsset;
 	OnShouldFilterActor = InArgs._OnShouldFilterActor;
 	ObjectPath = InArgs._ObjectPath;
 
@@ -345,17 +344,45 @@ void SPropertyEditorAsset::Construct(const FArguments& InArgs, const TSharedPtr<
 	bIsSoftObjectPath = CastField<FSoftObjectProperty>(Property) != nullptr;
 	
 	InitializeAssetDataTags(Property);
+
+	auto AppendOnShouldFilterAssetCallback = [this](FOnShouldFilterAsset&& OnShouldFilterAssetCallback)
+	{
+		check(OnShouldFilterAssetCallback.IsBound());
+		if (OnShouldFilterAsset.IsBound())
+		{
+			OnShouldFilterAsset.BindLambda([BaseOnShouldFilterAsset = OnShouldFilterAsset, OnShouldFilterAssetCallback = MoveTemp(OnShouldFilterAssetCallback)](const FAssetData& InAssetData)
+			{
+				return BaseOnShouldFilterAsset.Execute(InAssetData) || OnShouldFilterAssetCallback.Execute(InAssetData);
+			});
+		}
+		else
+		{
+			OnShouldFilterAsset = MoveTemp(OnShouldFilterAssetCallback);
+		}
+	};
+
+	OnShouldFilterAsset = InArgs._OnShouldFilterAsset;
+	
 	if (DisallowedAssetDataTags.IsValid() || RequiredAssetDataTags.IsValid())
 	{
 		// re-route the filter delegate to our own if we have our own asset data tags filter :
-		OnShouldFilterAsset.BindLambda([this, AssetFilter = InArgs._OnShouldFilterAsset](const FAssetData& InAssetData)
+		AppendOnShouldFilterAssetCallback(FOnShouldFilterAsset::CreateRaw(this, &SPropertyEditorAsset::IsAssetAllowed));
+	}
+
+	if (Property && Property->GetOwnerProperty()->HasMetaData("GetAssetFilter"))
+	{
+		// Add MetaData asset filter
+		const FString GetAssetFilterFunctionName = Property->GetOwnerProperty()->GetMetaData("GetAssetFilter");
+		if (!GetAssetFilterFunctionName.IsEmpty())
 		{
-			if (IsAssetAllowed(InAssetData))
+			TArray<UObject*> ObjectList;
+			PropertyEditor->GetPropertyHandle()->GetOuterObjects(ObjectList);
+			const UFunction* GetAssetFilterFunction = !ObjectList.IsEmpty() ? ObjectList[0]->FindFunction(*GetAssetFilterFunctionName) : nullptr;
+			if (GetAssetFilterFunction)
 			{
-				return AssetFilter.IsBound() ? AssetFilter.Execute(InAssetData) : false;
+				AppendOnShouldFilterAssetCallback(FOnShouldFilterAsset::CreateUFunction(ObjectList[0], GetAssetFilterFunction->GetFName()));
 			}
-			return true;
-		});
+		}
 	}
 
 	InitializeClassFilters(Property);
diff --git a/Engine/Source/Runtime/CoreUObject/Public/UObject/ObjectMacros.h b/Engine/Source/Runtime/CoreUObject/Public/UObject/ObjectMacros.h
--- a/Engine/Source/Runtime/CoreUObject/Public/UObject/ObjectMacros.h
+++ b/Engine/Source/Runtime/CoreUObject/Public/UObject/ObjectMacros.h
@@ -1477,6 +1477,12 @@ namespace UM
 
 		/// [PropertyMetadata] Used for enum properties to define the subset of valid values as a comma-separated string; values outside of this subset are not made available.
 		ValidEnumValues,
+		
+		/// [PropertyMetadata] Causes assets to be filtered through a UFunction, e.g. meta=(GetAssetFilter="FuncName"). The UFunction should return true to exclude the asset.
+		///
+		/// UFUNCTION()
+		/// bool FuncName(const FAssetData& AssetData) const;
+		GetAssetFilter,
 
 	    /// [PropertyMetadata] Used for enum properties to define the subset of invalid values as a comma-separated string; values within this subset are not made available.
         InvalidEnumValues,
