GetAssetFilter: Fix for mutli selection
#rnx

diff --git a/Engine/Source/Editor/PropertyEditor/Private/UserInterface/PropertyEditor/SPropertyEditorAsset.cpp b/Engine/Source/Editor/PropertyEditor/Private/UserInterface/PropertyEditor/SPropertyEditorAsset.cpp
--- a/Engine/Source/Editor/PropertyEditor/Private/UserInterface/PropertyEditor/SPropertyEditorAsset.cpp
+++ b/Engine/Source/Editor/PropertyEditor/Private/UserInterface/PropertyEditor/SPropertyEditorAsset.cpp
@@ -377,10 +377,13 @@ void SPropertyEditorAsset::Construct(const FArguments& InArgs, const TSharedPtr<
 		{
 			TArray<UObject*> ObjectList;
 			PropertyEditor->GetPropertyHandle()->GetOuterObjects(ObjectList);
-			const UFunction* GetAssetFilterFunction = !ObjectList.IsEmpty() ? ObjectList[0]->FindFunction(*GetAssetFilterFunctionName) : nullptr;
-			if (GetAssetFilterFunction)
+			for (UObject* Object : ObjectList)
 			{
-				AppendOnShouldFilterAssetCallback(FOnShouldFilterAsset::CreateUFunction(ObjectList[0], GetAssetFilterFunction->GetFName()));
+				const UFunction* GetAssetFilterFunction = Object ? Object->FindFunction(*GetAssetFilterFunctionName) : nullptr;
+				if (GetAssetFilterFunction)
+				{
+					AppendOnShouldFilterAssetCallback(FOnShouldFilterAsset::CreateUFunction(Object, GetAssetFilterFunction->GetFName()));
+				}
 			}
 		}
 	}
