- Added IDetailLayoutBuilder::DoesCategoryHaveGeneratedChildren() to allow to check if a category has any generated rows
- Added option to IDetailCategoryBuilder::HeaderContent() to let the content span the whole row (hides category label)

diff --git a/Engine/Source/Editor/PropertyEditor/Private/DetailCategoryBuilderImpl.cpp b/Engine/Source/Editor/PropertyEditor/Private/DetailCategoryBuilderImpl.cpp
--- a/Engine/Source/Editor/PropertyEditor/Private/DetailCategoryBuilderImpl.cpp
+++ b/Engine/Source/Editor/PropertyEditor/Private/DetailCategoryBuilderImpl.cpp
@@ -383,10 +383,11 @@ IDetailCategoryBuilder& FDetailCategoryImpl::RestoreExpansionState(bool bRestore
 	return *this;
 }
 
-IDetailCategoryBuilder& FDetailCategoryImpl::HeaderContent(TSharedRef<SWidget> InHeaderContent)
+IDetailCategoryBuilder& FDetailCategoryImpl::HeaderContent(TSharedRef<SWidget> InHeaderContent, bool bWholeRowContent)
 {
 	ensureMsgf(!this->HeaderContentWidget.IsValid(), TEXT("Category already has a header content widget defined!"));
 	this->HeaderContentWidget = InHeaderContent;
+	this->bHeaderContentWholeRowContent = bWholeRowContent;
 	return *this;
 }
 
@@ -845,7 +846,8 @@ TSharedRef<ITableRow> FDetailCategoryImpl::GenerateWidgetForTableView(const TSha
 		.IsEmpty( bIsEmpty )
 		.InnerCategory(ParentLayout.IsValid() ? ParentLayout->IsLayoutForExternalRoot() : false)
 		.DisplayName(GetDisplayName())
-		.HeaderContent(HeaderContent);
+		.HeaderContent(HeaderContent)
+		.WholeRowHeaderContent(bHeaderContentWholeRowContent);
 }
 
 void FDetailCategoryImpl::InitializeObjectName()
diff --git a/Engine/Source/Editor/PropertyEditor/Private/DetailCategoryBuilderImpl.h b/Engine/Source/Editor/PropertyEditor/Private/DetailCategoryBuilderImpl.h
--- a/Engine/Source/Editor/PropertyEditor/Private/DetailCategoryBuilderImpl.h
+++ b/Engine/Source/Editor/PropertyEditor/Private/DetailCategoryBuilderImpl.h
@@ -167,7 +167,7 @@ public:
 	virtual IDetailCategoryBuilder& InitiallyCollapsed(bool bShouldBeInitiallyCollapsed) override;
 	virtual IDetailCategoryBuilder& OnExpansionChanged(FOnBooleanValueChanged InOnExpansionChanged) override;
 	virtual IDetailCategoryBuilder& RestoreExpansionState(bool bRestore) override;
-	virtual IDetailCategoryBuilder& HeaderContent(TSharedRef<SWidget> InHeaderContent) override;
+	virtual IDetailCategoryBuilder& HeaderContent(TSharedRef<SWidget> InHeaderContent, bool bWholeRowContent = false) override;
 	virtual IDetailPropertyRow& AddProperty(FName PropertyPath, UClass* ClassOuter = nullptr, FName InstanceName = NAME_None, EPropertyLocation::Type Location = EPropertyLocation::Default) override;
 	virtual IDetailPropertyRow& AddProperty(TSharedPtr<IPropertyHandle> PropertyHandle, EPropertyLocation::Type Location = EPropertyLocation::Default) override;
 	virtual IDetailPropertyRow* AddExternalObjects(const TArray<UObject*>& Objects, EPropertyLocation::Type Location = EPropertyLocation::Default, const FAddPropertyParams& Params = FAddPropertyParams()) override;
@@ -437,6 +437,8 @@ private:
 	FString CategoryPathName;
 	/** Custom header content displayed to the right of the category name */
 	TSharedPtr<SWidget> HeaderContentWidget;
+	/** True if the HeaderContentWidget should span the whole row (hides the category label)  */
+	bool bHeaderContentWholeRowContent = false;
 	/** A property node that is displayed in the header row to the right of the category name. */
 	TSharedPtr<FDetailTreeNode> InlinePropertyNode;
 	/** The parent detail builder */
diff --git a/Engine/Source/Editor/PropertyEditor/Private/DetailLayoutBuilderImpl.cpp b/Engine/Source/Editor/PropertyEditor/Private/DetailLayoutBuilderImpl.cpp
--- a/Engine/Source/Editor/PropertyEditor/Private/DetailLayoutBuilderImpl.cpp
+++ b/Engine/Source/Editor/PropertyEditor/Private/DetailLayoutBuilderImpl.cpp
@@ -219,6 +219,29 @@ IDetailPropertyRow* FDetailLayoutBuilderImpl::EditDefaultProperty(TSharedPtr<IPr
 	return nullptr;
 }
 
+bool FDetailLayoutBuilderImpl::DoesCategoryHaveGeneratedChildren(FName CategoryName)
+{
+	FDetailNodeList Children;
+
+	FDetailCategoryImpl* Category = nullptr;
+	for (const TSharedRef<FDetailTreeNode>& RootTreeNode : AllRootTreeNodes)
+	{
+		if (RootTreeNode->GetNodeType() == EDetailNodeType::Category && 
+			CategoryName == RootTreeNode->GetNodeName())
+		{
+			Category = (FDetailCategoryImpl*)&RootTreeNode.Get();
+		}
+	}
+
+	if (Category)
+	{
+		Category->GetGeneratedChildren(Children, /*bIgnoreVisibility*/true, /*bIgnoreAdvancedDropdown*/false);
+	}
+
+	return Children.Num() > 0;
+}
+
+
 TSharedRef<IPropertyHandle> FDetailLayoutBuilderImpl::GetProperty( const FName PropertyPath, const UStruct* ClassOutermost, FName InInstanceName ) const
 {	
 	TSharedPtr<FPropertyHandleBase> PropertyHandle; 
diff --git a/Engine/Source/Editor/PropertyEditor/Private/DetailLayoutBuilderImpl.h b/Engine/Source/Editor/PropertyEditor/Private/DetailLayoutBuilderImpl.h
--- a/Engine/Source/Editor/PropertyEditor/Private/DetailLayoutBuilderImpl.h
+++ b/Engine/Source/Editor/PropertyEditor/Private/DetailLayoutBuilderImpl.h
@@ -38,6 +38,7 @@ public:
 	virtual TSharedPtr<IPropertyHandle> AddObjectPropertyData(TConstArrayView<UObject*> Objects, FName PropertyName) override;
 	virtual TSharedPtr<IPropertyHandle> AddStructurePropertyData(const TSharedPtr<FStructOnScope>& StructData, FName PropertyName) override;
 	virtual IDetailPropertyRow* EditDefaultProperty(TSharedPtr<IPropertyHandle> InPropertyHandle) override;
+	virtual bool DoesCategoryHaveGeneratedChildren(FName CategoryName) override;
 	virtual TSharedRef<IPropertyHandle> GetProperty(const FName PropertyPath, const UStruct* ClassOutermost, FName InInstanceName) const override;
 	virtual FName GetTopLevelProperty() override;
 	virtual void HideProperty(const TSharedPtr<IPropertyHandle> Property) override;
diff --git a/Engine/Source/Editor/PropertyEditor/Private/SDetailCategoryTableRow.cpp b/Engine/Source/Editor/PropertyEditor/Private/SDetailCategoryTableRow.cpp
--- a/Engine/Source/Editor/PropertyEditor/Private/SDetailCategoryTableRow.cpp
+++ b/Engine/Source/Editor/PropertyEditor/Private/SDetailCategoryTableRow.cpp
@@ -78,17 +78,21 @@ void SDetailCategoryTableRow::Construct(const FArguments& InArgs, TSharedRef<FDe
 					}
 					return EVisibility::Visible;
 				})
-		]
-		+ SHorizontalBox::Slot()
-		.VAlign(VAlign_Center)
-		.Padding(4, 0, 0, 0)
-		.FillWidth(1)
-		[
-			SNew(STextBlock)
-			.Text(InArgs._DisplayName)
-			.Font(FAppStyle::Get().GetFontStyle(bIsInnerCategory ? PropertyEditorConstants::PropertyFontStyle : PropertyEditorConstants::CategoryFontStyle))
-			.TextStyle(FAppStyle::Get(), "DetailsView.CategoryTextStyle")
 		];
+	
+	if (!InArgs._WholeRowHeaderContent)
+	{
+		HeaderBox->AddSlot()
+			.VAlign(VAlign_Center)
+			.Padding(4, 0, 0, 0)
+			.FillWidth(1)
+			[
+				SNew(STextBlock)
+				.Text(InArgs._DisplayName)
+				.Font(FAppStyle::Get().GetFontStyle(bIsInnerCategory ? PropertyEditorConstants::PropertyFontStyle : PropertyEditorConstants::CategoryFontStyle))
+				.TextStyle(FAppStyle::Get(), "DetailsView.CategoryTextStyle")
+			];
+	}
 
 	if (InArgs._HeaderContent.IsValid())
 	{
diff --git a/Engine/Source/Editor/PropertyEditor/Private/SDetailCategoryTableRow.h b/Engine/Source/Editor/PropertyEditor/Private/SDetailCategoryTableRow.h
--- a/Engine/Source/Editor/PropertyEditor/Private/SDetailCategoryTableRow.h
+++ b/Engine/Source/Editor/PropertyEditor/Private/SDetailCategoryTableRow.h
@@ -26,6 +26,7 @@ public:
 		SLATE_ARGUMENT(FText, DisplayName)
 		SLATE_ARGUMENT(bool, InnerCategory)
 		SLATE_ARGUMENT(TSharedPtr<SWidget>, HeaderContent)
+		SLATE_ARGUMENT(bool, WholeRowHeaderContent)
 		SLATE_ARGUMENT(bool, ShowBorder)
 		
 		/** If true, this Category should have no UProperty data associated with it, and will be shown as an 
diff --git a/Engine/Source/Editor/PropertyEditor/Public/DetailCategoryBuilder.h b/Engine/Source/Editor/PropertyEditor/Public/DetailCategoryBuilder.h
--- a/Engine/Source/Editor/PropertyEditor/Public/DetailCategoryBuilder.h
+++ b/Engine/Source/Editor/PropertyEditor/Public/DetailCategoryBuilder.h
@@ -63,8 +63,9 @@ public:
 	 * Adds header content to the category
 	 *
 	 * @param InHeaderContent	The header content widget
+	 * @param bWholeRowContent	The header content should span the whole row (hides category name)
 	 */
-	virtual IDetailCategoryBuilder& HeaderContent(TSharedRef<SWidget> InHeaderContent) = 0;
+	virtual IDetailCategoryBuilder& HeaderContent(TSharedRef<SWidget> InHeaderContent, bool bWholeRowContent = false) = 0;
 
 	/**
 	 * Gets the current sort order of the category
diff --git a/Engine/Source/Editor/PropertyEditor/Public/DetailLayoutBuilder.h b/Engine/Source/Editor/PropertyEditor/Public/DetailLayoutBuilder.h
--- a/Engine/Source/Editor/PropertyEditor/Public/DetailLayoutBuilder.h
+++ b/Engine/Source/Editor/PropertyEditor/Public/DetailLayoutBuilder.h
@@ -180,6 +180,11 @@ public:
 	 */
 	virtual IDetailPropertyRow* EditDefaultProperty(TSharedPtr<IPropertyHandle> InPropertyHandle) = 0;
 
+	/**
+	 * @return true if the category contains child rows. 
+	 */
+	virtual bool DoesCategoryHaveGeneratedChildren(FName CategoryName) = 0;
+
 	/**
 	 * Hides an entire category
 	 *
