PropertyEditorAsset: replace IsAssetAllowed by IsAssetFiltered to match other filters
#rnx
#tests PIE

diff --git a/Engine/Source/Editor/PropertyEditor/Private/UserInterface/PropertyEditor/SPropertyEditorAsset.cpp b/Engine/Source/Editor/PropertyEditor/Private/UserInterface/PropertyEditor/SPropertyEditorAsset.cpp
--- a/Engine/Source/Editor/PropertyEditor/Private/UserInterface/PropertyEditor/SPropertyEditorAsset.cpp
+++ b/Engine/Source/Editor/PropertyEditor/Private/UserInterface/PropertyEditor/SPropertyEditorAsset.cpp
@@ -256,7 +256,7 @@ void SPropertyEditorAsset::InitializeAssetDataTags(const FProperty* Property)
 	}
 }
 
-bool SPropertyEditorAsset::IsAssetAllowed(const FAssetData& InAssetData)
+bool SPropertyEditorAsset::IsAssetFiltered(const FAssetData& InAssetData)
 {
 	if (DisallowedAssetDataTags.IsValid())
 	{
@@ -264,7 +264,7 @@ bool SPropertyEditorAsset::IsAssetAllowed(const FAssetData& InAssetData)
 		{
 			if (InAssetData.TagsAndValues.ContainsKeyValue(DisallowedTagAndValue.Key, DisallowedTagAndValue.Value))
 			{
-				return false;
+				return true;
 			}
 		}
 	}
@@ -280,11 +280,11 @@ bool SPropertyEditorAsset::IsAssetAllowed(const FAssetData& InAssetData)
 				{
 					continue;
 				}
-				return false;
+				return true;
 			}
 		}
 	}
-	return true;
+	return false;
 }
 
 // Awful hack to deal with UClass::FindCommonBase taking an array of non-const classes...
@@ -366,7 +366,7 @@ void SPropertyEditorAsset::Construct(const FArguments& InArgs, const TSharedPtr<
 	if (DisallowedAssetDataTags.IsValid() || RequiredAssetDataTags.IsValid())
 	{
 		// re-route the filter delegate to our own if we have our own asset data tags filter :
-		AppendOnShouldFilterAssetCallback(FOnShouldFilterAsset::CreateRaw(this, &SPropertyEditorAsset::IsAssetAllowed));
+		AppendOnShouldFilterAssetCallback(FOnShouldFilterAsset::CreateRaw(this, &SPropertyEditorAsset::IsAssetFiltered));
 	}
 
 	if (Property && Property->GetOwnerProperty()->HasMetaData("GetAssetFilter"))
diff --git a/Engine/Source/Editor/PropertyEditor/Private/UserInterface/PropertyEditor/SPropertyEditorAsset.h b/Engine/Source/Editor/PropertyEditor/Private/UserInterface/PropertyEditor/SPropertyEditorAsset.h
--- a/Engine/Source/Editor/PropertyEditor/Private/UserInterface/PropertyEditor/SPropertyEditorAsset.h
+++ b/Engine/Source/Editor/PropertyEditor/Private/UserInterface/PropertyEditor/SPropertyEditorAsset.h
@@ -354,8 +354,8 @@ private:
 	 */
 	void InitializeAssetDataTags(const FProperty* Property);
 
-	/** @return Returns true if the asset is relevant for this property*/
-	bool IsAssetAllowed(const FAssetData& InAssetData);
+	/** @return Returns true if the asset is excluded for this property*/
+	bool IsAssetFiltered(const FAssetData& InAssetData);
 
 private:
 
