Fix crash when PropertyEditor is not valid, use PropertyHandle in this case
#rnx

diff --git a/Engine/Source/Editor/PropertyEditor/Private/UserInterface/PropertyEditor/SPropertyEditorAsset.cpp b/Engine/Source/Editor/PropertyEditor/Private/UserInterface/PropertyEditor/SPropertyEditorAsset.cpp
--- a/Engine/Source/Editor/PropertyEditor/Private/UserInterface/PropertyEditor/SPropertyEditorAsset.cpp
+++ b/Engine/Source/Editor/PropertyEditor/Private/UserInterface/PropertyEditor/SPropertyEditorAsset.cpp
@@ -376,7 +376,14 @@ void SPropertyEditorAsset::Construct(const FArguments& InArgs, const TSharedPtr<
 		if (!GetAssetFilterFunctionName.IsEmpty())
 		{
 			TArray<UObject*> ObjectList;
-			PropertyEditor->GetPropertyHandle()->GetOuterObjects(ObjectList);
+			if (PropertyEditor.IsValid())
+			{
+				PropertyEditor->GetPropertyHandle()->GetOuterObjects(ObjectList);
+			}
+			else if (PropertyHandle.IsValid())
+			{
+				PropertyHandle->GetOuterObjects(ObjectList);
+			}
 			for (UObject* Object : ObjectList)
 			{
 				const UFunction* GetAssetFilterFunction = Object ? Object->FindFunction(*GetAssetFilterFunctionName) : nullptr;
