Details Panel - Overrides widget support phase 1 - UI creation, styles support

* added style classes and keys for the overrides widget image brushes and the code to create them
* added a getter to DetailsViewStyle for the overrides widget combobutton style
* svgs for images
* slight modifications in pixels and alignments in and around the cateogory buttons

diff --git a/Engine/Source/Editor/EditorStyle/Private/StarshipStyle.cpp b/Engine/Source/Editor/EditorStyle/Private/StarshipStyle.cpp
--- a/Engine/Source/Editor/EditorStyle/Private/StarshipStyle.cpp
+++ b/Engine/Source/Editor/EditorStyle/Private/StarshipStyle.cpp
@@ -3044,27 +3044,6 @@ void FStarshipEditorStyle::FStyle::SetupPropertyEditorStyles()
 		CategoryComboButton.ButtonStyle = CategoryRowButton;
 		Set( "DetailsView.CategoryComboButton", CategoryComboButton);
 
-		Set( "DetailsView.OverrideAdded", new CORE_IMAGE_BRUSH_SVG("Starship/Common/OverrideAdded", Icon16x16));
-		Set( "DetailsView.OverrideInside", new CORE_IMAGE_BRUSH_SVG("Starship/Common/OverrideInside", Icon16x16));
-		Set( "DetailsView.OverrideHereInside", new CORE_IMAGE_BRUSH_SVG("Starship/Common/OverrideHereInside", Icon16x16));
-		Set( "DetailsView.OverrideHere", new CORE_IMAGE_BRUSH_SVG("Starship/Common/OverrideHere", Icon16x16));	
-		Set( "DetailsView.OverrideOptions", new CORE_IMAGE_BRUSH_SVG("Starship/Common/OverrideOptions", Icon16x16));
-		Set( "DetailsView.OverrideRemoved", new CORE_IMAGE_BRUSH_SVG("Starship/Common/OverrideRemoved", Icon16x16));
-
-		const FButtonStyle OverridesButton = FButtonStyle()
-			.SetNormalForeground(FStyleColors::Primary)
-			.SetNormal(FSlateRoundedBoxBrush(FStyleColors::Primary, 0.f))
-			.SetHovered(FSlateRoundedBoxBrush(FStyleColors::Primary, 0.f))
-			.SetPressed(FSlateRoundedBoxBrush(FStyleColors::Primary, 0.f))
-			.SetDisabled(FSlateRoundedBoxBrush(FStyleColors::Primary, 0.f))
-			.SetNormalPadding(FMargin(2.f, 0.f, 2.f, 0.f))
-			.SetPressedPadding(FMargin(2.f, 0.f, 2.f, 0.f));
-		
-		FComboButtonStyle OverridesComboButton = FComboButtonStyle(FStarshipCoreStyle::GetCoreStyle().GetWidgetStyle<FComboButtonStyle>("ComboButton"))
-		.SetButtonStyle(OverridesButton)
-		.SetDownArrowPadding(FMargin(2.f, 5.f, 3.f, 5.f));
-		OverridesComboButton.ButtonStyle = CategoryRowButton;
-		Set( "DetailsView.OverridesComboButton", OverridesComboButton);
 		Set( "DetailsView.CategoryTop_Hovered",       new FSlateColorBrush(FStyleColors::Hover));
 		Set( "DetailsView.CategoryBottom",            new FSlateColorBrush(FStyleColors::Recessed));
 		
diff --git a/Engine/Source/Editor/PropertyEditor/Private/DetailsViewStyle.cpp b/Engine/Source/Editor/PropertyEditor/Private/DetailsViewStyle.cpp
--- a/Engine/Source/Editor/PropertyEditor/Private/DetailsViewStyle.cpp
+++ b/Engine/Source/Editor/PropertyEditor/Private/DetailsViewStyle.cpp
@@ -2,6 +2,97 @@
 
 
 #include "DetailsViewStyle.h"
+#include "Containers/Map.h"
+#include "Brushes/SlateImageBrush.h"
+#include "Styling/StyleColors.h"
+#include "Styling/StarshipCoreStyle.h"
+#include "Brushes/SlateRoundedBoxBrush.h"
+
+const FSlateBrush& FOverridesWidgetStyleKey::GetConstStyleBrush() const
+{
+	return ImageBrush;
+}
+
+const FComboButtonStyle& FDetailsViewStyle::GetOverridesComboButtonStyle(
+										const FOverridesWidgetStyleKey* OverridesWidgetStyleKey,
+										const bool bIsOverridesWidgetForOuterCategory) const
+{
+	static TMap<const FOverridesWidgetStyleKey*, const FComboButtonStyle> OverridesKeyToComboButtonStyleMap;
+	const FComboButtonStyle* ComboButtonStylePtr = OverridesKeyToComboButtonStyleMap.Find(OverridesWidgetStyleKey);
+
+	if ( ComboButtonStylePtr )
+	{
+		return *ComboButtonStylePtr;
+	}
+	
+	const FSlateColor BackgroundColor = bIsOverridesWidgetForOuterCategory ? FStyleColors::Header : FStyleColors::Panel;
+	static const FSlateColor HoveredBackgroundColor = FStyleColors::Header;
+			
+	const FButtonStyle OverridesButton = FButtonStyle()
+	                                     .SetNormalForeground(FStyleColors::AccentBlue)
+	                                     .SetHoveredForeground(FStyleColors::AccentBlue)
+	                                     .SetPressedForeground(FStyleColors::AccentBlue)
+	                                     .SetHovered(FSlateRoundedBoxBrush(HoveredBackgroundColor, 0.f))
+	                                     .SetNormal(FSlateRoundedBoxBrush(BackgroundColor, 0.f))
+	                                     .SetPressed(FSlateRoundedBoxBrush(BackgroundColor, 0.f))
+	                                     .SetNormalPadding(FMargin(2.f, 0.f, 0.f, 0.f))
+	                                     .SetPressedPadding(FMargin(2.f, 0.f, 0.f, 0.f));
+			
+
+	OverridesKeyToComboButtonStyleMap.Add(OverridesWidgetStyleKey, FComboButtonStyle(FStarshipCoreStyle::GetCoreStyle().GetWidgetStyle<FComboButtonStyle>("ComboButton"))
+												   .SetButtonStyle(OverridesButton)
+												   .SetDownArrowImage(OverridesWidgetStyleKey->GetConstStyleBrush())
+												   .SetDownArrowPadding(FMargin(2.f, 5.f, 3.f, 5.f)));
+
+		
+	return *OverridesKeyToComboButtonStyleMap.Find(OverridesWidgetStyleKey);
+}
+
+const FOverridesWidgetStyleKey& FOverridesWidgetStyleKeys::Here()
+{
+	static const FOverridesWidgetStyleKey Here{"OverrideHere"};
+	return Here;
+}
+
+const FOverridesWidgetStyleKey& FOverridesWidgetStyleKeys::Added()
+{
+	static const FOverridesWidgetStyleKey Added{"OverrideAdded"};
+	return Added;
+}
+
+const FOverridesWidgetStyleKey& FOverridesWidgetStyleKeys::Options()
+{
+	static const FOverridesWidgetStyleKey Options{"OverrideOptions"};
+	return Options;
+}
+
+const FOverridesWidgetStyleKey& FOverridesWidgetStyleKeys::Removed()
+{
+	static const FOverridesWidgetStyleKey Removed{"OverrideRemoved"};
+	return Removed;
+}
+
+const FOverridesWidgetStyleKey& FOverridesWidgetStyleKeys::Inside()
+{
+	static const FOverridesWidgetStyleKey Inside{"OverrideInside"};
+	return Inside;
+}
+
+const FOverridesWidgetStyleKey& FOverridesWidgetStyleKeys::HereInside()
+{
+	static const FOverridesWidgetStyleKey HereInside{"OverrideHereInside"};
+	return HereInside;
+}
+
+FOverridesWidgetStyleKey::FOverridesWidgetStyleKey(FName InName) : Name{InName}
+{
+	static const FVector2D Icon16x16{16.0f, 16.0f};
+	static const FString Path = FPaths::ConvertRelativePathToFull(FPaths::Combine(FPaths::EngineContentDir(), TEXT("Slate/Starship/Common/")));
+		
+	const FSlateVectorImageBrush Brush{Path + Name.ToString() + ".svg", Icon16x16};
+	const FSlateBrush* SlateBrushPtr = &Brush;
+	ImageBrush = *SlateBrushPtr;
+}
 
 FDetailsViewStyle::FDetailsViewStyle()
 {
diff --git a/Engine/Source/Editor/PropertyEditor/Private/DetailsViewStyle.h b/Engine/Source/Editor/PropertyEditor/Private/DetailsViewStyle.h
--- a/Engine/Source/Editor/PropertyEditor/Private/DetailsViewStyle.h
+++ b/Engine/Source/Editor/PropertyEditor/Private/DetailsViewStyle.h
@@ -1,11 +1,83 @@
 //  Copyright Epic Games, Inc. All Rights Reserved.
 
 #pragma once
+
+#include "CoreMinimal.h"
+#include "Brushes/SlateImageBrush.h"
 #include "DetailsViewStyleKey.h"
 #include "Layout/Margin.h"
-#include "Styling/AppStyle.h"
+#include "Styling/SlateTypes.h"
+#include "Misc/Paths.h"
 #include "Styling/SlateWidgetStyle.h"
-#include "UObject/NameTypes.h"
+
+/**
+ * A class which provides a key with the information to create the Overrides widget style (including the Icom)
+ */
+class FOverridesWidgetStyleKey
+{
+public:
+
+	/**
+	 * The constructor for the style key.
+	 *
+	 * @param InName the FName which is the name of the key
+	 */
+	PROPERTYEDITOR_API FOverridesWidgetStyleKey(FName InName);
+
+	const FSlateBrush& GetConstStyleBrush() const;
+
+	/**
+	 * the name of the key
+	 */
+	const FName Name;
+
+private:
+
+	/**
+	 * The image brush specified by this style key
+	 */
+	FSlateBrush ImageBrush;
+};
+
+/**
+ * The FOverridesWidgetStyleKeys class provides style keys which can
+ * create the needed styles for overrides widgets
+ */
+class FOverridesWidgetStyleKeys
+{
+public:
+
+	/**
+	 * The style for an override widget when an item is completely overridden and has no
+	 * nested properties which are not.
+	 */
+	PROPERTYEDITOR_API static const FOverridesWidgetStyleKey& Here();
+
+	/**
+	 * The style for the override widget when an item has been newly added 
+	 */
+	static const FOverridesWidgetStyleKey& Added();
+
+	/**
+	 * The style for the override widget when the user hovers over it which shows that they have associated
+	 * action options to choose from
+	 */
+	PROPERTYEDITOR_API static const FOverridesWidgetStyleKey& Options();
+
+	/**
+	 * The style for the override widget when an item has been removed 
+	 */
+	static const FOverridesWidgetStyleKey& Removed();
+
+	/**
+	 * The style for the override widget when it has some nested items inside that have been overridden, but
+	 * not all of them
+	 */
+	static const FOverridesWidgetStyleKey& Inside();
+
+	static const FOverridesWidgetStyleKey& HereInside();
+
+};
 
 /**
  * A Class which holds information regarding the style of a Details View
@@ -14,6 +86,13 @@ class FDetailsViewStyle : public FSlateWidgetStyle
 {
 public:
 
+	/**
+	 * returns the const FComboButtonStyle& for the key const FOverridesWidgetStyleKey* OverridesWidgetStyleKey
+	 *
+	 * @param OverridesWidgetStyleKey the FOverridesWidgetStyleKey for the overrides combobutton style
+	 */
+	PROPERTYEDITOR_API const FComboButtonStyle& GetOverridesComboButtonStyle(const FOverridesWidgetStyleKey* OverridesWidgetStyleKey, const bool bIsForOuterCategory = false) const;
+
 	/**
 	 * The default constructor of this @code FDetailsViewStyle @endcode.
 	 */
@@ -45,7 +124,7 @@ public:
 	/**
 	 * Returns the padding for the outer Category row
 	 */
-	 FMargin GetOuterCategoryRowPadding() const;
+	FMargin GetOuterCategoryRowPadding() const;
 	
 	/**
 	 * Returns the padding for details panel rows which are not outer Category rows
@@ -162,9 +241,9 @@ private:
 	/** the Slate Units of the top padding for an outer Category row */
 	float TopCategoryPadding = 0.f;
 	
- 	/**
-	* The FMargin which provides the padding around the whole details view table with scrollbar
-	*/
+	/**
+   * The FMargin which provides the padding around the whole details view table with scrollbar
+   */
 	FMargin TablePaddingWithScrollbar = FMargin(0, 0, 16, 1);
 
 	/**
@@ -175,8 +254,8 @@ private:
 	/**
 	* The FMargin which provides the padding around the Category buttons
 	*/
-	FMargin CategoryButtonsMargin = FMargin(0, 0,  10, 0);
-	
+	FMargin CategoryButtonsMargin = FMargin(0, 0,  0, 0);
 	
 };
+	
 
diff --git a/Engine/Source/Editor/PropertyEditor/Private/SDetailCategoryTableRow.cpp b/Engine/Source/Editor/PropertyEditor/Private/SDetailCategoryTableRow.cpp
--- a/Engine/Source/Editor/PropertyEditor/Private/SDetailCategoryTableRow.cpp
+++ b/Engine/Source/Editor/PropertyEditor/Private/SDetailCategoryTableRow.cpp
@@ -153,7 +153,7 @@ void SDetailCategoryTableRow::Construct(const FArguments& InArgs, TSharedRef<FDe
 
 			SNew(SHorizontalBox)
 			+SHorizontalBox::Slot()
-			.HAlign(HAlign_Left)
+			.HAlign(HAlign_Right)
 			.VAlign(VAlign_Center)
 			.AutoWidth()
 			[
@@ -171,7 +171,7 @@ void SDetailCategoryTableRow::Construct(const FArguments& InArgs, TSharedRef<FDe
 				SNullWidget::NullWidget
 			]
 			+SHorizontalBox::Slot()
-			.HAlign(HAlign_Left)
+			.HAlign(HAlign_Right)
 			.VAlign(VAlign_Center)
 			.AutoWidth()
 			[
diff --git a/Engine/Source/Editor/PropertyEditor/Private/SDetailSingleItemRow.cpp b/Engine/Source/Editor/PropertyEditor/Private/SDetailSingleItemRow.cpp
--- a/Engine/Source/Editor/PropertyEditor/Private/SDetailSingleItemRow.cpp
+++ b/Engine/Source/Editor/PropertyEditor/Private/SDetailSingleItemRow.cpp
@@ -656,7 +656,7 @@ void SDetailSingleItemRow::Construct( const FArguments& InArgs, FDetailLayoutCus
 				SNew(SBorder)
 				.BorderImage(FAppStyle::Get().GetBrush("DetailsView.CategoryMiddle"))
 				.BorderBackgroundColor(this, &SDetailSingleItemRow::GetOuterBackgroundColor)
-				.HAlign(HAlign_Fill)
+				.HAlign(HAlign_Right)
 				.VAlign(VAlign_Center)
 				.Padding(0)
 				[
diff --git a/Engine/Source/Editor/PropertyEditor/Private/UserInterface/Widgets/OverridesComboButtonBuilder.cpp b/Engine/Source/Editor/PropertyEditor/Private/UserInterface/Widgets/OverridesComboButtonBuilder.cpp
--- a/Engine/Source/Editor/PropertyEditor/Private/UserInterface/Widgets/OverridesComboButtonBuilder.cpp
+++ b/Engine/Source/Editor/PropertyEditor/Private/UserInterface/Widgets/OverridesComboButtonBuilder.cpp
@@ -1,11 +1,17 @@
 //  Copyright Epic Games, Inc. All Rights Reserved.
 
 #include "UserInterface/Widgets/OverridesComboButtonBuilder.h"
+
+#include "DetailsViewStyle.h"
 #include "Styling/SlateBrush.h"
+#include "Brushes/SlateRoundedBoxBrush.h"
 
-FOverridesComboButtonBuilder::FOverridesComboButtonBuilder( TSharedRef<FDetailsDisplayManager> InDetailsDisplayManager):
+FOverridesComboButtonBuilder::FOverridesComboButtonBuilder(
+	TSharedRef<FDetailsDisplayManager> InDetailsDisplayManager,
+	bool bInIsCategoryOverridesComboButton):
 	FPropertyUpdatedWidgetBuilder(),
-	DisplayManager(InDetailsDisplayManager)
+	DisplayManager(InDetailsDisplayManager),
+	bIsCategoryOverridesComboButton(bInIsCategoryOverridesComboButton)
 {
 }
 
@@ -18,17 +24,9 @@ FOverridesComboButtonBuilder& FOverridesComboButtonBuilder::Set_OnGetContent(FOn
 
 TSharedPtr<SWidget> FOverridesComboButtonBuilder::GenerateWidget()
 {
-	static FComboButtonStyle NormalStyle =  FAppStyle::Get().GetWidgetStyle<FComboButtonStyle>("DetailsView.OverridesComboButton");
-	static FComboButtonStyle HoverStyle =  FAppStyle::Get().GetWidgetStyle<FComboButtonStyle>("DetailsView.OverridesComboButton");
-	static const FSlateBrush* OverrideOptionsBrush = FAppStyle::GetBrush("DetailsView.OverrideOptions");
-	static const FSlateBrush* OverrideHereBrush = FAppStyle::GetBrush("DetailsView.OverrideHere");
+    const FDetailsViewStyle* DetailsViewStyle = DisplayManager->GetDetailsViewStyle();
 
-	NormalStyle.SetDownArrowImage( *OverrideHereBrush );
-	HoverStyle.SetDownArrowImage( *OverrideOptionsBrush );
- 
-	bool bShouldShowOverrides = true;
-	return 	bShouldShowOverrides  ?
-				SNew(SHorizontalBox)
+	return	SNew(SHorizontalBox)
 				.Visibility(IsVisible)
 				+SHorizontalBox::Slot()
 					
@@ -37,7 +35,8 @@ TSharedPtr<SWidget> FOverridesComboButtonBuilder::GenerateWidget()
 				.AutoWidth()
 				[
 					SNew( SComboButton )
-					.ComboButtonStyle( &NormalStyle )
+					.ComboButtonStyle( &DetailsViewStyle->GetOverridesComboButtonStyle(
+										GetOverridesStyleKey(false ), bIsCategoryOverridesComboButton ) )
 					.Visibility_Lambda([this]
 					{
 						if ( IsRowHoveredAttr.IsBound() && IsRowHoveredAttr.Get())
@@ -55,7 +54,7 @@ TSharedPtr<SWidget> FOverridesComboButtonBuilder::GenerateWidget()
 				.AutoWidth()
 				[
 					SNew( SComboButton )
-					.ComboButtonStyle( &HoverStyle )
+					.ComboButtonStyle(  &DetailsViewStyle->GetOverridesComboButtonStyle( GetOverridesStyleKey(true), bIsCategoryOverridesComboButton ))
 					.Visibility_Lambda([this]
 					{
 						if (IsRowHoveredAttr.IsBound() && IsRowHoveredAttr.Get())
@@ -66,9 +65,7 @@ TSharedPtr<SWidget> FOverridesComboButtonBuilder::GenerateWidget()
 					})
 					.HasDownArrow(true)
 					.OnGetMenuContent(OnGetContent)
-				]
-	:
-				SNullWidget::NullWidget;
+				];
 }
 
 TSharedRef<SWidget> FOverridesComboButtonBuilder::operator*()
@@ -76,7 +73,21 @@ TSharedRef<SWidget> FOverridesComboButtonBuilder::operator*()
 	return GenerateWidget().ToSharedRef();
 }
 
+const FOverridesWidgetStyleKey* FOverridesComboButtonBuilder::GetOverridesStyleKey(bool bIsHoveredOver) const
+{
+	if (bIsCategoryOverridesComboButton && !bIsHoveredOver)
+	{
+		return &FOverridesWidgetStyleKeys::Here();
+	} 
+	if (!bIsHoveredOver)
+	{
+		return &FOverridesWidgetStyleKeys::Here();
+	}
+	
+	return &FOverridesWidgetStyleKeys::Options();
+}
+
 FOverridesComboButtonBuilder::~FOverridesComboButtonBuilder()
 {
 	OnGetContent.Unbind();
-}
\ No newline at end of file
+}
diff --git a/Engine/Source/Editor/PropertyEditor/Public/UserInterface/Widgets/OverridesComboButtonBuilder.h b/Engine/Source/Editor/PropertyEditor/Public/UserInterface/Widgets/OverridesComboButtonBuilder.h
--- a/Engine/Source/Editor/PropertyEditor/Public/UserInterface/Widgets/OverridesComboButtonBuilder.h
+++ b/Engine/Source/Editor/PropertyEditor/Public/UserInterface/Widgets/OverridesComboButtonBuilder.h
@@ -8,6 +8,7 @@
 #include "Templates/SharedPointer.h"
 
 class FDetailsDisplayManager;
+class FOverridesWidgetStyleKey;
 
 DECLARE_DELEGATE_OneParam(FGetIsRowHoveredOver, bool)
 
@@ -26,8 +27,14 @@ public:
 	/**
 	 * The constructor, which takes a @code TSharedRef<FDetailsDisplayManager> @endcode to initialize
 	 * the Details Display Manager
+	 *
+	 * @param InDetailsDisplayManager the FDetailsDisplayManager which manages the details display
+	 * @param bInIsCategoryOverridesComboButton if true, this FOverridesComboButtonBuilder is for a Category rather
+	 * than a property row   
 	 */
-	PROPERTYEDITOR_API FOverridesComboButtonBuilder( TSharedRef<FDetailsDisplayManager> InDetailsDisplayManager );
+	PROPERTYEDITOR_API FOverridesComboButtonBuilder(
+		TSharedRef<FDetailsDisplayManager> InDetailsDisplayManager,
+		bool bInIsCategoryOverridesComboButton);
 
 	/**
 	 * Set the OnGetContent for the menu that this button is responsible for
@@ -49,10 +56,20 @@ public:
 	 */
 	TSharedRef<SWidget> operator*();
 
+	/**
+	 * returns the const FOverridesWidgetStyleKey* which points to the proper FOverridesWidgetStyleKey for this
+	 */
+	const FOverridesWidgetStyleKey* GetOverridesStyleKey(bool bIsHoveredOver = false) const;
+
 private:
 	/**
 	 * The @code DetailsDisplayManager @endcode which provides an API to manage some of the characteristics of the
 	 * details display
 	 */
 	TSharedRef<FDetailsDisplayManager> DisplayManager;
+
+	/**
+	 * if true, this is an overrides combo button for a Category 
+	 */
+	bool bIsCategoryOverridesComboButton;
 };
