Fix missing initialization of WholeRowHeaderContent in SDetailCategoryTableRow.

diff --git a/Engine/Source/Editor/PropertyEditor/Private/SDetailCategoryTableRow.h b/Engine/Source/Editor/PropertyEditor/Private/SDetailCategoryTableRow.h
--- a/Engine/Source/Editor/PropertyEditor/Private/SDetailCategoryTableRow.h
+++ b/Engine/Source/Editor/PropertyEditor/Private/SDetailCategoryTableRow.h
@@ -19,6 +19,7 @@ class SDetailCategoryTableRow : public SDetailTableRowBase
 public:
 	SLATE_BEGIN_ARGS(SDetailCategoryTableRow)
 		: _InnerCategory(false)
+		, _WholeRowHeaderContent(false)
 		, _ShowBorder(true)
 		, _IsEmpty(false)
 		, _ObjectName(NAME_Name)
