Added single struct property support to PropertyEditor
Integrated OOM issue fix from [REDACTED_HANDLE] when building a single struct property

diff --git a/Engine/Source/Editor/PropertyEditor/Private/DetailPropertyRow.cpp b/Engine/Source/Editor/PropertyEditor/Private/DetailPropertyRow.cpp
--- a/Engine/Source/Editor/PropertyEditor/Private/DetailPropertyRow.cpp
+++ b/Engine/Source/Editor/PropertyEditor/Private/DetailPropertyRow.cpp
@@ -539,7 +539,7 @@ void FDetailPropertyRow::MakeExternalPropertyRowCustomization(TSharedPtr<FStruct
 	InitParams.ArrayOffset = 0;
 	InitParams.ArrayIndex = INDEX_NONE;
 	InitParams.bForceHiddenPropertyVisibility = Parameters.ShouldForcePropertyVisible() || FPropertySettings::Get().ShowHiddenProperties();
-	InitParams.bCreateCategoryNodes = false;
+	InitParams.bCreateCategoryNodes = PropertyName == NAME_None;
 	InitParams.bAllowChildren = false;
 
 	Parameters.OverrideAllowChildren(InitParams.bAllowChildren);
@@ -551,20 +551,13 @@ void FDetailPropertyRow::MakeExternalPropertyRowCustomization(TSharedPtr<FStruct
 
 	if (PropertyName != NAME_None)
 	{
-		RootPropertyNode->RebuildChildren();
-
-		for (int32 ChildIdx = 0; ChildIdx < RootPropertyNode->GetNumChildNodes(); ++ChildIdx)
+		TSharedPtr<FPropertyNode> PropertyNode = RootPropertyNode->GenerateSingleChild(PropertyName);
+		if (PropertyNode.IsValid())
 		{
-			TSharedPtr< FPropertyNode > PropertyNode = RootPropertyNode->GetChildNode(ChildIdx);
-			if (FProperty* Property = PropertyNode->GetProperty())
-			{
-				if (Property->GetFName() == PropertyName)
-				{
-					OutCustomization.PropertyRow = MakeShareable(new FDetailPropertyRow(PropertyNode, ParentCategory, RootPropertyNode));
-					OutCustomization.PropertyRow->SetCustomExpansionId(Parameters.GetUniqueId());
-					break;
-				}
-			}
+			PropertyNode->RebuildChildren();
+
+			OutCustomization.PropertyRow = MakeShared<FDetailPropertyRow>(PropertyNode, ParentCategory, RootPropertyNode);
+			OutCustomization.PropertyRow->SetCustomExpansionId(Parameters.GetUniqueId());
 		}
 	}
 	else
diff --git a/Engine/Source/Editor/PropertyEditor/Private/ObjectPropertyNode.h b/Engine/Source/Editor/PropertyEditor/Private/ObjectPropertyNode.h
--- a/Engine/Source/Editor/PropertyEditor/Private/ObjectPropertyNode.h
+++ b/Engine/Source/Editor/PropertyEditor/Private/ObjectPropertyNode.h
@@ -134,7 +134,7 @@ public:
 
 
 	/** Generates a single child from the provided property name.  Any existing children are destroyed */
-	TSharedPtr<FPropertyNode> GenerateSingleChild( FName ChildPropertyName );
+	virtual TSharedPtr<FPropertyNode> GenerateSingleChild( FName ChildPropertyName ) override;
 
 	/**
 	 * @return The hidden categories 
diff --git a/Engine/Source/Editor/PropertyEditor/Private/PropertyEditorModule.cpp b/Engine/Source/Editor/PropertyEditor/Private/PropertyEditorModule.cpp
--- a/Engine/Source/Editor/PropertyEditor/Private/PropertyEditorModule.cpp
+++ b/Engine/Source/Editor/PropertyEditor/Private/PropertyEditorModule.cpp
@@ -296,35 +296,51 @@ TSharedPtr<IDetailsView> FPropertyEditorModule::FindDetailView( const FName View
 }
 
 TSharedPtr<ISinglePropertyView> FPropertyEditorModule::CreateSingleProperty( UObject* InObject, FName InPropertyName, const FSinglePropertyParams& InitParams )
+{
+	return CreateSinglePropertyImpl(InObject, TSharedPtr<IStructureDataProvider>(), InPropertyName, InitParams);
+}
+
+TSharedPtr<class ISinglePropertyView> FPropertyEditorModule::CreateSingleProperty( const TSharedPtr<IStructureDataProvider>& InStruct, FName InPropertyName, const struct FSinglePropertyParams& InitParams )
+{
+	return CreateSinglePropertyImpl(nullptr, InStruct, InPropertyName, InitParams);
+}
+
+TSharedPtr<class ISinglePropertyView> FPropertyEditorModule::CreateSinglePropertyImpl(UObject* InObject, const TSharedPtr<IStructureDataProvider>& InStruct, FName InPropertyName, const struct FSinglePropertyParams& InitParams)
 {
 	// Compact the list of detail view instances
+	CompactSinglePropertyViewArray();
+
+	TSharedRef<SSingleProperty> Property =
+		SNew(SSingleProperty)
+		.Object(InObject)
+		.StructData(InStruct)
+		.PropertyName(InPropertyName)
+		.NamePlacement(InitParams.NamePlacement)
+		.NameOverride(InitParams.NameOverride)
+		.NotifyHook(InitParams.NotifyHook)
+		.PropertyFont(InitParams.Font)
+		.bShouldHideAssetThumbnail(InitParams.bHideAssetThumbnail);
+
+	if (Property->HasValidProperty())
+	{
+		AllSinglePropertyViews.Add(Property);
+
+		return Property;
+	}
+
+	return nullptr;
+}
+
+void FPropertyEditorModule::CompactSinglePropertyViewArray()
+{
 	for( int32 ViewIndex = 0; ViewIndex < AllSinglePropertyViews.Num(); ++ViewIndex )
 	{
-		if ( !AllSinglePropertyViews[ViewIndex].IsValid() )
+		if (!AllSinglePropertyViews[ViewIndex].IsValid())
 		{
-			AllSinglePropertyViews.RemoveAtSwap( ViewIndex );
+			AllSinglePropertyViews.RemoveAtSwap(ViewIndex);
 			--ViewIndex;
 		}
 	}
-
-	TSharedRef<SSingleProperty> Property = 
-		SNew( SSingleProperty )
-		.Object( InObject )
-		.PropertyName( InPropertyName )
-		.NamePlacement( InitParams.NamePlacement )
-		.NameOverride( InitParams.NameOverride )
-		.NotifyHook( InitParams.NotifyHook )
-		.PropertyFont( InitParams.Font )
-		.bShouldHideAssetThumbnail( InitParams.bHideAssetThumbnail);
-
-	if( Property->HasValidProperty() )
-	{
-		AllSinglePropertyViews.Add( Property );
-
-		return Property;
-	}
-
-	return NULL;
 }
 
 TSharedRef< IPropertyTable > FPropertyEditorModule::CreatePropertyTable()
diff --git a/Engine/Source/Editor/PropertyEditor/Private/PropertyNode.h b/Engine/Source/Editor/PropertyEditor/Private/PropertyNode.h
--- a/Engine/Source/Editor/PropertyEditor/Private/PropertyNode.h
+++ b/Engine/Source/Editor/PropertyEditor/Private/PropertyNode.h
@@ -1334,4 +1334,7 @@ public:
 	virtual EPropertyType GetPropertyType() const = 0;
 
 	virtual void Disconnect() = 0;
+
+	/** Generates a single child from the provided property name.  Any existing children are destroyed */
+	virtual TSharedPtr<FPropertyNode> GenerateSingleChild(FName ChildPropertyName) = 0;
 };
diff --git a/Engine/Source/Editor/PropertyEditor/Private/SSingleProperty.cpp b/Engine/Source/Editor/PropertyEditor/Private/SSingleProperty.cpp
--- a/Engine/Source/Editor/PropertyEditor/Private/SSingleProperty.cpp
+++ b/Engine/Source/Editor/PropertyEditor/Private/SSingleProperty.cpp
@@ -18,6 +18,7 @@
 #include "UObject/UnrealType.h"
 #include "Widgets/Colors/SColorPicker.h"
 #include "Widgets/Text/STextBlock.h"
+#include "StructurePropertyNode.h"
 
 
 class FSinglePropertyUtilities : public IPropertyUtilities
@@ -117,26 +118,67 @@ void SSingleProperty::Construct( const FArguments& InArgs )
 
 	PropertyUtilities = MakeShareable( new FSinglePropertyUtilities( SharedThis( this ), InArgs._bShouldHideAssetThumbnail ) );
 
-	SetObject( InArgs._Object );
+	if (InArgs._Object != nullptr)
+	{
+		SetObject( InArgs._Object );
+	}
+	else if(InArgs._StructData.IsValid())
+	{
+		SetStruct(InArgs._StructData);
+	}
 }
 
 void SSingleProperty::SetObject( UObject* InObject )
 {
-	DestroyColorPicker();
-
-	if( !RootPropertyNode.IsValid() )
+	if( !RootPropertyNode.IsValid() || RootPropertyNode->GetPropertyType() != FComplexPropertyNode::EPT_Object)
 	{
 		RootPropertyNode = MakeShareable( new FObjectPropertyNode );
 	}
 
-	RootPropertyNode->RemoveAllObjects();
+	FObjectPropertyNode* RootObjectPropertyNode = static_cast<FObjectPropertyNode*>(RootPropertyNode.Get());
+
+	RootObjectPropertyNode->RemoveAllObjects();
 	ValueNode.Reset();
 
 	if( InObject )
 	{
-		RootPropertyNode->AddObject( InObject );
+		RootObjectPropertyNode->AddObject( InObject );
+	}
+
+	if( !GeneratePropertyCustomization() )
+	{
+		// invalid or missing property
+		RootObjectPropertyNode->RemoveAllObjects();
+	}
+}
+
+void SSingleProperty::SetStruct(const TSharedPtr<IStructureDataProvider>& InStruct)
+{
+	if (!RootPropertyNode.IsValid() || RootPropertyNode->GetPropertyType() != FComplexPropertyNode::EPT_StandaloneStructure)
+	{
+		RootPropertyNode = MakeShareable(new FStructurePropertyNode);
+	}
+
+	FStructurePropertyNode* RootStructPropertyNode = (FStructurePropertyNode*)RootPropertyNode.Get();
+
+	RootStructPropertyNode->RemoveStructure();
+	ValueNode.Reset();
+
+	if (InStruct)
+	{
+		RootStructPropertyNode->SetStructure(InStruct);
 	}
 
+	if( !GeneratePropertyCustomization() )
+	{
+		// invalid or missing property
+		RootStructPropertyNode->RemoveStructure();
+	}
+}
+
+bool SSingleProperty::GeneratePropertyCustomization()
+{
+	DestroyColorPicker();
 
 	FPropertyNodeInitParams InitParams;
 	InitParams.ParentNode = NULL;
@@ -244,13 +286,14 @@ void SSingleProperty::SetObject( UObject* InObject )
 			.ToolTipText(NSLOCTEXT("PropertyEditor", "SinglePropertyInvalidType_Tooltip", "Properties of this type cannot be edited inline; edit it elsewhere"))
 		];
 
-		// invalid or missing property
-		RootPropertyNode->RemoveAllObjects();
 		ValueNode.Reset();
 		RootPropertyNode.Reset();
 	}
+
+	return bIsAcceptableProperty;
 }
 
+
 void SSingleProperty::SetOnPropertyValueChanged( FSimpleDelegate& InOnPropertyValueChanged )
 {
 	if( HasValidProperty() )
@@ -261,13 +304,15 @@ void SSingleProperty::SetOnPropertyValueChanged( FSimpleDelegate& InOnPropertyVa
 
 void SSingleProperty::ReplaceObjects( const TMap<UObject*, UObject*>& OldToNewObjectMap )
 {
-	if( HasValidProperty() )
+	if( HasValidProperty() && RootPropertyNode->GetPropertyType() == FComplexPropertyNode::EPT_Object )
 	{
 		TArray<UObject*> NewObjectList;
 		bool bObjectsReplaced = false;
 
+		FObjectPropertyNode* RootObjectPropertyNode = static_cast<FObjectPropertyNode*>(RootPropertyNode.Get());
+
 		// Scan all objects and look for objects which need to be replaced
-		for ( TPropObjectIterator Itor( RootPropertyNode->ObjectIterator() ); Itor; ++Itor )
+		for ( TPropObjectIterator Itor(RootObjectPropertyNode->ObjectIterator() ); Itor; ++Itor )
 		{
 			UObject* Replacement = OldToNewObjectMap.FindRef( Itor->Get(true) );
 			if( Replacement )
@@ -292,10 +337,12 @@ void SSingleProperty::ReplaceObjects( const TMap<UObject*, UObject*>& OldToNewOb
 void SSingleProperty::RemoveDeletedObjects( const TArray<UObject*>& DeletedObjects )
 {
 
-	if( HasValidProperty() )
+	if( HasValidProperty() && RootPropertyNode->GetPropertyType() == FComplexPropertyNode::EPT_Object )
 	{
+		FObjectPropertyNode* RootObjectPropertyNode = static_cast<FObjectPropertyNode*>(RootPropertyNode.Get());
+
 		// Scan all objects and look for objects which need to be replaced
-		for ( TPropObjectIterator Itor( RootPropertyNode->ObjectIterator() ); Itor; ++Itor )
+		for ( TPropObjectIterator Itor(RootObjectPropertyNode->ObjectIterator() ); Itor; ++Itor )
 		{
 			if( DeletedObjects.Contains( Itor->Get() ) )
 			{
diff --git a/Engine/Source/Editor/PropertyEditor/Private/SSingleProperty.h b/Engine/Source/Editor/PropertyEditor/Private/SSingleProperty.h
--- a/Engine/Source/Editor/PropertyEditor/Private/SSingleProperty.h
+++ b/Engine/Source/Editor/PropertyEditor/Private/SSingleProperty.h
@@ -12,12 +12,14 @@ class FNotifyHook;
 class FObjectPropertyNode;
 class FPropertyNode;
 class FSinglePropertyUtilities;
+class FStructurePropertyNode;
 
 class SSingleProperty : public ISinglePropertyView
 {
 public:
 	SLATE_BEGIN_ARGS( SSingleProperty )
 		: _Object(NULL)
+		, _StructData(NULL)
 		, _NotifyHook( NULL )
 		, _PropertyFont( FAppStyle::GetFontStyle( PropertyEditorConstants::PropertyFontStyle ) ) 
 		, _NamePlacement( EPropertyNamePlacement::Left )
@@ -26,6 +28,7 @@ public:
 	{}
 
 		SLATE_ARGUMENT( UObject*, Object )
+		SLATE_ARGUMENT(TSharedPtr<IStructureDataProvider>, StructData)
 		SLATE_ARGUMENT( FName, PropertyName )
 		SLATE_ARGUMENT( FNotifyHook*, NotifyHook )
 		SLATE_ARGUMENT( FSlateFontInfo, PropertyFont )
@@ -37,8 +40,9 @@ public:
 	void Construct( const FArguments& InArgs );
 
 	/** ISinglePropertyView interface */
-	virtual bool HasValidProperty() const override { return RootPropertyNode.IsValid() && ValueNode.IsValid(); }
+	virtual bool HasValidProperty() const override { return (RootPropertyNode.IsValid() || RootPropertyNode) && ValueNode.IsValid(); }
 	virtual void SetObject( UObject* InObject ) override;
+	virtual void SetStruct( const TSharedPtr<IStructureDataProvider>& InStruct ) override;
 	virtual void SetOnPropertyValueChanged( FSimpleDelegate& InOnPropertyValueChanged ) override;
 
 	/**
@@ -69,9 +73,17 @@ private:
 	 * @param NewColor The color to set
 	 */
 	void SetColorPropertyFromColorPicker(FLinearColor NewColor);
+
+	/**
+	 * Generates the SingleProperty customization
+	 *
+	 * @return true if valid property and widget has been generated
+	 */
+	bool GeneratePropertyCustomization();
+
 private:
 	/** The root property node for the value node (contains the root object */
-	TSharedPtr<FObjectPropertyNode> RootPropertyNode;
+	TSharedPtr<FComplexPropertyNode> RootPropertyNode;
 	/** The node for the property being edited */
 	TSharedPtr<FPropertyNode> ValueNode;
 	/** Property utilities for handling common functionality of property editors */
diff --git a/Engine/Source/Editor/PropertyEditor/Private/StructurePropertyNode.cpp b/Engine/Source/Editor/PropertyEditor/Private/StructurePropertyNode.cpp
--- a/Engine/Source/Editor/PropertyEditor/Private/StructurePropertyNode.cpp
+++ b/Engine/Source/Editor/PropertyEditor/Private/StructurePropertyNode.cpp
@@ -6,6 +6,11 @@
 #include "PropertyEditorHelpers.h"
 
 void FStructurePropertyNode::InitChildNodes()
+{
+	InternalInitChildNodes(FName());
+}
+
+void FStructurePropertyNode::InternalInitChildNodes(FName SinglePropertyName)
 {
 	const bool bShouldShowHiddenProperties = !!HasNodeFlags(EPropertyNodeFlags::ShouldShowHiddenProperties);
 	const bool bShouldShowDisableEditOnInstance = !!HasNodeFlags(EPropertyNodeFlags::ShouldShowDisableEditOnInstance);
@@ -19,7 +24,14 @@ void FStructurePropertyNode::InitChildNodes()
 		FProperty* StructMember = *It;
 		if (PropertyEditorHelpers::ShouldBeVisible(*this, StructMember))
 		{
-			StructMembers.Add(StructMember);
+			if (SinglePropertyName == NAME_None || StructMember->GetFName() == SinglePropertyName)
+			{
+				StructMembers.Add(StructMember);
+				if (SinglePropertyName != NAME_None)
+				{
+					break;
+				}
+			}
 		}
 	}
 
@@ -34,7 +46,7 @@ void FStructurePropertyNode::InitChildNodes()
 		InitParams.Property = StructMember;
 		InitParams.ArrayOffset = 0;
 		InitParams.ArrayIndex = INDEX_NONE;
-		InitParams.bAllowChildren = true;
+		InitParams.bAllowChildren = SinglePropertyName == NAME_None;
 		InitParams.bForceHiddenPropertyVisibility = bShouldShowHiddenProperties;
 		InitParams.bCreateDisableEditOnInstanceNodes = bShouldShowDisableEditOnInstance;
 		InitParams.bCreateCategoryNodes = false;
@@ -94,6 +106,27 @@ uint8* FStructurePropertyNode::GetValueBaseAddress(uint8* StartAddress, bool bIs
 	return nullptr;
 }
 
+TSharedPtr<FPropertyNode> FStructurePropertyNode::GenerateSingleChild(FName ChildPropertyName)
+{
+	constexpr bool bDestroySelf = false;
+	DestroyTree(bDestroySelf);
+
+	// No category nodes should be created in single property mode
+	SetNodeFlags(EPropertyNodeFlags::ShowCategories, false);
+
+	InternalInitChildNodes(ChildPropertyName);
+
+	if (ChildNodes.Num() > 0)
+	{
+		// only one node should be been created
+		check(ChildNodes.Num() == 1);
+
+		return ChildNodes[0];
+	}
+
+	return nullptr;
+}
+
 EPropertyDataValidationResult FStructurePropertyNode::EnsureDataIsValid()
 {
 	CachedReadAddresses.Reset();
diff --git a/Engine/Source/Editor/PropertyEditor/Private/StructurePropertyNode.h b/Engine/Source/Editor/PropertyEditor/Private/StructurePropertyNode.h
--- a/Engine/Source/Editor/PropertyEditor/Private/StructurePropertyNode.h
+++ b/Engine/Source/Editor/PropertyEditor/Private/StructurePropertyNode.h
@@ -353,6 +353,9 @@ public:
 		StructProvider = nullptr;
 	}
 
+	/** Generates a single child from the provided property name.  Any existing children are destroyed */
+	virtual TSharedPtr<FPropertyNode> GenerateSingleChild(FName ChildPropertyName) override;
+
 protected:
 
 	virtual EPropertyDataValidationResult EnsureDataIsValid() override;
@@ -360,6 +363,8 @@ protected:
 	/** FPropertyNode interface */
 	virtual void InitChildNodes() override;
 
+	void InternalInitChildNodes(FName SinglePropertyName);
+
 	virtual uint8* GetValueBaseAddress(uint8* Base, bool bIsSparseData, bool bIsStruct) const override;
 
 	virtual bool GetQualifiedName(FString& PathPlusIndex, const bool bWithArrayIndex, const FPropertyNode* StopParent = nullptr, bool bIgnoreCategories = false) const override
diff --git a/Engine/Source/Editor/PropertyEditor/Public/ISinglePropertyView.h b/Engine/Source/Editor/PropertyEditor/Public/ISinglePropertyView.h
--- a/Engine/Source/Editor/PropertyEditor/Public/ISinglePropertyView.h
+++ b/Engine/Source/Editor/PropertyEditor/Public/ISinglePropertyView.h
@@ -50,6 +50,9 @@ public:
 	/** Sets the object to view/edit on the widget */
 	virtual void SetObject( UObject* InObject ) = 0;
 
+	/** Sets the struct to view/edit on the widget */
+	virtual void SetStruct( const TSharedPtr<class IStructureDataProvider>& InStruct) = 0;
+
 	/** Sets a delegate called when the property value changes */
 	virtual void SetOnPropertyValueChanged( FSimpleDelegate& InOnPropertyValueChanged ) = 0;
 
diff --git a/Engine/Source/Editor/PropertyEditor/Public/PropertyEditorModule.h b/Engine/Source/Editor/PropertyEditor/Public/PropertyEditorModule.h
--- a/Engine/Source/Editor/PropertyEditor/Public/PropertyEditorModule.h
+++ b/Engine/Source/Editor/PropertyEditor/Public/PropertyEditorModule.h
@@ -355,7 +355,7 @@ public:
 	virtual TSharedRef<SWindow> CreateFloatingDetailsView( const TArray< UObject* >& InObjects, bool bIsLockable );
 
 	/**
-	 * Creates a standalone widget for a single property
+	 * Creates a standalone widget for a single object property
 	 *
 	 * @param InObject			The object to view
 	 * @param InPropertyName	The name of the property to display
@@ -364,6 +364,16 @@ public:
 	 */
 	virtual TSharedPtr<class ISinglePropertyView> CreateSingleProperty( UObject* InObject, FName InPropertyName, const struct FSinglePropertyParams& InitParams );
 
+	/**
+	 * Creates a standalone widget for a single struct property
+	 *
+	 * @param InStruct			The struct containing the property to view
+	 * @param InPropertyName	The name of the property to display
+	 * @param InitParams		Optional init params for a single property
+	 * @return The new property if valid or null
+	 */
+	virtual TSharedPtr<class ISinglePropertyView> CreateSingleProperty(const TSharedPtr<class IStructureDataProvider>& InStruct, FName InPropertyName, const struct FSinglePropertyParams& InitParams);
+
 	virtual TSharedRef<class IStructureDetailsView> CreateStructureDetailView(const struct FDetailsViewArgs& DetailsViewArgs, const FStructureDetailsViewArgs& StructureDetailsViewArgs, TSharedPtr<class FStructOnScope> StructData, const FText& CustomName = FText::GetEmpty());
 
 	virtual TSharedRef<class IPropertyRowGenerator> CreatePropertyRowGenerator(const struct FPropertyRowGeneratorArgs& InArgs);
@@ -441,6 +451,9 @@ private:
 	void GetAllSectionsHelper(const UStruct* Struct, TArray<TSharedPtr<FPropertySection>>& OutSections, TSet<const UStruct*>& ProcessedStructs) const;
 	void FindSectionsForCategoryHelper(const UStruct* Struct, FName CategoryName, TArray<TSharedPtr<FPropertySection>>& OutSections, TSet<const UStruct*>& SearchedStructs) const;
 
+	TSharedPtr<class ISinglePropertyView> CreateSinglePropertyImpl(UObject* InObject, const TSharedPtr<IStructureDataProvider>& InStruct, FName InPropertyName, const struct FSinglePropertyParams& InitParams);
+	void CompactSinglePropertyViewArray();
+
 private:
 	/** All created detail views */
 	TArray< TWeakPtr<class SDetailsView> > AllDetailViews;
