ShowOnlyInnerProperties now works on nested structs to show the properties on the innermost struct marked with the specifier.

diff --git a/Engine/Source/Editor/PropertyEditor/Private/DetailLayoutHelpers.cpp b/Engine/Source/Editor/PropertyEditor/Private/DetailLayoutHelpers.cpp
--- a/Engine/Source/Editor/PropertyEditor/Private/DetailLayoutHelpers.cpp
+++ b/Engine/Source/Editor/PropertyEditor/Private/DetailLayoutHelpers.cpp
@@ -87,8 +87,9 @@ namespace DetailLayoutHelpers
 					}
 				}
 
-				// Whether or not to push out struct properties to their own categories or show them inside an expandable struct 
-				const bool bPushOutStructProps = bIsStruct && !bIsCustomizedStruct && !ParentStructProp && Property->HasMetaData(ShowOnlyInners);
+				// Whether or not to push out struct properties to their own categories or show them inside an expandable struct
+				// This recursively applies for any nested structs that have the ShowOnlyInners metadata 
+				const bool bPushOutStructProps = bIsStruct && !bIsCustomizedStruct && Property->HasMetaData(ShowOnlyInners);
 
 				// Is the property edit inline new 
 				const bool bIsEditInlineNew = ChildNode.HasNodeFlags(EPropertyNodeFlags::ShowInnerObjectProperties) || SPropertyEditorEditInline::Supports(&ChildNode, ChildNode.GetArrayIndex());
