Add GetAllowedClasses and GetDisallowedClasses metadata
#tests PIE (single and multi selection)

diff --git a/Engine/Source/Editor/PropertyEditor/Private/PropertyEditorUtils.cpp b/Engine/Source/Editor/PropertyEditor/Private/PropertyEditorUtils.cpp
--- a/Engine/Source/Editor/PropertyEditor/Private/PropertyEditorUtils.cpp
+++ b/Engine/Source/Editor/PropertyEditor/Private/PropertyEditorUtils.cpp
@@ -1,6 +1,8 @@
 ﻿// Copyright Epic Games, Inc. All Rights Reserved.
 
 #include "PropertyEditorUtils.h"
+
+#include "PropertyCustomizationHelpers.h"
 #include "PropertyPathHelpers.h"
 
 namespace PropertyEditorUtils
@@ -67,4 +69,90 @@ namespace PropertyEditorUtils
 			Algo::Transform(OptionIntersection, InOutOptions, [](const FString& InString) { return MakeShared<FString>(InString); });
 		}
 	}
+
+	void GetAllowedAndDisallowedClasses(const TArray<UObject*>& ObjectList, const FProperty& MetadataProperty, TArray<const UClass*>& AllowedClasses, TArray<const UClass*>& DisallowedClasses, bool bExactClass, const UClass* ObjectClass)
+	{
+		AllowedClasses = PropertyCustomizationHelpers::GetClassesFromMetadataString(MetadataProperty.GetOwnerProperty()->GetMetaData("AllowedClasses"));
+		DisallowedClasses = PropertyCustomizationHelpers::GetClassesFromMetadataString(MetadataProperty.GetOwnerProperty()->GetMetaData("DisallowedClasses"));
+		
+		bool bMergeAllowedClasses = !AllowedClasses.IsEmpty();
+
+		if (MetadataProperty.GetOwnerProperty()->HasMetaData("GetAllowedClasses"))
+		{
+			const FString GetAllowedClassesFunctionName = MetadataProperty.GetOwnerProperty()->GetMetaData("GetAllowedClasses");
+			if (!GetAllowedClassesFunctionName.IsEmpty())
+			{
+				for (UObject* Object : ObjectList)
+				{
+					const UFunction* GetAllowedClassesFunction = Object ? Object->FindFunction(*GetAllowedClassesFunctionName) : nullptr;
+					if (GetAllowedClassesFunction)
+					{
+						DECLARE_DELEGATE_RetVal(TArray<UClass*>, FGetAllowedClasses);
+						if (!bMergeAllowedClasses)
+						{
+							AllowedClasses.Append(FGetAllowedClasses::CreateUFunction(Object, GetAllowedClassesFunction->GetFName()).Execute());
+							if (AllowedClasses.IsEmpty())
+							{
+								// No allowed class means all classes are valid
+								continue;
+							}
+							bMergeAllowedClasses = true;
+						}
+						else
+						{
+							TArray<UClass*> MergedClasses = FGetAllowedClasses::CreateUFunction(Object, GetAllowedClassesFunction->GetFName()).Execute();
+							if (MergedClasses.IsEmpty())
+							{
+								// No allowed class means all classes are valid
+								continue;
+							}
+							
+							TArray<const UClass*> CurrentAllowedClassFilters = MoveTemp(AllowedClasses);
+							ensure(AllowedClasses.IsEmpty());
+							for (const UClass* MergedClass : MergedClasses)
+							{
+								// Keep classes that match both allow list
+								for (const UClass* CurrentClass : CurrentAllowedClassFilters)
+								{
+									if (CurrentClass == MergedClass || (!bExactClass && CurrentClass->IsChildOf(MergedClass)))
+									{
+										AllowedClasses.Add(CurrentClass);
+										break;
+									}
+									if (!bExactClass && MergedClass->IsChildOf(CurrentClass))
+									{
+										AllowedClasses.Add(MergedClass);
+										break;
+									}
+								}
+							}
+							if (AllowedClasses.IsEmpty())
+							{
+								// An empty AllowedClasses array means that everything is allowed: in this case, forbid UObject
+								DisallowedClasses.Add(ObjectClass);
+								return;
+							}
+						}
+					}
+				}
+			}
+		}
+
+		if (MetadataProperty.GetOwnerProperty()->HasMetaData("GetDisallowedClasses"))
+		{
+			const FString GetDisallowedClassesFunctionName = MetadataProperty.GetOwnerProperty()->GetMetaData("GetDisallowedClasses");
+			if (!GetDisallowedClassesFunctionName.IsEmpty())
+			{
+				for (UObject* Object : ObjectList)
+				{
+					const UFunction* GetDisallowedClassesFunction = Object ? Object->FindFunction(*GetDisallowedClassesFunctionName) : nullptr;
+					if (GetDisallowedClassesFunction)
+					{
+						DECLARE_DELEGATE_RetVal(TArray<UClass*>, FGetDisallowedClasses);
+						DisallowedClasses.Append(FGetDisallowedClasses::CreateUFunction(Object, GetDisallowedClassesFunction->GetFName()).Execute());
+					}
+				}
+			}
+		}
+	}
 }
diff --git a/Engine/Source/Editor/PropertyEditor/Private/UserInterface/PropertyEditor/SPropertyEditorAsset.cpp b/Engine/Source/Editor/PropertyEditor/Private/UserInterface/PropertyEditor/SPropertyEditorAsset.cpp
--- a/Engine/Source/Editor/PropertyEditor/Private/UserInterface/PropertyEditor/SPropertyEditorAsset.cpp
+++ b/Engine/Source/Editor/PropertyEditor/Private/UserInterface/PropertyEditor/SPropertyEditorAsset.cpp
@@ -37,6 +37,7 @@
 #include "DetailWidgetRow.h"
 #include "Subsystems/AssetEditorSubsystem.h"
 #include "PropertyEditorConstants.h"
+#include "PropertyEditorUtils.h"
 #include "Misc/EditorPathHelper.h"
 
 #define LOCTEXT_NAMESPACE "PropertyEditor"
@@ -195,14 +196,23 @@ void SPropertyEditorAsset::InitializeClassFilters(const FProperty* Property)
 
 	bExactClass = GetTagOrBoolMetadata(MetadataProperty, "ExactClass", false);
 	
-	AllowedClassFilters = PropertyCustomizationHelpers::GetClassesFromMetadataString(MetadataProperty->GetMetaData("AllowedClasses"));
+	TArray<UObject*> ObjectList;
+	if (PropertyEditor && PropertyEditor->GetPropertyHandle()->IsValidHandle())
+	{
+		PropertyEditor->GetPropertyHandle()->GetOuterObjects(ObjectList);
+	}
+	else if (PropertyHandle.IsValid())
+	{
+		PropertyHandle->GetOuterObjects(ObjectList);
+	}
+	
+	PropertyEditorUtils::GetAllowedAndDisallowedClasses(ObjectList, *MetadataProperty, AllowedClassFilters, DisallowedClassFilters, bExactClass, ObjectClass);
+	
 	if (AllowedClassFilters.Num() == 0)
 	{
 		// always add the object class to the filters
 		AllowedClassFilters.Add(ObjectClass);
 	}
-
-	DisallowedClassFilters = PropertyCustomizationHelpers::GetClassesFromMetadataString(MetadataProperty->GetMetaData("DisallowedClasses"));
 }
 
 void SPropertyEditorAsset::InitializeAssetDataTags(const FProperty* Property)
diff --git a/Engine/Source/Editor/PropertyEditor/Private/UserInterface/PropertyEditor/SPropertyEditorClass.cpp b/Engine/Source/Editor/PropertyEditor/Private/UserInterface/PropertyEditor/SPropertyEditorClass.cpp
--- a/Engine/Source/Editor/PropertyEditor/Private/UserInterface/PropertyEditor/SPropertyEditorClass.cpp
+++ b/Engine/Source/Editor/PropertyEditor/Private/UserInterface/PropertyEditor/SPropertyEditorClass.cpp
@@ -9,6 +9,7 @@
 #include "DragAndDrop/AssetDragDropOp.h"
 #include "ClassViewerModule.h"
 #include "ClassViewerFilter.h"
+#include "PropertyEditorUtils.h"
 #include "UObject/UObjectIterator.h"
 #include "PropertyNode.h"
 #include "PropertyRestriction.h"
@@ -161,8 +162,12 @@ void SPropertyEditorClass::Construct(const FArguments& InArgs, const TSharedPtr<
 		}
 
 		// Filter based on UPROPERTY meta data
-		AllowedClassFilters = PropertyCustomizationHelpers::GetClassesFromMetadataString(Property->GetOwnerProperty()->GetMetaData("AllowedClasses"));
-		DisallowedClassFilters = PropertyCustomizationHelpers::GetClassesFromMetadataString(Property->GetOwnerProperty()->GetMetaData("DisallowedClasses"));
+		TArray<UObject*> ObjectList;
+		if (PropertyEditor->GetPropertyHandle()->IsValidHandle())
+		{
+			PropertyEditor->GetPropertyHandle()->GetOuterObjects(ObjectList);
+		}
+		PropertyEditorUtils::GetAllowedAndDisallowedClasses(ObjectList, *Property, AllowedClassFilters, DisallowedClassFilters, false);
 
 		using namespace UE::PropertyEditor::Class::Private;
 
@@ -444,8 +449,6 @@ FReply SPropertyEditorClass::OnDrop(const FGeometry& MyGeometry, const FDragDrop
 	TSharedPtr<FAssetDragDropOp> UnloadedClassOp = DragDropEvent.GetOperationAs<FAssetDragDropOp>();
 	if (UnloadedClassOp.IsValid())
 	{
-		bool bAllAssetWereLoaded = true;
-
 		FString AssetPath;
 
 		// Find the class/blueprint path
diff --git a/Engine/Source/Editor/PropertyEditor/Private/UserInterface/PropertyEditor/SPropertyEditorEditInline.cpp b/Engine/Source/Editor/PropertyEditor/Private/UserInterface/PropertyEditor/SPropertyEditorEditInline.cpp
--- a/Engine/Source/Editor/PropertyEditor/Private/UserInterface/PropertyEditor/SPropertyEditorEditInline.cpp
+++ b/Engine/Source/Editor/PropertyEditor/Private/UserInterface/PropertyEditor/SPropertyEditorEditInline.cpp
@@ -13,6 +13,7 @@
 #include "Styling/SlateIconFinder.h"
 #include "UObject/ConstructorHelpers.h"
 #include "Editor.h"
+#include "PropertyEditorUtils.h"
 
 class FPropertyEditorInlineClassFilter : public IClassViewerFilter
 {
@@ -266,12 +267,16 @@ TSharedRef<SWidget> SPropertyEditorEditInline::GenerateClassPicker()
 		}
 	}
 
-	static const FName NAME_AllowedClasses(ANSITEXTVIEW("AllowedClasses"));
-	TArray<const UClass*> AllowedClassFilters = PropertyCustomizationHelpers::GetClassesFromMetadataString(MetadataProperty->GetMetaData(NAME_AllowedClasses));
-
-	static const FName NAME_DisallowedClasses(ANSITEXTVIEW("DisallowedClasses"));
-	TArray<const UClass*> DisallowedClassFilters = PropertyCustomizationHelpers::GetClassesFromMetadataString(MetadataProperty->GetMetaData(NAME_DisallowedClasses));
+	TArray<UObject*> ObjectList;
+	if (PropertyEditor && PropertyEditor->GetPropertyHandle()->IsValidHandle())
+	{
+		PropertyEditor->GetPropertyHandle()->GetOuterObjects(ObjectList);
+	}
 
+	TArray<const UClass*> AllowedClassFilters;
+	TArray<const UClass*> DisallowedClassFilters;
+	PropertyEditorUtils::GetAllowedAndDisallowedClasses(ObjectList, *Property, AllowedClassFilters, DisallowedClassFilters, false);
+	
 	using namespace UE::PropertyEditor::EditInline::Private;
 
 	// Filter based on restrictions
diff --git a/Engine/Source/Editor/PropertyEditor/Public/PropertyEditorUtils.h b/Engine/Source/Editor/PropertyEditor/Public/PropertyEditorUtils.h
--- a/Engine/Source/Editor/PropertyEditor/Public/PropertyEditorUtils.h
+++ b/Engine/Source/Editor/PropertyEditor/Public/PropertyEditorUtils.h
@@ -13,4 +13,15 @@ namespace PropertyEditorUtils
 	 * @param	InOutOptions		The resulting options
 	 */
 	PROPERTYEDITOR_API void GetPropertyOptions(TArray<UObject*>& InOutContainers, FString& InOutPropertyPath, TArray<TSharedPtr<FString>>& InOutOptions);
+
+	/**
+	 * Get all allowed and disallowed classes according to meta data
+	 * @param	ObjectList			The list of object that owns the property
+	 * @param	Property			The FProperty that contains the meta data
+	 * @param	AllowedClasses		The list of allowed classes
+	 * @param	DisallowedClasses	The list of disallowed classes
+	 * @param	bExactClass			Whether subclasses allowed or not
+	 * @param	ObjectClass			The class to add to the disallowed classes if there is no allowed classes
+	 */
+	PROPERTYEDITOR_API void GetAllowedAndDisallowedClasses(const TArray<UObject*>& ObjectList, const FProperty& Property, TArray<const UClass*>& AllowedClasses, TArray<const UClass*>& DisallowedClasses, bool bExactClass, const UClass* ObjectClass = UObject::StaticClass());
 }
diff --git a/Engine/Source/Runtime/CoreUObject/Public/UObject/ObjectMacros.h b/Engine/Source/Runtime/CoreUObject/Public/UObject/ObjectMacros.h
--- a/Engine/Source/Runtime/CoreUObject/Public/UObject/ObjectMacros.h
+++ b/Engine/Source/Runtime/CoreUObject/Public/UObject/ObjectMacros.h
@@ -1483,6 +1483,18 @@ namespace UM
 		/// UFUNCTION()
 		/// bool FuncName(const FAssetData& AssetData) const;
 		GetAssetFilter,
+		
+		/// [PropertyMetadata] Get allowed classes through a UFunction, e.g. meta=(GetAllowedClasses="FuncName").
+		///
+		/// UFUNCTION()
+		/// TArray<UClass*> FuncName() const;
+		GetAllowedClasses,
+		
+		/// [PropertyMetadata] Get disallowed classes through a UFunction, e.g. meta=(GetDisallowedClasses="FuncName").
+		///
+		/// UFUNCTION()
+		/// TArray<UClass*> FuncName() const;
+		GetDisallowedClasses,
 
 	    /// [PropertyMetadata] Used for enum properties to define the subset of invalid values as a comma-separated string; values within this subset are not made available.
         InvalidEnumValues,
