Fix warning V501: There are identical sub-expressions to the left and to the right of the '==' operator: SlotName == SlotName
#rnx

diff --git a/Engine/Source/Editor/PropertyEditor/Public/MaterialList.h b/Engine/Source/Editor/PropertyEditor/Public/MaterialList.h
--- a/Engine/Source/Editor/PropertyEditor/Public/MaterialList.h
+++ b/Engine/Source/Editor/PropertyEditor/Public/MaterialList.h
@@ -134,7 +134,7 @@ struct FMaterialListItem
 
 	bool operator==( const FMaterialListItem& Other ) const
 	{
-		return Material == Other.Material && SlotIndex == Other.SlotIndex && SlotName == SlotName;
+		return Material == Other.Material && SlotIndex == Other.SlotIndex && SlotName == Other.SlotName;
 	}
 
 	bool operator!=( const FMaterialListItem& Other ) const
