Support setting actor picker with the currently selected actor in the level editor

diff --git a/Engine/Source/Editor/PropertyEditor/Private/PropertyCustomizationHelpers.cpp b/Engine/Source/Editor/PropertyEditor/Private/PropertyCustomizationHelpers.cpp
--- a/Engine/Source/Editor/PropertyEditor/Private/PropertyCustomizationHelpers.cpp
+++ b/Engine/Source/Editor/PropertyEditor/Private/PropertyCustomizationHelpers.cpp
@@ -142,11 +142,15 @@ namespace PropertyCustomizationHelpers
 			.IsFocusable( false );
 	}
 
-	TSharedRef<SWidget> MakeUseSelectedButton( FSimpleDelegate OnUseSelectedClicked, TAttribute<FText> OptionalToolTipText, TAttribute<bool> IsEnabled )
+	TSharedRef<SWidget> MakeUseSelectedButton( FSimpleDelegate OnUseSelectedClicked, TAttribute<FText> OptionalToolTipText, TAttribute<bool> IsEnabled, const bool IsActor )
 	{
 		return
 			SNew( SPropertyEditorButton )
-			.Text( OptionalToolTipText.Get().IsEmpty() ? LOCTEXT( "UseButtonToolTipText", "Use Selected Asset from Content Browser") : OptionalToolTipText )
+			.Text(
+				!OptionalToolTipText.Get().IsEmpty() ? OptionalToolTipText
+				: IsActor ? LOCTEXT( "UseActorButtonToolTipText", "Use Selected Actor from the Level Editor")
+				: LOCTEXT( "UseButtonToolTipText", "Use Selected Asset from Content Browser")
+			)
 			.Image( FAppStyle::GetBrush("Icons.Use") )
 			.OnClickAction( OnUseSelectedClicked )
 			.IsEnabled(IsEnabled)
diff --git a/Engine/Source/Editor/PropertyEditor/Private/UserInterface/PropertyEditor/SPropertyEditorAsset.cpp b/Engine/Source/Editor/PropertyEditor/Private/UserInterface/PropertyEditor/SPropertyEditorAsset.cpp
--- a/Engine/Source/Editor/PropertyEditor/Private/UserInterface/PropertyEditor/SPropertyEditorAsset.cpp
+++ b/Engine/Source/Editor/PropertyEditor/Private/UserInterface/PropertyEditor/SPropertyEditorAsset.cpp
@@ -665,14 +665,14 @@ void SPropertyEditorAsset::Construct(const FArguments& InArgs, const TSharedPtr<
 		];
 	}
 
-	if( !bIsActor && InArgs._DisplayUseSelected )
+	if( InArgs._DisplayUseSelected )
 	{
 		ButtonBox->AddSlot()
 		.VAlign(VAlign_Center)
 		.AutoWidth()
 		.Padding( 2.0f, 0.0f )
 		[
-			PropertyCustomizationHelpers::MakeUseSelectedButton( FSimpleDelegate::CreateSP( this, &SPropertyEditorAsset::OnUse ), FText(), IsEnabledAttribute )
+			PropertyCustomizationHelpers::MakeUseSelectedButton( FSimpleDelegate::CreateSP( this, &SPropertyEditorAsset::OnUse ), FText(), IsEnabledAttribute, bIsActor )
 		];
 	}
 
@@ -1372,6 +1372,7 @@ void SPropertyEditorAsset::OnUse()
 	// Use the property editor path if it is valid and there is no custom filtering required
 	if(PropertyEditor.IsValid()
 		&& !OnShouldFilterAsset.IsBound()
+		&& !OnShouldFilterActor.IsBound()
 		&& AllowedClassFilters.Num() == 0
 		&& DisallowedClassFilters.Num() == 0
 		&& (GEditor ? !GEditor->MakeAssetReferenceFilter(FAssetReferenceFilterContext()).IsValid() : true))
@@ -1388,6 +1389,12 @@ void SPropertyEditorAsset::OnUse()
 		if( ObjectClass && ObjectClass->IsChildOf( AActor::StaticClass() ) )
 		{
 			Selection = GEditor->GetSelectedActors()->GetTop( ObjectClass );
+
+			// For actors filtered means allowed, unlike for assets (where filtered means NOT allowed)
+			if (!IsFilteredActor(static_cast<const AActor*>(Selection)))
+			{
+				Selection = nullptr;
+			}
 		}
 		else if( ObjectClass )
 		{
diff --git a/Engine/Source/Editor/PropertyEditor/Public/PropertyCustomizationHelpers.h b/Engine/Source/Editor/PropertyEditor/Public/PropertyCustomizationHelpers.h
--- a/Engine/Source/Editor/PropertyEditor/Public/PropertyCustomizationHelpers.h
+++ b/Engine/Source/Editor/PropertyEditor/Public/PropertyCustomizationHelpers.h
@@ -100,7 +100,7 @@ namespace PropertyCustomizationHelpers
 	PROPERTYEDITOR_API TSharedRef<SWidget> MakeClearButton( FSimpleDelegate OnClearClicked, TAttribute<FText> OptionalToolTipText = FText(), TAttribute<bool> IsEnabled = true );
 	PROPERTYEDITOR_API TSharedRef<SWidget> MakeVisibilityButton(FOnClicked OnVisibilityClicked, TAttribute<FText> OptionalToolTipText = FText(), TAttribute<bool> VisibilityDelegate = true);
 	PROPERTYEDITOR_API TSharedRef<SWidget> MakeNewBlueprintButton( FSimpleDelegate OnNewBlueprintClicked, TAttribute<FText> OptionalToolTipText = FText(), TAttribute<bool> IsEnabled = true );
-	PROPERTYEDITOR_API TSharedRef<SWidget> MakeUseSelectedButton( FSimpleDelegate OnUseSelectedClicked, TAttribute<FText> OptionalToolTipText = FText(), TAttribute<bool> IsEnabled = true );
+	PROPERTYEDITOR_API TSharedRef<SWidget> MakeUseSelectedButton( FSimpleDelegate OnUseSelectedClicked, TAttribute<FText> OptionalToolTipText = FText(), TAttribute<bool> IsEnabled = true, const bool IsActor = false );
 	PROPERTYEDITOR_API TSharedRef<SWidget> MakeBrowseButton( FSimpleDelegate OnFindClicked, TAttribute<FText> OptionalToolTipText = FText(), TAttribute<bool> IsEnabled = true, const bool IsActor = false);
 	PROPERTYEDITOR_API TSharedRef<SWidget> MakeAssetPickerAnchorButton( FOnGetAllowedClasses OnGetAllowedClasses, FOnAssetSelected OnAssetSelectedFromPicker, const TSharedPtr<IPropertyHandle>& PropertyHandle = TSharedPtr<IPropertyHandle>());
 	PROPERTYEDITOR_API TSharedRef<SWidget> MakeAssetPickerWithMenu( const FAssetData& InitialObject, const bool AllowClear, const TArray<const UClass*>& AllowedClasses, const TArray<UFactory*>& NewAssetFactories, FOnShouldFilterAsset OnShouldFilterAsset, FOnAssetSelected OnSet, FSimpleDelegate OnClose, const TSharedPtr<IPropertyHandle>& PropertyHandle = TSharedPtr<IPropertyHandle>(), const TArray<FAssetData>& OwnerAssetArray = TArray<FAssetData>());
