Customizations that use FDetailPropertyRow::CustomWidget and set the visibility attribute do not honor it

Customizations that use the "CustomWidget" call on the FDetailPropertyRow are not being added to the tickable nodes for the details panel for the visibility to be reassessed, and the tick function is what checks and updates the visibility attribute. To fix:

* In DetailItemNode.cpp,  check to see if it has a CustomPropertyWidget for the row  and if it does, then add it to the tickable nodes

* In  FDetailPropertyRow::GetPropertyVisibility() ensure to update the visibility based on the custom widget

diff --git a/Engine/Source/Editor/PropertyEditor/Private/DetailItemNode.cpp b/Engine/Source/Editor/PropertyEditor/Private/DetailItemNode.cpp
--- a/Engine/Source/Editor/PropertyEditor/Private/DetailItemNode.cpp
+++ b/Engine/Source/Editor/PropertyEditor/Private/DetailItemNode.cpp
@@ -31,7 +31,12 @@ FDetailItemNode::FDetailItemNode(const FDetailLayoutCustomization& InCustomizati
 
 void FDetailItemNode::Initialize()
 {
-	if( ( Customization.HasCustomWidget() && Customization.WidgetDecl->VisibilityAttr.IsBound() )
+	bool bHasCustomPropertyRowWidget = Customization.PropertyRow 
+			                            && ( Customization.PropertyRow->CustomNameWidget()
+			                                || Customization.PropertyRow->CustomValueWidget());
+	
+	if( bHasCustomPropertyRowWidget
+	    || ( Customization.HasCustomWidget() && Customization.WidgetDecl->VisibilityAttr.IsBound() )
 		|| ( Customization.HasCustomBuilder() && Customization.CustomBuilderRow->RequiresTick() )
 		|| ( Customization.HasPropertyNode() && Customization.PropertyRow->RequiresTick() )
 		|| ( Customization.HasGroup() && Customization.DetailGroup->RequiresTick() ) )
diff --git a/Engine/Source/Editor/PropertyEditor/Private/DetailPropertyRow.cpp b/Engine/Source/Editor/PropertyEditor/Private/DetailPropertyRow.cpp
--- a/Engine/Source/Editor/PropertyEditor/Private/DetailPropertyRow.cpp
+++ b/Engine/Source/Editor/PropertyEditor/Private/DetailPropertyRow.cpp
@@ -640,6 +640,10 @@ EVisibility FDetailPropertyRow::GetPropertyVisibility() const
 	{
 		return EVisibility::Collapsed;
 	}
+	else if (CustomPropertyWidget.IsValid() && CustomPropertyWidget->VisibilityAttr.IsBound())
+	{
+		return CustomPropertyWidget->VisibilityAttr.Get();
+	}
 
 	return PropertyVisibility.Get();
 }
