Fix behaviour of material details panel to respect whether a selected actor is editable or not.

diff --git a/Engine/Source/Editor/PropertyEditor/Private/MaterialList.cpp b/Engine/Source/Editor/PropertyEditor/Private/MaterialList.cpp
--- a/Engine/Source/Editor/PropertyEditor/Private/MaterialList.cpp
+++ b/Engine/Source/Editor/PropertyEditor/Private/MaterialList.cpp
@@ -789,6 +789,17 @@ void FMaterialList::AddMaterialItem( FDetailWidgetRow& Row, int32 CurrentSlot, c
 	[
 		RightSideContent.ToSharedRef()
 	];
+
+	if(USceneComponent* SceneComponent = Cast<USceneComponent>(InActorComponent))
+	{
+		UObject* OwnerObject = nullptr;
+		FString PropertyPath;
+		FProperty* MaterialProperty = nullptr;
+		if (SceneComponent->GetMaterialPropertyPath(Item.SlotIndex, OwnerObject, PropertyPath, MaterialProperty))
+		{
+			Row.IsEnabled(SceneComponent->CanEditChange(MaterialProperty));
+		}		
+	}
 }
 
 #undef LOCTEXT_NAMESPACE
\ No newline at end of file
