[CrashFix] Editor crashes after enabling filters in the Data Asset Merge tool

diff --git a/Engine/Source/Editor/PropertyEditor/Private/SDetailsSplitter.cpp b/Engine/Source/Editor/PropertyEditor/Private/SDetailsSplitter.cpp
--- a/Engine/Source/Editor/PropertyEditor/Private/SDetailsSplitter.cpp
+++ b/Engine/Source/Editor/PropertyEditor/Private/SDetailsSplitter.cpp
@@ -302,6 +302,11 @@ namespace DetailsSplitterHelpers
 
 	bool CanCopyPropertyValue(const TSharedPtr<FDetailTreeNode>& SourceDetailsNode, const TSharedPtr<FDetailTreeNode>& DestinationDetailsNode, ETreeDiffResult Diff)
 	{
+		if (!SourceDetailsNode || !DestinationDetailsNode)
+		{
+			return false;
+		}
+		
 		// in order to copy properties there needs to be the same number of objects in each panel
 		const TArray<TWeakObjectPtr<UObject>>* SourceObjects = DetailsSplitterHelpers::GetObjects(SourceDetailsNode);
 		const TArray<TWeakObjectPtr<UObject>>* DestinationObjects = DetailsSplitterHelpers::GetObjects(DestinationDetailsNode);
