Fixed a crash at SSingleProperty when it is created with an object and the requested property does not exist in the object, due to the root node getting destroyed on failure before the call to destroy the root node objects

#rnx

diff --git a/Engine/Source/Editor/PropertyEditor/Private/SSingleProperty.cpp b/Engine/Source/Editor/PropertyEditor/Private/SSingleProperty.cpp
--- a/Engine/Source/Editor/PropertyEditor/Private/SSingleProperty.cpp
+++ b/Engine/Source/Editor/PropertyEditor/Private/SSingleProperty.cpp
@@ -149,6 +149,7 @@ void SSingleProperty::SetObject( UObject* InObject )
 	{
 		// invalid or missing property
 		RootObjectPropertyNode->RemoveAllObjects();
+		RootPropertyNode.Reset();
 	}
 }
 
@@ -174,6 +175,7 @@ void SSingleProperty::SetStruct(const TSharedPtr<IStructureDataProvider>& InStru
 	{
 		// invalid or missing property
 		RootStructPropertyNode->RemoveStructure();
+		RootPropertyNode.Reset();
 	}
 }
 
@@ -288,7 +290,6 @@ bool SSingleProperty::GeneratePropertyCustomization()
 		];
 
 		ValueNode.Reset();
-		RootPropertyNode.Reset();
 	}
 
 	return bIsAcceptableProperty;
