Added a way to define EditConditionHides behavior for FDetailPropertyRows with custom EditConditions defined.
- Added EditConditionHides interface method to IDetailPropertyRow (and corresponding implementation in FDetailPropertyRow). This serves as a counterpart for the existing EditCondition interface method exposed for defining custom edit conditions on manually created FDetailPropertyRows.

diff --git a/Engine/Source/Editor/PropertyEditor/Private/DetailPropertyRow.cpp b/Engine/Source/Editor/PropertyEditor/Private/DetailPropertyRow.cpp
--- a/Engine/Source/Editor/PropertyEditor/Private/DetailPropertyRow.cpp
+++ b/Engine/Source/Editor/PropertyEditor/Private/DetailPropertyRow.cpp
@@ -106,6 +106,12 @@ IDetailPropertyRow& FDetailPropertyRow::EditCondition( TAttribute<bool> EditCond
 	return *this;
 }
 
+IDetailPropertyRow& FDetailPropertyRow::EditConditionHides(bool bEditConditionHidesValue)
+{
+	bCustomEditConditionHides = bEditConditionHidesValue;
+	return *this;
+}
+
 IDetailPropertyRow& FDetailPropertyRow::IsEnabled( TAttribute<bool> InIsEnabled )
 {
 	CustomIsEnabledAttrib = InIsEnabled;
@@ -175,8 +181,7 @@ bool FDetailPropertyRow::ShowOnlyChildren() const
 
 bool FDetailPropertyRow::RequiresTick() const
 {
-	return PropertyVisibility.IsBound() || 
-		(PropertyEditor.IsValid() && PropertyEditor->IsOnlyVisibleWhenEditConditionMet());
+	return PropertyVisibility.IsBound() || IsOnlyVisibleWhenEditConditionMet();
 }
 
 FDetailWidgetRow& FDetailPropertyRow::CustomWidget( bool bShowChildren )
@@ -636,7 +641,7 @@ void FDetailPropertyRow::MakeExternalPropertyRowCustomization(const TArray<UObje
 
 EVisibility FDetailPropertyRow::GetPropertyVisibility() const
 {
-	if (PropertyEditor.IsValid() && PropertyEditor->IsOnlyVisibleWhenEditConditionMet() && !PropertyEditor->IsEditConditionMet())
+	if (IsOnlyVisibleWhenEditConditionMet() && !IsEditConditionMet())
 	{
 		return EVisibility::Collapsed;
 	}
@@ -658,20 +663,32 @@ bool FDetailPropertyRow::GetEnabledState() const
 	bool Result = IsParentEnabled.Get(true);
 
 	Result = Result && CustomIsEnabledAttrib.Get(true);
+	Result = Result && IsEditConditionMet();
+
+	return Result;
+}
 
+bool FDetailPropertyRow::IsEditConditionMet() const
+{
 	if (HasEditCondition())
 	{
 		if (CustomEditConditionValue.IsSet())
 		{
-			Result = Result && CustomEditConditionValue.Get();
+			return CustomEditConditionValue.Get();
 		}
 		else if (PropertyEditor.IsValid())
 		{
-			Result = Result && PropertyEditor->IsEditConditionMet();
+			return PropertyEditor->IsEditConditionMet();
 		}
 	}
 
-	return Result;
+	// Default to True, matching the default output of FPropertyNode::IsEditConditionMet().
+	return true;
+}
+
+bool FDetailPropertyRow::IsOnlyVisibleWhenEditConditionMet() const
+{
+	return (PropertyEditor.IsValid() && PropertyEditor->IsOnlyVisibleWhenEditConditionMet()) || (bCustomEditConditionHides && HasEditCondition());
 }
 
 TSharedPtr<IPropertyTypeCustomization>& FDetailPropertyRow::GetTypeInterface()
diff --git a/Engine/Source/Editor/PropertyEditor/Private/DetailPropertyRow.h b/Engine/Source/Editor/PropertyEditor/Private/DetailPropertyRow.h
--- a/Engine/Source/Editor/PropertyEditor/Private/DetailPropertyRow.h
+++ b/Engine/Source/Editor/PropertyEditor/Private/DetailPropertyRow.h
@@ -31,6 +31,7 @@ public:
 	virtual IDetailPropertyRow& ToolTip( const FText& InToolTip ) override;
 	virtual IDetailPropertyRow& ShowPropertyButtons( bool bInShowPropertyButtons ) override;
 	virtual IDetailPropertyRow& EditCondition( TAttribute<bool> EditConditionValue, FOnBooleanValueChanged OnEditConditionValueChanged ) override;
+	virtual IDetailPropertyRow& EditConditionHides(bool bEditConditionHideValue) override;
 	virtual IDetailPropertyRow& IsEnabled(TAttribute<bool> InIsEnabled) override;
 	virtual IDetailPropertyRow& ShouldAutoExpand(bool bForceExpansion) override;
 	virtual IDetailPropertyRow& Visibility( TAttribute<EVisibility> Visibility ) override;
@@ -150,6 +151,16 @@ private:
 	 */
 	bool GetEnabledState() const;
 
+	/**
+	 * @return true if this row's edit condition has been met
+	 */
+	bool IsEditConditionMet() const;
+
+	/**
+	 * @return true if this row should be shown/hidden based on the edit condition state
+	 */
+	bool IsOnlyVisibleWhenEditConditionMet() const;
+
 	/** 
 	 * Looks up and caches (if not already) the associated type interface, and then returns it.
 	 * Could be null/invalid if no customization was specified for this row.
@@ -216,6 +227,8 @@ private:
 	TAttribute<bool> CustomEditConditionValue;
 	/** User customized edit condition change handler. */
 	FOnBooleanValueChanged CustomEditConditionValueChanged;
+	/** User customized edit condition hides specifier */
+	bool bCustomEditConditionHides = false;
 	/** User customized reset to default */
 	TOptional<FResetToDefaultOverride> CustomResetToDefault;
 	/** User customized drag/drop handler */
diff --git a/Engine/Source/Editor/PropertyEditor/Public/IDetailPropertyRow.h b/Engine/Source/Editor/PropertyEditor/Public/IDetailPropertyRow.h
--- a/Engine/Source/Editor/PropertyEditor/Public/IDetailPropertyRow.h
+++ b/Engine/Source/Editor/PropertyEditor/Public/IDetailPropertyRow.h
@@ -207,6 +207,13 @@ public:
 	 */
 	virtual IDetailPropertyRow& EditCondition( TAttribute<bool> EditConditionValue, FOnBooleanValueChanged OnEditConditionValueChanged ) = 0;
 
+	/**
+	 * Sets whether or not the edit condition for this property should affect its visibility.  If the edit condition fails, the property will be hidden outright.
+	 * 
+	 * @param bEditConditionHidesValue		if true the property be shown/hidden based on the edit condition
+	 */
+	virtual IDetailPropertyRow& EditConditionHides( bool bEditConditionHidesValue ) = 0;
+
 	/**
 	 * Sets whether or not this property is enabled
 	 *
