Insights scopes

diff --git a/Engine/Source/Editor/PropertyEditor/Private/SDetailsView.cpp b/Engine/Source/Editor/PropertyEditor/Private/SDetailsView.cpp
--- a/Engine/Source/Editor/PropertyEditor/Private/SDetailsView.cpp
+++ b/Engine/Source/Editor/PropertyEditor/Private/SDetailsView.cpp
@@ -512,6 +512,7 @@ EVisibility SDetailsView::GetActorNameAreaVisibility() const
 
 void SDetailsView::ForceRefresh()
 {
+	TRACE_CPUPROFILER_EVENT_SCOPE_STR("SDetailsView::ForceRefresh");
 	TArray<UObject*> NewObjectList;
 	NewObjectList.Reserve(UnfilteredSelectedObjects.Num());
 	TArray<TWeakObjectPtr<UObject>> ValidSelectedObjects;
@@ -538,6 +539,7 @@ void SDetailsView::MoveScrollOffset(int32 DeltaOffset)
 
 void SDetailsView::SetObjects(const TArray<UObject*>& InObjects, bool bForceRefresh/* = false*/, bool bOverrideLock/* = false*/)
 {
+	TRACE_CPUPROFILER_EVENT_SCOPE_STR("SDetailsView::SetObjects");
 	if (!IsLocked() || bOverrideLock)
 	{
 		if( bForceRefresh || ShouldSetNewObjects(InObjects) )
diff --git a/Engine/Source/Editor/UnrealEd/Private/ObjectTools.cpp b/Engine/Source/Editor/UnrealEd/Private/ObjectTools.cpp
--- a/Engine/Source/Editor/UnrealEd/Private/ObjectTools.cpp
+++ b/Engine/Source/Editor/UnrealEd/Private/ObjectTools.cpp
@@ -5656,6 +5656,7 @@ namespace ThumbnailTools
 	/** Loads thumbnails from the specified package file name, try loading from external cache file if not found in package file */
 	bool LoadThumbnailsFromPackage( const FString& InPackageFileName, const TSet< FName >& InObjectFullNames, FThumbnailMap& InOutThumbnails )
 	{
+		TRACE_CPUPROFILER_EVENT_SCOPE_STR("ThumbnailTools::LoadThumbnailsFromPackage");
 		if (LoadThumbnailsFromPackageInternal(InPackageFileName, InObjectFullNames, InOutThumbnails))
 		{
 			return true;
