[PropertyEditor] Added UToolMenu extensibility to PropertyEditor context menu

diff --git a/Engine/Source/Editor/PropertyEditor/Private/DetailRowMenuContextPrivate.h b/Engine/Source/Editor/PropertyEditor/Private/DetailRowMenuContextPrivate.h
new file mode 100644
--- /dev/null
+++ b/Engine/Source/Editor/PropertyEditor/Private/DetailRowMenuContextPrivate.h
@@ -0,0 +1,36 @@
+﻿// Copyright Epic Games, Inc. All Rights Reserved.
+
+#pragma once
+
+#include "UObject/Object.h"
+
+#include "DetailRowMenuContextPrivate.generated.h"
+
+class SDetailTableRowBase;
+
+/** The private context provides the slate widget. */
+UCLASS()
+class UDetailRowMenuContextPrivate : public UObject
+{
+	GENERATED_BODY()
+
+public:
+	template <typename RowType>
+	TSharedPtr<RowType> GetRowWidget() const
+	{
+		if (const TWeakPtr<RowType> RowContextWeak = StaticCastWeakPtr<RowType>(Row);
+			RowContextWeak.IsValid())
+		{
+			if (TSharedPtr<RowType> RowContext = RowContextWeak.Pin())
+			{
+				return RowContext;
+			}
+		}
+
+		return nullptr;
+	}
+
+public:
+	TWeakPtr<SDetailTableRowBase> Row;
+};
+
diff --git a/Engine/Source/Editor/PropertyEditor/Private/DetailWidgetRow.cpp b/Engine/Source/Editor/PropertyEditor/Private/DetailWidgetRow.cpp
--- a/Engine/Source/Editor/PropertyEditor/Private/DetailWidgetRow.cpp
+++ b/Engine/Source/Editor/PropertyEditor/Private/DetailWidgetRow.cpp
@@ -5,3 +5,13 @@
 
 const float FDetailWidgetRow::DefaultValueMinWidth = 125.0f;
 const float FDetailWidgetRow::DefaultValueMaxWidth = 125.0f;
+
+const FName FDetailWidgetRow::FCustomMenuData::GetEntryName() const
+{
+	if (!EntryName.IsNone())
+	{
+		return EntryName;
+	}
+	
+	return FName(Name.ToString().Replace(TEXT(" "), TEXT("")));
+}
diff --git a/Engine/Source/Editor/PropertyEditor/Private/PropertyEditorModule.cpp b/Engine/Source/Editor/PropertyEditor/Private/PropertyEditorModule.cpp
--- a/Engine/Source/Editor/PropertyEditor/Private/PropertyEditorModule.cpp
+++ b/Engine/Source/Editor/PropertyEditor/Private/PropertyEditorModule.cpp
@@ -1,8 +1,8 @@
 // Copyright Epic Games, Inc. All Rights Reserved.
 
-
 #include "PropertyEditorModule.h"
 #include "AssetToolsModule.h"
+#include "DetailRowMenuContextPrivate.h"
 #include "IAssetTools.h"
 #include "IDetailsView.h"
 #include "IPropertyChangeListener.h"
@@ -32,7 +32,7 @@
 #include "Widgets/Colors/SColorPicker.h"
 #include "Widgets/Layout/SBorder.h"
 #include "DetailsViewStyle.h"
-
+#include "ToolMenus.h"
 
 IMPLEMENT_MODULE( FPropertyEditorModule, PropertyEditor );
 
@@ -94,10 +94,15 @@ void FPropertyEditorModule::StartupModule()
 	FCoreUObjectDelegates::OnObjectsReplaced.AddRaw(this, &FPropertyEditorModule::ReplaceViewedObjects);
 
 	FDetailsViewStyle::InitializeDetailsViewStyles();
+
+	UToolMenus::RegisterStartupCallback(FSimpleMulticastDelegate::FDelegate::CreateRaw(this, &FPropertyEditorModule::RegisterMenus));
 }
 
 void FPropertyEditorModule::ShutdownModule()
 {
+	UToolMenus::UnRegisterStartupCallback(this);
+	UToolMenus::UnregisterOwner(this);
+	
 	// No need to remove this object from root since the final GC pass doesn't care about root flags
 	StructOnScopePropertyOwner = nullptr;
 
@@ -703,6 +708,57 @@ void FPropertyEditorModule::FindSectionsForCategoryHelper(const UStruct* Struct,
 	}
 }
 
+void FPropertyEditorModule::RegisterMenus()
+{
+	// Owner will be used for cleanup in call to UToolMenus::UnregisterOwner
+	FToolMenuOwnerScoped OwnerScoped(this);
+
+	UToolMenus* ToolMenus = UToolMenus::Get();
+	if (!ToolMenus)
+	{
+		return;		
+	}
+
+	// Property row context menu
+	{
+		static FName MenuName = UE::PropertyEditor::RowContextMenuName;
+		if (!ToolMenus->IsMenuRegistered(MenuName))
+		{
+			UToolMenu* Menu = ToolMenus->RegisterMenu(MenuName, NAME_None, EMultiBoxType::Menu, false);
+			Menu->AddSection(
+				TEXT("Expansion"),
+				NSLOCTEXT("PropertyView", "ExpansionHeading", "Expansion"),
+				FToolMenuInsert(TEXT("Edit"), EToolMenuInsertType::Before));
+		
+			Menu->AddSection(TEXT("Edit"), NSLOCTEXT("PropertyView", "EditHeading", "Edit"));
+
+			Menu->AddDynamicSection(NAME_None, FNewToolMenuDelegate::CreateStatic(&FPropertyEditorModule::PopulateRowContextMenu));
+		}
+	}
+}
+
+void FPropertyEditorModule::PopulateRowContextMenu(UToolMenu* InToolMenu)
+{
+	if (!InToolMenu)
+	{
+		return;
+	}
+				
+	const UDetailRowMenuContextPrivate* MenuContext = InToolMenu->FindContext<UDetailRowMenuContextPrivate>();
+	if (!MenuContext)
+	{
+		return;
+	}
+
+	const TSharedPtr<SDetailTableRowBase> RowContext = MenuContext->GetRowWidget<SDetailTableRowBase>();
+	if (!RowContext.IsValid())
+	{
+		return;
+	}
+
+	RowContext->PopulateContextMenu(InToolMenu);
+}
+
 void FPropertyEditorModule::GetAllSections(const UStruct* Struct, TArray<TSharedPtr<FPropertySection>>& OutSections) const
 {
 	if (Struct == nullptr)
diff --git a/Engine/Source/Editor/PropertyEditor/Private/SDetailCategoryTableRow.cpp b/Engine/Source/Editor/PropertyEditor/Private/SDetailCategoryTableRow.cpp
--- a/Engine/Source/Editor/PropertyEditor/Private/SDetailCategoryTableRow.cpp
+++ b/Engine/Source/Editor/PropertyEditor/Private/SDetailCategoryTableRow.cpp
@@ -14,12 +14,12 @@
 #include "DetailsViewStyle.h"
 #include "SDetailsView.h"
 #include "Serialization/JsonSerializer.h"
-#include "Styling/StyleColors.h"
 #include "UserInterface/Categories/CategoryMenuComboButtonBuilder.h"
 #include "UserInterface/PropertyEditor/PropertyEditorConstants.h"
 #include "Widgets/Input/SComboButton.h"
 #include "Widgets/Layout/SSeparator.h"
 #include "Brushes/SlateColorBrush.h"
+#include "ToolMenus.h"
 
 void SDetailCategoryTableRow::Construct(const FArguments& InArgs, TSharedRef<FDetailTreeNode> InOwnerTreeNode, const TSharedRef<STableViewBase>& InOwnerTableView)
 {
@@ -114,7 +114,7 @@ void SDetailCategoryTableRow::Construct(const FArguments& InArgs, TSharedRef<FDe
 		TAttribute<bool> IsHovered = TAttribute<bool>::CreateSP( this, &SDetailCategoryTableRow::IsHovered);
 		PropertyUpdatedWidgetBuilder->Bind_IsRowHovered(IsHovered);
 	}
-	
+
 	this->ChildSlot
 	[
 		SNew(SBorder)
@@ -226,50 +226,55 @@ FReply SDetailCategoryTableRow::OnMouseButtonUp(const FGeometry& MyGeometry, con
 	return SDetailTableRowBase::OnMouseButtonUp(MyGeometry, MouseEvent);
 }
 
-bool SDetailCategoryTableRow::OnContextMenuOpening(FMenuBuilder& MenuBuilder)
+void SDetailCategoryTableRow::PopulateContextMenu(UToolMenu* ToolMenu)
 {
-	// Don't add anything if neither actions are bound
-	if (!CopyAction.IsBound() || !PasteAction.IsBound())
-	{
-		return true;
-	}
-	
-	MenuBuilder.BeginSection(NAME_None, NSLOCTEXT("PropertyView", "EditHeading", "Edit"));
-	{
-		bool bLongDisplayName = false;
+	SDetailTableRowBase::PopulateContextMenu(ToolMenu);
 
-		FMenuEntryParams CopyContentParams;
-		CopyContentParams.LabelOverride = NSLOCTEXT("PropertyView", "CopyCategoryProperties", "Copy All Properties in Category");
-		CopyContentParams.ToolTipOverride = TAttribute<FText>::CreateLambda([this]()
+	FToolMenuSection& EditSection = ToolMenu->FindOrAddSection(TEXT("Edit"));
+	{
+		// Don't add anything if neither actions are bound
+		if (CopyAction.IsBound() && PasteAction.IsBound())
 		{
-			return CanCopyCategory()
-				? NSLOCTEXT("PropertyView", "CopyCategoryProperties_ToolTip", "Copy all properties in this category")
-				: NSLOCTEXT("PropertyView", "CantCopyCategoryProperties_ToolTip", "None of the properties in this category can be copied");
-		});
-		CopyContentParams.InputBindingOverride = FInputChord(EModifierKey::Shift, EKeys::RightMouseButton).GetInputText(bLongDisplayName);
-		CopyContentParams.IconOverride = FSlateIcon(FCoreStyle::Get().GetStyleSetName(), "GenericCommands.Copy");
-		CopyContentParams.DirectActions = CopyAction;
-		
-		
-		MenuBuilder.AddMenuEntry(CopyContentParams);
+			bool bLongDisplayName = false;
 
-		FMenuEntryParams PasteContentParams;
-		PasteContentParams.LabelOverride = NSLOCTEXT("PropertyView", "PasteCategoryProperties", "Paste All Properties in Category");
-		PasteContentParams.ToolTipOverride = TAttribute<FText>::CreateLambda([this]()
-		{
-			return CanPasteCategory()
-				? NSLOCTEXT("PropertyView", "PasteCategoryProperties_ToolTip", "Paste the copied property values here")
-				// @note: this is specific to the constraint that the destination category has to match the source category (copied from) exactly 
-				: NSLOCTEXT("PropertyView", "CantPasteCategoryProperties_ToolTip", "The properties in this category don't match the contents of the clipboard");
-		});
-		PasteContentParams.InputBindingOverride = FInputChord(EModifierKey::Shift, EKeys::LeftMouseButton).GetInputText(bLongDisplayName);
-		PasteContentParams.IconOverride = FSlateIcon(FCoreStyle::Get().GetStyleSetName(), "GenericCommands.Paste");
-		PasteContentParams.DirectActions = PasteAction;
-		MenuBuilder.AddMenuEntry(PasteContentParams);
+			{
+				// Copy
+				{
+					FToolMenuEntry& CopyMenuEntry = EditSection.AddMenuEntry(
+						TEXT("Copy"),
+						NSLOCTEXT("PropertyView", "CopyCategoryProperties", "Copy All Properties in Category"),
+						TAttribute<FText>::CreateLambda([this]()
+						{
+							return CanCopyCategory()
+								? NSLOCTEXT("PropertyView", "CopyCategoryProperties_ToolTip", "Copy all properties in this category")
+								: NSLOCTEXT("PropertyView", "CantCopyCategoryProperties_ToolTip", "None of the properties in this category can be copied");
+						}),
+						FSlateIcon(FCoreStyle::Get().GetStyleSetName(), "GenericCommands.Copy"),
+						CopyAction);
+
+					CopyMenuEntry.InputBindingLabel = FInputChord(EModifierKey::Shift, EKeys::RightMouseButton).GetInputText(bLongDisplayName);
+				}
+
+				// Paste
+				{
+					FToolMenuEntry& PasteMenuEntry = EditSection.AddMenuEntry(
+						TEXT("Paste"),
+						NSLOCTEXT("PropertyView", "PasteCategoryProperties", "Paste All Properties in Category"),
+						TAttribute<FText>::CreateLambda([this]()
+						{
+							return CanPasteCategory()
+								? NSLOCTEXT("PropertyView", "PasteCategoryProperties_ToolTip", "Paste the copied property values here")
+								// @note: this is specific to the constraint that the destination category has to match the source category (copied from) exactly 
+								: NSLOCTEXT("PropertyView", "CantPasteCategoryProperties_ToolTip", "The properties in this category don't match the contents of the clipboard");
+						}),
+						FSlateIcon(FCoreStyle::Get().GetStyleSetName(), "GenericCommands.Paste"),
+						PasteAction);
+
+					PasteMenuEntry.InputBindingLabel = FInputChord(EModifierKey::Shift, EKeys::LeftMouseButton).GetInputText(bLongDisplayName);
+				}
+			}
+		}
 	}
-	MenuBuilder.EndSection();
-	
-	return true;
 }
 
 EVisibility SDetailCategoryTableRow::IsSeparatorVisible() const
diff --git a/Engine/Source/Editor/PropertyEditor/Private/SDetailCategoryTableRow.h b/Engine/Source/Editor/PropertyEditor/Private/SDetailCategoryTableRow.h
--- a/Engine/Source/Editor/PropertyEditor/Private/SDetailCategoryTableRow.h
+++ b/Engine/Source/Editor/PropertyEditor/Private/SDetailCategoryTableRow.h
@@ -44,8 +44,8 @@ public:
 	// ~Begin End Interface
 
 protected:
-	virtual bool OnContextMenuOpening( FMenuBuilder& MenuBuilder ) override;
-	
+	virtual void PopulateContextMenu(UToolMenu* ToolMenu) override;
+
 private:
 	EVisibility IsSeparatorVisible() const;
 
diff --git a/Engine/Source/Editor/PropertyEditor/Private/SDetailSingleItemRow.cpp b/Engine/Source/Editor/PropertyEditor/Private/SDetailSingleItemRow.cpp
--- a/Engine/Source/Editor/PropertyEditor/Private/SDetailSingleItemRow.cpp
+++ b/Engine/Source/Editor/PropertyEditor/Private/SDetailSingleItemRow.cpp
@@ -2,12 +2,10 @@
 
 #include "SDetailSingleItemRow.h"
 
-#include "Algo/AnyOf.h"
 #include "Algo/Compare.h"
 #include "DetailGroup.h"
 #include "DetailPropertyRow.h"
 #include "DetailWidgetRow.h"
-#include "Editor.h"
 #include "IDetailDragDropHandler.h"
 #include "IDetailPropertyExtensionHandler.h"
 #include "Modules/ModuleInterface.h"
@@ -28,6 +26,7 @@
 #include "Widgets/Input/SComboButton.h"
 #include "DetailsViewStyle.h"
 #include "SDetailsView.h"
+#include "ToolMenus.h"
 
 namespace DetailWidgetConstants
 {
@@ -713,7 +712,7 @@ void SDetailSingleItemRow::Construct( const FArguments& InArgs, FDetailLayoutCus
 	{
 		this->ChildSlot
 		[ 
-			SNullWidget::NullWidget		
+			SNullWidget::NullWidget
 		];
 	}
 	else
@@ -953,8 +952,7 @@ void SDetailSingleItemRow::OnPasteGroup()
 					*FString::Join(PropertiesNotPasted, TEXT("\n")));
 			}
 
-			IDetailsViewPrivate* DetailsView = OwnerTreeNode.Pin()->GetDetailsView();
-			DetailsView->ForceRefresh();
+			ForceRefresh();
 		}
 	}
 }
@@ -1020,162 +1018,205 @@ bool SDetailSingleItemRow::CanPasteGroup()
 	return PreviousClipboardData.bIsApplicable = Algo::Compare(PreviousClipboardData.PropertyNames, PropertyNames);
 }
 
-bool SDetailSingleItemRow::OnContextMenuOpening(FMenuBuilder& MenuBuilder)
+void SDetailSingleItemRow::PopulateContextMenu(UToolMenu* ToolMenu)
 {
-	FUIAction CopyDisplayNameAction = FExecuteAction::CreateSP(this, &SDetailSingleItemRow::OnCopyPropertyDisplayName);
-	CopyDisplayNameAction.CanExecuteAction = FCanExecuteAction::CreateSP(this, &SDetailSingleItemRow::CanCopyPropertyDisplayName);
-	FUIAction CopyInternalNameAction = FExecuteAction::CreateSP(this, &SDetailSingleItemRow::OnCopyPropertyInternalName);
-	CopyInternalNameAction.CanExecuteAction = FCanExecuteAction::CreateSP(this, &SDetailSingleItemRow::CanCopyPropertyInternalName);
+	SDetailTableRowBase::PopulateContextMenu(ToolMenu);
 
-	// Show a separator line if there are items before this in the menu
-	if (MenuBuilder.GetMultiBox()->GetBlocks().Num() > 1)
+	IDetailsViewPrivate* OwningDetailsView = nullptr;
+	if (TSharedPtr<FDetailTreeNode> OwnerTreeNodePtr = OwnerTreeNode.Pin())
 	{
-		MenuBuilder.AddMenuSeparator();
+		OwningDetailsView = OwnerTreeNodePtr->GetDetailsView();
 	}
-
-	if (CopyAction.IsBound() && PasteAction.IsBound())
+	
+	FToolMenuSection& EditSection = ToolMenu->FindOrAddSection(TEXT("Edit"));
 	{
-		const bool bLongDisplayName = false;
-		const bool bIsGroup = Customization->IsValidCustomization() && Customization->DetailGroup.IsValid();
-
-		FMenuEntryParams CopyContentParams;
-		if (bIsGroup)
+		if (CopyAction.IsBound() && PasteAction.IsBound())
 		{
-			CopyContentParams.LabelOverride = NSLOCTEXT("PropertyView", "CopyGroupProperties", "Copy All Properties in Group");
-			CopyContentParams.ToolTipOverride = TAttribute<FText>::CreateLambda([this]()
+			constexpr bool bLongDisplayName = false;
+			const bool bIsGroup = Customization->IsValidCustomization() && Customization->DetailGroup.IsValid();
+
+			// Copy
 			{
-				return CanCopyGroup()
-					? NSLOCTEXT("PropertyView", "CopyGroupProperties_ToolTip", "Copy all properties in this group")
-					: NSLOCTEXT("PropertyView", "CantCopyGroupProperties_ToolTip", "None of the properties in this group can be copied");
-			});
-		}
-		else
-		{
-			CopyContentParams.LabelOverride = NSLOCTEXT("PropertyView", "CopyProperty", "Copy");
-			CopyContentParams.ToolTipOverride = NSLOCTEXT("PropertyView", "CopyProperty_ToolTip", "Copy this property value");
-		}
-		CopyContentParams.InputBindingOverride = FInputChord(EModifierKey::Shift, EKeys::RightMouseButton).GetInputText(bLongDisplayName);
-		CopyContentParams.IconOverride = FSlateIcon(FCoreStyle::Get().GetStyleSetName(), "GenericCommands.Copy");
-		CopyContentParams.DirectActions = CopyAction;
-		MenuBuilder.AddMenuEntry(CopyContentParams);
+				TAttribute<FText> Label;
+				TAttribute<FText> ToolTip;
+				
+				if (bIsGroup)
+				{
+					Label = NSLOCTEXT("PropertyView", "CopyGroupProperties", "Copy All Properties in Group");
+					ToolTip = TAttribute<FText>::CreateLambda([this]()
+					{
+						return CanCopyGroup()
+							? NSLOCTEXT("PropertyView", "CopyGroupProperties_ToolTip", "Copy all properties in this group")
+							: NSLOCTEXT("PropertyView", "CantCopyGroupProperties_ToolTip", "None of the properties in this group can be copied");
+					});
+				}
+				else
+				{
+					Label = NSLOCTEXT("PropertyView", "CopyProperty", "Copy");
+					ToolTip = NSLOCTEXT("PropertyView", "CopyProperty_ToolTip", "Copy this property value");
+				}
 
-		FMenuEntryParams PasteContentParams;
-		if (bIsGroup)
-		{
-			PasteContentParams.LabelOverride = NSLOCTEXT("PropertyView", "PasteGroupProperties", "Paste All Properties in Group");
-			PasteContentParams.ToolTipOverride = TAttribute<FText>::CreateLambda([this]()
+				FToolMenuEntry& CopyMenuEntry = EditSection.AddMenuEntry(
+					TEXT("Copy"),
+					Label,
+					ToolTip,
+					FSlateIcon(FCoreStyle::Get().GetStyleSetName(), "GenericCommands.Copy"),
+					CopyAction);
+
+				CopyMenuEntry.InputBindingLabel = FInputChord(EModifierKey::Shift, EKeys::RightMouseButton).GetInputText(bLongDisplayName);
+			}
+
+			// Paste
 			{
-				return CanPasteGroup()
-					? NSLOCTEXT("PropertyView", "PasteGroupProperties_ToolTip", "Paste the copied property values here")
-					// @note: this is specific to the constraint that the destination group has to match the source group (copied from) exactly 
-					: NSLOCTEXT("PropertyView", "CantPasteGroupProperties_ToolTip", "The properties in this group don't match the contents of the clipboard");
-			});
+				// Paste is only enabled if property editing is enabled
+				if (OwningDetailsView && OwningDetailsView->IsPropertyEditingEnabled())
+				{
+					TAttribute<FText> Label;
+					TAttribute<FText> ToolTip;
+
+					if (bIsGroup)
+					{
+						Label = NSLOCTEXT("PropertyView", "PasteGroupProperties", "Paste All Properties in Group");
+						ToolTip = TAttribute<FText>::CreateLambda([this]()
+						{
+							return CanPasteGroup()
+								? NSLOCTEXT("PropertyView", "PasteGroupProperties_ToolTip", "Paste the copied property values here")
+								// @note: this is specific to the constraint that the destination group has to match the source group (copied from) exactly 
+								: NSLOCTEXT("PropertyView", "CantPasteGroupProperties_ToolTip", "The properties in this group don't match the contents of the clipboard");
+						});
+					}
+					else
+					{
+						Label = NSLOCTEXT("PropertyView", "PasteProperty", "Paste");
+						ToolTip = NSLOCTEXT("PropertyView", "PasteProperty_ToolTip", "Paste the copied value here");	
+					}
+																
+					FToolMenuEntry& PasteMenuEntry = EditSection.AddMenuEntry(
+						TEXT("Paste"),
+						Label,
+						ToolTip,
+						FSlateIcon(FCoreStyle::Get().GetStyleSetName(), "GenericCommands.Paste"),
+						PasteAction);
+
+					PasteMenuEntry.InputBindingLabel = FInputChord(EModifierKey::Shift, EKeys::LeftMouseButton).GetInputText(bLongDisplayName);	
+				}
+			}
 		}
-		else
+
+		// Copy Display Name
 		{
-			PasteContentParams.LabelOverride = NSLOCTEXT("PropertyView", "PasteProperty", "Paste");
-			PasteContentParams.ToolTipOverride = NSLOCTEXT("PropertyView", "PasteProperty_ToolTip", "Paste the copied value here");	
+			FUIAction CopyDisplayNameAction = FExecuteAction::CreateSP(this, &SDetailSingleItemRow::OnCopyPropertyDisplayName);
+			CopyDisplayNameAction.CanExecuteAction = FCanExecuteAction::CreateSP(this, &SDetailSingleItemRow::CanCopyPropertyDisplayName);
+
+			EditSection.AddMenuEntry(
+				TEXT("CopyDisplayName"),
+				NSLOCTEXT("PropertyView", "CopyPropertyDisplayName", "Copy Display Name"),
+				NSLOCTEXT("PropertyView", "CopyPropertyDisplayName_ToolTip", "Copy the display name of this property to the system clipboard."),
+				FSlateIcon(FCoreStyle::Get().GetStyleSetName(), "GenericCommands.Copy"),
+				CopyDisplayNameAction);
 		}
-		PasteContentParams.InputBindingOverride = FInputChord(EModifierKey::Shift, EKeys::LeftMouseButton).GetInputText(bLongDisplayName);
-		PasteContentParams.IconOverride = FSlateIcon(FCoreStyle::Get().GetStyleSetName(), "GenericCommands.Paste");
-		PasteContentParams.DirectActions = PasteAction;
 
-		// Paste is disabled if property editing is disabled
-		if(OwnerTreeNode.Pin()->GetDetailsView()->IsPropertyEditingEnabled())
+		// Copy Internal Name
 		{
-			MenuBuilder.AddMenuEntry(PasteContentParams);
+			FUIAction CopyInternalNameAction = FExecuteAction::CreateSP(this, &SDetailSingleItemRow::OnCopyPropertyInternalName);
+			CopyInternalNameAction.CanExecuteAction = FCanExecuteAction::CreateSP(this, &SDetailSingleItemRow::CanCopyPropertyInternalName);
+
+			EditSection.AddMenuEntry(
+				TEXT("CopyInternalName"),
+				NSLOCTEXT("PropertyView", "CopyPropertyInternalName", "Copy Internal Name"),
+				NSLOCTEXT("PropertyView", "CopyPropertyInternalName_ToolTip", "Copy the internal name of this property to the system clipboard."),
+				FSlateIcon(FCoreStyle::Get().GetStyleSetName(), "GenericCommands.Copy"),
+				CopyInternalNameAction);
 		}
-	}
 
-	MenuBuilder.AddMenuEntry(
-		NSLOCTEXT("PropertyView", "CopyPropertyDisplayName", "Copy Display Name"),
-		NSLOCTEXT("PropertyView", "CopyPropertyDisplayName_ToolTip", "Copy the display name of this property to the system clipboard."),
-		FSlateIcon(FCoreStyle::Get().GetStyleSetName(), "GenericCommands.Copy"),
-		CopyDisplayNameAction);
-
-	MenuBuilder.AddMenuEntry(
-		NSLOCTEXT("PropertyView", "CopyPropertyInternalName", "Copy Internal Name"),
-		NSLOCTEXT("PropertyView", "CopyPropertyInternalName_ToolTip", "Copy the internal name of this property to the system clipboard."),
-		FSlateIcon(FCoreStyle::Get().GetStyleSetName(), "GenericCommands.Copy"),
-		CopyInternalNameAction);
+		// Favorite
+		{
+			if (OwnerTreeNode.Pin()->GetDetailsView()->IsFavoritingEnabled())
+			{
+				FUIAction FavoriteAction;
+				FavoriteAction.ExecuteAction = FExecuteAction::CreateSP(this, &SDetailSingleItemRow::OnFavoriteMenuToggle);
+				FavoriteAction.CanExecuteAction = FCanExecuteAction::CreateSP(this, &SDetailSingleItemRow::CanFavorite);
 
-	if (OwnerTreeNode.Pin()->GetDetailsView()->IsFavoritingEnabled())
-	{
-		FUIAction FavoriteAction;
-		FavoriteAction.ExecuteAction = FExecuteAction::CreateSP(this, &SDetailSingleItemRow::OnFavoriteMenuToggle);
-		FavoriteAction.CanExecuteAction = FCanExecuteAction::CreateSP(this, &SDetailSingleItemRow::CanFavorite);
+				FText FavoriteText = NSLOCTEXT("PropertyView", "FavoriteProperty", "Add to Favorites");
+				FText FavoriteTooltipText = NSLOCTEXT("PropertyView", "FavoriteProperty_ToolTip", "Add this property to your favorites.");
+				FName FavoriteIcon = "DetailsView.PropertyIsFavorite";
 
-		FText FavoriteText = NSLOCTEXT("PropertyView", "FavoriteProperty", "Add to Favorites");
-		FText FavoriteTooltipText = NSLOCTEXT("PropertyView", "FavoriteProperty_ToolTip", "Add this property to your favorites.");
-		FName FavoriteIcon = "DetailsView.PropertyIsFavorite";
+				if (IsFavorite())
+				{
+					FavoriteText = NSLOCTEXT("PropertyView", "RemoveFavoriteProperty", "Remove from Favorites");
+					FavoriteTooltipText = NSLOCTEXT("PropertyView", "RemoveFavoriteProperty_ToolTip", "Remove this property from your favorites.");
+					FavoriteIcon = "DetailsView.PropertyIsNotFavorite";
+				}
 
-		if (IsFavorite())
-		{
-			FavoriteText = NSLOCTEXT("PropertyView", "RemoveFavoriteProperty", "Remove from Favorites");
-			FavoriteTooltipText = NSLOCTEXT("PropertyView", "RemoveFavoriteProperty_ToolTip", "Remove this property from your favorites.");
-			FavoriteIcon = "DetailsView.PropertyIsNotFavorite";
+				EditSection.AddMenuEntry(
+					TEXT("ToggleFavorite"),
+					FavoriteText,
+					FavoriteTooltipText,
+					FSlateIcon(FAppStyle::Get().GetStyleSetName(), FavoriteIcon),
+					FavoriteAction);
+			}
 		}
 
-		MenuBuilder.AddMenuEntry(
-			FavoriteText,
-			FavoriteTooltipText,
-			FSlateIcon(FAppStyle::Get().GetStyleSetName(), FavoriteIcon),
-			FavoriteAction);
-	}
-
-	if (FPropertyEditorPermissionList::Get().ShouldShowMenuEntries())
-	{
-		// Hide separator line if it only contains the SearchWidget, making the next 2 elements the top of the list
-		if (MenuBuilder.GetMultiBox()->GetBlocks().Num() > 1)
+		if (FPropertyEditorPermissionList::Get().ShouldShowMenuEntries())
 		{
-			MenuBuilder.AddMenuSeparator();
-		}
-		
-		MenuBuilder.AddMenuEntry(
-        	NSLOCTEXT("PropertyView", "CopyRowName", "Copy internal row name"),
-        	NSLOCTEXT("PropertyView", "CopyRowName_ToolTip", "Copy the row's parent struct and internal name to use in the property editor's allow/deny lists."),
-        	FSlateIcon(),
-        	FUIAction(FExecuteAction::CreateSP(this, &SDetailSingleItemRow::CopyRowNameText)));
-
-		MenuBuilder.AddMenuEntry(
-			NSLOCTEXT("PropertyView", "AddAllowList", "Add to Allowed"),
-			NSLOCTEXT("PropertyView", "AddAllowList_ToolTip", "Add this row to the property editor's allowed properties list."),
-			FSlateIcon(),
-			FUIAction(FExecuteAction::CreateSP(this, &SDetailSingleItemRow::OnToggleAllowList), FCanExecuteAction(),
-					  FIsActionChecked::CreateSP(this, &SDetailSingleItemRow::IsAllowListChecked)),
-			NAME_None,
-			EUserInterfaceActionType::Check);
-
-		MenuBuilder.AddMenuEntry(
-			NSLOCTEXT("PropertyView", "AddDenyList", "Add to Denied"),
-			NSLOCTEXT("PropertyView", "AddDenyList_ToolTip", "Add this row to the property editor's denied properties list."),
-			FSlateIcon(),
-			FUIAction(FExecuteAction::CreateSP(this, &SDetailSingleItemRow::OnToggleDenyList), FCanExecuteAction(),
-					  FIsActionChecked::CreateSP(this, &SDetailSingleItemRow::IsDenyListChecked)),
-			NAME_None,
-			EUserInterfaceActionType::Check);
-	}
-
-	if (WidgetRow.CustomMenuItems.Num() > 0)
-	{
-		// Hide separator line if it only contains the SearchWidget, making the next 2 elements the top of the list
-		if (MenuBuilder.GetMultiBox()->GetBlocks().Num() > 1)
-		{
-			MenuBuilder.AddMenuSeparator();
+			// Hide separator line if it only contains the SearchWidget, making the next 2 elements the top of the list
+			if (EditSection.Blocks.Num() > 1)
+			{
+				EditSection.AddSeparator(NAME_None);
+			}
+                    
+			EditSection.AddMenuEntry(
+				TEXT("CopyRowName"),
+				NSLOCTEXT("PropertyView", "CopyRowName", "Copy internal row name"),
+				NSLOCTEXT("PropertyView", "CopyRowName_ToolTip", "Copy the row's parent struct and internal name to use in the property editor's allow/deny lists."),
+				FSlateIcon(),
+				FUIAction(FExecuteAction::CreateSP(this, &SDetailSingleItemRow::CopyRowNameText)));
+
+			EditSection.AddMenuEntry(
+				TEXT("AddAllowList"),
+				NSLOCTEXT("PropertyView", "AddAllowList", "Add to Allowed"),
+				NSLOCTEXT("PropertyView", "AddAllowList_ToolTip", "Add this row to the property editor's allowed properties list."),
+				FSlateIcon(),
+				FUIAction(
+					FExecuteAction::CreateSP(this, &SDetailSingleItemRow::OnToggleAllowList),
+					FCanExecuteAction(),
+					FIsActionChecked::CreateSP(this, &SDetailSingleItemRow::IsAllowListChecked)),
+				EUserInterfaceActionType::Check,
+				NAME_None);
+
+			EditSection.AddMenuEntry(
+				TEXT("AddDenyList"),
+				NSLOCTEXT("PropertyView", "AddDenyList", "Add to Denied"),
+				NSLOCTEXT("PropertyView", "AddDenyList_ToolTip", "Add this row to the property editor's denied properties list."),
+				FSlateIcon(),
+				FUIAction(
+					FExecuteAction::CreateSP(this, &SDetailSingleItemRow::OnToggleDenyList),
+					FCanExecuteAction(),
+					FIsActionChecked::CreateSP(this, &SDetailSingleItemRow::IsDenyListChecked)),
+				EUserInterfaceActionType::Check,
+				NAME_None);
 		}
 
-		for (const FDetailWidgetRow::FCustomMenuData& CustomMenuData : WidgetRow.CustomMenuItems)
+		if (WidgetRow.CustomMenuItems.Num() > 0)
 		{
-			// Add the menu entry
-			MenuBuilder.AddMenuEntry(
-				CustomMenuData.Name,
-				CustomMenuData.Tooltip,
-				CustomMenuData.SlateIcon,
-				CustomMenuData.Action);
+			// Hide separator line if it only contains the SearchWidget, making the next 2 elements the top of the list
+			if (EditSection.Blocks.Num() > 1)
+			{
+				EditSection.AddSeparator(NAME_None);
+			}
+
+			for (const FDetailWidgetRow::FCustomMenuData& CustomMenuData : WidgetRow.CustomMenuItems)
+			{
+				// Add the menu entry
+				EditSection.AddMenuEntry(
+					CustomMenuData.GetEntryName(),
+					CustomMenuData.Name,
+					CustomMenuData.Tooltip,
+					CustomMenuData.SlateIcon,
+					CustomMenuData.Action);
+			}
 		}
 	}
-
-	return true;
 }
 
 TArray<TSharedPtr<IPropertyHandle>> SDetailSingleItemRow::GetPropertyHandles(const bool& bRecursive) const
@@ -1339,7 +1380,7 @@ void SDetailSingleItemRow::OnPasteProperty()
 		DetailsView->MarkNodeAnimating(PropertyNode, UE::PropertyEditor::Private::PulseAnimationLength);
 
 		// Need to refresh the details panel in case a property was pasted over another.
-		OwnerTreeNode.Pin()->GetDetailsView()->ForceRefresh();
+		ForceRefresh();
 	}
 }
 
@@ -1650,7 +1691,7 @@ void SDetailSingleItemRow::OnFavoriteMenuToggle()
 	DetailsView->MoveScrollOffset(bNewValue ? ExpandSize : -ExpandSize);
 
 	// Refresh the tree
-	DetailsView->ForceRefresh();
+	ForceRefresh();
 }
 
 void SDetailSingleItemRow::CreateGlobalExtensionWidgets(TArray<FPropertyRowExtensionButton>& OutExtensions) const
diff --git a/Engine/Source/Editor/PropertyEditor/Private/SDetailSingleItemRow.h b/Engine/Source/Editor/PropertyEditor/Private/SDetailSingleItemRow.h
--- a/Engine/Source/Editor/PropertyEditor/Private/SDetailSingleItemRow.h
+++ b/Engine/Source/Editor/PropertyEditor/Private/SDetailSingleItemRow.h
@@ -69,9 +69,9 @@ public:
 	virtual void Tick(const FGeometry& AllottedGeometry, const double InCurrentTime, const float InDeltaTime);
 
 protected:
-	virtual bool OnContextMenuOpening( FMenuBuilder& MenuBuilder ) override;
+	virtual void PopulateContextMenu(UToolMenu* ToolMenu) override;
 	virtual TArray<TSharedPtr<IPropertyHandle>> GetPropertyHandles(const bool& bRecursive = false) const override;
-	
+
 private:
 	void OnCopyProperty();
 	void OnCopyPropertyDisplayName();
diff --git a/Engine/Source/Editor/PropertyEditor/Private/SDetailTableRowBase.cpp b/Engine/Source/Editor/PropertyEditor/Private/SDetailTableRowBase.cpp
--- a/Engine/Source/Editor/PropertyEditor/Private/SDetailTableRowBase.cpp
+++ b/Engine/Source/Editor/PropertyEditor/Private/SDetailTableRowBase.cpp
@@ -2,10 +2,59 @@
 
 #include "SDetailTableRowBase.h"
 
+#include "DetailRowMenuContext.h"
+#include "DetailRowMenuContextPrivate.h"
 #include "PropertyHandleImpl.h"
+#include "ToolMenus.h"
 
 const float SDetailTableRowBase::ScrollBarPadding = 16.0f;
 
+FReply SDetailTableRowBase::OnMouseButtonUp(const FGeometry& MyGeometry, const FPointerEvent& MouseEvent)
+{
+	if( OwnerTreeNode.IsValid() && MouseEvent.GetEffectingButton() == EKeys::RightMouseButton && !StaticCastSharedRef<STableViewBase>( OwnerTablePtr.Pin()->AsWidget() )->IsRightClickScrolling() )
+	{
+		FDetailNodeList VisibleChildren;
+		OwnerTreeNode.Pin()->GetChildren( VisibleChildren );
+
+		// Open context menu if this node can be expanded 
+		bool bShouldOpenMenu = true;
+
+		if( bShouldOpenMenu )
+		{
+			if (UToolMenus* ToolMenus = UToolMenus::Get())
+			{
+				if (UToolMenu* ToolMenu = ToolMenus->FindMenu(UE::PropertyEditor::RowContextMenuName))
+				{
+					const FWidgetPath WidgetPath = MouseEvent.GetEventPath() != nullptr ? *MouseEvent.GetEventPath() : FWidgetPath();
+
+					UDetailRowMenuContext* RowMenuContext = NewObject<UDetailRowMenuContext>();					
+					RowMenuContext->PropertyHandles = GetPropertyHandles(true);
+					RowMenuContext->DetailsView = OwnerTreeNode.Pin()->GetDetailsView();
+					RowMenuContext->ForceRefreshWidget().AddSPLambda(this, [this]
+					{
+						ForceRefresh();						
+					});
+
+					UDetailRowMenuContextPrivate* RowMenuContextPrivate = NewObject<UDetailRowMenuContextPrivate>();
+					RowMenuContextPrivate->Row = SharedThis(this);
+
+					FToolMenuContext MenuContext;
+					MenuContext.AddObject(RowMenuContext);
+					MenuContext.AddObject(RowMenuContextPrivate);
+
+					const TSharedRef<SWidget> ToolMenuWidget = ToolMenus->GenerateWidget(UE::PropertyEditor::RowContextMenuName, MenuContext);
+
+					FSlateApplication::Get().PushMenu(AsShared(), WidgetPath, ToolMenuWidget, MouseEvent.GetScreenSpacePosition(), FPopupTransitionEffect::ContextMenu);
+				}
+			}
+
+			return FReply::Handled();
+		}
+	}
+
+	return STableRow< TSharedPtr< FDetailTreeNode > >::OnMouseButtonUp( MyGeometry, MouseEvent );
+}
+
 int32 SDetailTableRowBase::GetIndentLevelForBackgroundColor() const
 {
 	int32 IndentLevel = 0; 
@@ -37,6 +86,35 @@ bool SDetailTableRowBase::IsScrollBarVisible(TWeakPtr<STableViewBase> OwnerTable
 	return false;
 }
 
+void SDetailTableRowBase::PopulateContextMenu(UToolMenu* ToolMenu)
+{
+	FToolMenuSection& ExpansionSection = ToolMenu->FindOrAddSection(TEXT("Expansion"));
+	{
+		FDetailNodeList VisibleChildren;
+		OwnerTreeNode.Pin()->GetChildren( VisibleChildren );
+
+		// Open context menu if this node can be expanded 
+		if( VisibleChildren.Num() )
+		{
+			const FUIAction CollapseAllAction( FExecuteAction::CreateSP( this, &SDetailTableRowBase::OnCollapseAllClicked ) );
+			ExpansionSection.AddMenuEntry(
+				TEXT("CollapseAll"),
+				NSLOCTEXT("PropertyView", "CollapseAll", "Collapse All"),
+				NSLOCTEXT("PropertyView", "CollapseAll_ToolTip", "Collapses this item and all children"),
+				FSlateIcon(),
+				CollapseAllAction);
+				
+			const FUIAction ExpandAllAction( FExecuteAction::CreateSP( this, &SDetailTableRowBase::OnExpandAllClicked ) );
+			ExpansionSection.AddMenuEntry(
+				TEXT("ExpandAll"),
+				 NSLOCTEXT("PropertyView", "ExpandAll", "Expand All"),
+				NSLOCTEXT("PropertyView", "ExpandAll_ToolTip", "Expands this item and all children"),
+				FSlateIcon(),
+				ExpandAllAction);
+		}
+	}
+}
+
 TArray<TSharedPtr<FPropertyNode>> SDetailTableRowBase::GetPropertyNodes(const bool& bRecursive) const
 {
 	TArray<TSharedPtr<IPropertyHandle>> PropertyHandles = GetPropertyHandles(bRecursive);
@@ -95,3 +173,15 @@ TArray<TSharedPtr<IPropertyHandle>> SDetailTableRowBase::GetPropertyHandles(cons
 
 	return PropertyHandles;
 }
+
+void SDetailTableRowBase::ForceRefresh()
+{
+	if (const TSharedPtr<FDetailTreeNode> Owner = OwnerTreeNode.Pin();
+		Owner.IsValid())
+	{
+		if (IDetailsViewPrivate* DetailsView = Owner->GetDetailsView())
+		{
+			DetailsView->ForceRefresh();
+		}
+	}
+}
diff --git a/Engine/Source/Editor/PropertyEditor/Private/SDetailTableRowBase.h b/Engine/Source/Editor/PropertyEditor/Private/SDetailTableRowBase.h
--- a/Engine/Source/Editor/PropertyEditor/Private/SDetailTableRowBase.h
+++ b/Engine/Source/Editor/PropertyEditor/Private/SDetailTableRowBase.h
@@ -16,54 +16,22 @@
 #include "Widgets/Views/STableViewBase.h"
 #include "Widgets/Views/STableRow.h"
 
+class UToolMenu;
+class UDetailRowMenuContext;
+
 class SDetailTableRowBase : public STableRow< TSharedPtr< FDetailTreeNode > >
 {
 public:
-	virtual FReply OnMouseButtonUp( const FGeometry& MyGeometry, const FPointerEvent& MouseEvent ) override
-	{
-		if( OwnerTreeNode.IsValid() && MouseEvent.GetEffectingButton() == EKeys::RightMouseButton && !StaticCastSharedRef<STableViewBase>( OwnerTablePtr.Pin()->AsWidget() )->IsRightClickScrolling() )
-		{
-			FMenuBuilder MenuBuilder( true, nullptr, nullptr, true );
-
-			FDetailNodeList VisibleChildren;
-			OwnerTreeNode.Pin()->GetChildren( VisibleChildren );
-
-			bool bShouldOpenMenu = false;
-			// Open context menu if this node can be expanded 
-			if( VisibleChildren.Num() )
-			{
-				bShouldOpenMenu = true;
-
-				FUIAction ExpandAllAction( FExecuteAction::CreateSP( this, &SDetailTableRowBase::OnExpandAllClicked ) );
-				FUIAction CollapseAllAction( FExecuteAction::CreateSP( this, &SDetailTableRowBase::OnCollapseAllClicked ) );
-
-				MenuBuilder.BeginSection( NAME_None, NSLOCTEXT("PropertyView", "ExpansionHeading", "Expansion") );
-					MenuBuilder.AddMenuEntry( NSLOCTEXT("PropertyView", "CollapseAll", "Collapse All"), NSLOCTEXT("PropertyView", "CollapseAll_ToolTip", "Collapses this item and all children"), FSlateIcon(), CollapseAllAction );
-					MenuBuilder.AddMenuEntry( NSLOCTEXT("PropertyView", "ExpandAll", "Expand All"), NSLOCTEXT("PropertyView", "ExpandAll_ToolTip", "Expands this item and all children"), FSlateIcon(), ExpandAllAction );
-				MenuBuilder.EndSection();
-
-			}
-
-			bShouldOpenMenu |= OnContextMenuOpening(MenuBuilder);
-
-			if( bShouldOpenMenu )
-			{
-				FWidgetPath WidgetPath = MouseEvent.GetEventPath() != nullptr ? *MouseEvent.GetEventPath() : FWidgetPath();
-
-				FSlateApplication::Get().PushMenu(AsShared(), WidgetPath, MenuBuilder.MakeWidget(), MouseEvent.GetScreenSpacePosition(), FPopupTransitionEffect::ContextMenu);
-
-				return FReply::Handled();
-			}
-		}
-
-		return STableRow< TSharedPtr< FDetailTreeNode > >::OnMouseButtonUp( MyGeometry, MouseEvent );
-	}
+	virtual FReply OnMouseButtonUp( const FGeometry& MyGeometry, const FPointerEvent& MouseEvent ) override;
 
 	int32 GetIndentLevelForBackgroundColor() const;
 
 	static bool IsScrollBarVisible(TWeakPtr<STableViewBase> OwnerTableViewWeak);
 	static const float ScrollBarPadding;
 
+	/** Called to register/extend the row context menu. */
+	virtual void PopulateContextMenu(UToolMenu* ToolMenu);
+
 protected:
 
 	/**
@@ -71,14 +39,6 @@ protected:
 	 * reset to default button
 	 */
 	TSharedPtr<FPropertyUpdatedWidgetBuilder> PropertyUpdatedWidgetBuilder;
-	
-	/**
-	 * Called when the user opens the context menu on this row
-	 *
-	 * @param MenuBuilder	The menu builder to add menu items to
-	 * @return true if menu items were added
-	 */
-	virtual bool OnContextMenuOpening( FMenuBuilder& MenuBuilder ) { return false; }
 
 	/** Retrieve all property nodes represented by this row, and it's children if Recursive specified. */
 	TArray<TSharedPtr<FPropertyNode>> GetPropertyNodes(const bool& bRecursive = false) const;
@@ -86,6 +46,9 @@ protected:
 	/** Retrieve all property handles represented by this row, and it's children if Recursive specified. */
 	virtual TArray<TSharedPtr<IPropertyHandle>> GetPropertyHandles(const bool& bRecursive = false) const;
 
+	/** Refreshes widget and associated list views. */
+	virtual void ForceRefresh();
+
 private:
 	void OnExpandAllClicked()
 	{
diff --git a/Engine/Source/Editor/PropertyEditor/PropertyEditor.Build.cs b/Engine/Source/Editor/PropertyEditor/PropertyEditor.Build.cs
--- a/Engine/Source/Editor/PropertyEditor/PropertyEditor.Build.cs
+++ b/Engine/Source/Editor/PropertyEditor/PropertyEditor.Build.cs
@@ -11,8 +11,8 @@ public class PropertyEditor : ModuleRules
 				"EditorConfig",
 				"EditorFramework",
 				"UnrealEd",
-                	"ActorPickerMode",
-                	"SceneDepthPickerMode",
+                "ActorPickerMode",
+                "SceneDepthPickerMode",
 			}
 		);
 
@@ -20,7 +20,7 @@ public class PropertyEditor : ModuleRules
 			new string[] {
 				"MainFrame",
                	"AssetRegistry",
-                	"AssetTools",
+                "AssetTools",
 				"ClassViewer",
 				"StructViewer",
 				"ContentBrowser",
@@ -29,7 +29,7 @@ public class PropertyEditor : ModuleRules
 
 		PrivateDependencyModuleNames.AddRange(
 			new string[] {
-                	"AppFramework",
+                "AppFramework",
 				"Core",
 				"CoreUObject",
 				"ApplicationCore",
@@ -41,19 +41,20 @@ public class PropertyEditor : ModuleRules
 				"EditorWidgets",
 				"Documentation",
 				"ConfigEditor",
-                	"SceneOutliner",
+                "SceneOutliner",
 				"DesktopPlatform",
 				"PropertyPath",
 				"ToolWidgets",
 				"WidgetRegistration",
-				"Json"
+				"Json",
+				"ToolMenus"
 			}
         );
 
 		DynamicallyLoadedModuleNames.AddRange(
 			new string[] {
-                	"AssetRegistry",
-                	"AssetTools",
+                "AssetRegistry",
+                "AssetTools",
 				"ClassViewer",
 				"StructViewer",
 				"ContentBrowser",
diff --git a/Engine/Source/Editor/PropertyEditor/Public/DetailRowMenuContext.h b/Engine/Source/Editor/PropertyEditor/Public/DetailRowMenuContext.h
new file mode 100644
--- /dev/null
+++ b/Engine/Source/Editor/PropertyEditor/Public/DetailRowMenuContext.h
@@ -0,0 +1,35 @@
+﻿// Copyright Epic Games, Inc. All Rights Reserved.
+
+#pragma once
+
+#include "UObject/Object.h"
+
+#include "DetailRowMenuContext.generated.h"
+
+class IDetailsView;
+class IPropertyHandle;
+
+UCLASS()
+class PROPERTYEDITOR_API UDetailRowMenuContext : public UObject
+{
+	GENERATED_BODY()
+
+public:
+	virtual ~UDetailRowMenuContext() override
+	{
+		DetailsView = nullptr;		
+	}
+	
+	/** Optionally invoke to refresh the widget. */
+	TMulticastDelegate<void()>& ForceRefreshWidget() { return ForceRefreshWidgetDelegate; } 
+	
+	/** PropertyHandles associated with the Row. */
+	TArray<TSharedPtr<IPropertyHandle>> PropertyHandles;
+
+	/** Containing DetailsView. */
+	IDetailsView* DetailsView = nullptr;
+	
+private:
+	/** Optionally invoke to refresh the widget. */
+	TMulticastDelegate<void()> ForceRefreshWidgetDelegate;
+};
diff --git a/Engine/Source/Editor/PropertyEditor/Public/DetailWidgetRow.h b/Engine/Source/Editor/PropertyEditor/Public/DetailWidgetRow.h
--- a/Engine/Source/Editor/PropertyEditor/Public/DetailWidgetRow.h
+++ b/Engine/Source/Editor/PropertyEditor/Public/DetailWidgetRow.h
@@ -391,10 +391,22 @@ public:
 			, SlateIcon(InSlateIcon)
 		{}
 
+		FCustomMenuData(const FName InEntryName, const FUIAction& InAction, const FText& InName, const FText& InTooltip, const FSlateIcon& InSlateIcon)
+			: EntryName(InEntryName)
+			, Action(InAction)
+			, Name(InName)
+			, Tooltip(InTooltip)
+			, SlateIcon(InSlateIcon)
+		{}
+
+		const FName EntryName;
 		const FUIAction Action;
 		const FText Name;
 		const FText Tooltip;
 		const FSlateIcon SlateIcon;
+
+		/** Returns the EntryName if specified, otherwise derive from Name. */
+		const FName GetEntryName() const;
 	};
 	/** Custom Action on this row */
 	TArray<FCustomMenuData> CustomMenuItems;
diff --git a/Engine/Source/Editor/PropertyEditor/Public/PropertyEditorModule.h b/Engine/Source/Editor/PropertyEditor/Public/PropertyEditorModule.h
--- a/Engine/Source/Editor/PropertyEditor/Public/PropertyEditorModule.h
+++ b/Engine/Source/Editor/PropertyEditor/Public/PropertyEditorModule.h
@@ -22,6 +22,12 @@ class ISinglePropertyView;
 class SDetailsView;
 class SPropertyTreeViewImpl;
 class SSingleProperty;
+class UToolMenu;
+
+namespace UE::PropertyEditor
+{
+	static FName RowContextMenuName = TEXT("PropertyEditor.RowContextMenu");
+}
 
 /**
  * The location of a property name relative to its editor widget                   
@@ -222,7 +228,6 @@ struct FRegisterCustomClassLayoutParams
 class FPropertyEditorModule : public IModuleInterface
 {
 public:
-	
 	/**
 	 * Called right after the module has been loaded                   
 	 */
@@ -454,6 +459,11 @@ private:
 	TSharedPtr<class ISinglePropertyView> CreateSinglePropertyImpl(UObject* InObject, const TSharedPtr<IStructureDataProvider>& InStruct, FName InPropertyName, const struct FSinglePropertyParams& InitParams);
 	void CompactSinglePropertyViewArray();
 
+	/** Register Menu extension points */
+	void RegisterMenus();
+
+	static void PopulateRowContextMenu(UToolMenu* InToolMenu);
+
 private:
 	/** All created detail views */
 	TArray< TWeakPtr<class SDetailsView> > AllDetailViews;
