Fixed FStructurePropertyNode::GetReadAddressUncached() on child nodes during init
- changed InternalInitChildNodes() to cached the base struct before adding child nodes, since they can call back to the struct during init
- Fixed the logic in GetReadAddressUncached() how empty instances are handled (before it was only kicking in when OutAddresses was specific)

diff --git a/Engine/Source/Editor/PropertyEditor/Private/StructurePropertyNode.cpp b/Engine/Source/Editor/PropertyEditor/Private/StructurePropertyNode.cpp
--- a/Engine/Source/Editor/PropertyEditor/Private/StructurePropertyNode.cpp
+++ b/Engine/Source/Editor/PropertyEditor/Private/StructurePropertyNode.cpp
@@ -35,6 +35,10 @@ void FStructurePropertyNode::InternalInitChildNodes(FName SinglePropertyName)
 		}
 	}
 
+	// Cache the init time base struct so that we can determine if the struct has changed.
+	// Store the cached base struct before calling AddChildNode() as they may call back to this node.
+	WeakCachedBaseStruct = Struct;
+
 	PropertyEditorHelpers::OrderPropertiesFromMetadata(StructMembers);
 
 	for (FProperty* StructMember : StructMembers)
@@ -54,11 +58,160 @@ void FStructurePropertyNode::InternalInitChildNodes(FName SinglePropertyName)
 		NewItemNode->InitNode(InitParams);
 		AddChildNode(NewItemNode);
 	}
+}
 
-	// Cache the init time base struct so that we can determine of the struct has changed.
-	WeakCachedBaseStruct = Struct;
+bool FStructurePropertyNode::GetReadAddressUncached(const FPropertyNode& InPropertyNode, FReadAddressListData& OutAddresses) const
+{
+	if (!HasValidStructData())
+	{
+		return false;
+	}
+	check(StructProvider.IsValid());
+
+	const FProperty* InItemProperty = InPropertyNode.GetProperty();
+	if (!InItemProperty)
+	{
+		return false;
+	}
+
+	UStruct* OwnerStruct = InItemProperty->GetOwnerStruct();
+	if (!OwnerStruct || OwnerStruct->IsStructTrashed())
+	{
+		// Verify that the property is not part of an invalid trash class
+		return false;
+	}
+
+	TArray<TSharedPtr<FStructOnScope>> Instances;
+	StructProvider->GetInstances(Instances, WeakCachedBaseStruct.Get());
+	bool bHasData = false;
+
+	for (TSharedPtr<FStructOnScope>& Instance : Instances)
+	{
+		uint8* ReadAddress = Instance.IsValid() ? Instance->GetStructMemory() : nullptr;
+		if (ReadAddress)
+		{
+			OutAddresses.Add(nullptr, InPropertyNode.GetValueBaseAddress(ReadAddress, InPropertyNode.HasNodeFlags(EPropertyNodeFlags::IsSparseClassData) != 0, /*bIsStruct=*/true), /*bIsStruct=*/true);
+			bHasData = true;
+		}
+	}
+	return bHasData;
 }
 
+bool FStructurePropertyNode::GetReadAddressUncached(const FPropertyNode& InPropertyNode,
+	bool InRequiresSingleSelection,
+	FReadAddressListData* OutAddresses,
+	bool bComparePropertyContents,
+	bool bObjectForceCompare,
+	bool bArrayPropertiesCanDifferInSize) const
+{
+	if (!HasValidStructData())
+	{
+		return false;
+	}
+	check(StructProvider.IsValid());
+
+	const FProperty* InItemProperty = InPropertyNode.GetProperty();
+	if (!InItemProperty)
+	{
+		return false;
+	}
+
+	const UStruct* OwnerStruct = InItemProperty->GetOwnerStruct();
+	if (!OwnerStruct || OwnerStruct->IsStructTrashed())
+	{
+		// Verify that the property is not part of an invalid trash class
+		return false;
+	}
+
+	bool bAllTheSame = true;
+
+	TArray<TSharedPtr<FStructOnScope>> Instances;
+	StructProvider->GetInstances(Instances, WeakCachedBaseStruct.Get());
+	
+	if (Instances.IsEmpty())
+	{
+		return false;
+	}
+
+	if (bComparePropertyContents || bObjectForceCompare)
+	{
+		const bool bIsSparse = InPropertyNode.HasNodeFlags(EPropertyNodeFlags::IsSparseClassData) != 0;
+		const uint8* BaseAddress = nullptr;
+		const UStruct* BaseStruct = nullptr;
+
+		for (TSharedPtr<FStructOnScope>& Instance : Instances)
+		{
+			if (Instance.IsValid())
+			{
+				if (const UStruct* Struct = Instance->GetStruct())
+				{
+					if (const uint8* ReadAddress = InPropertyNode.GetValueBaseAddress(Instance->GetStructMemory(), bIsSparse, /*bIsStruct=*/true))
+					{
+						if (!BaseAddress)
+						{
+							BaseAddress = ReadAddress;
+							BaseStruct = Struct;
+						}
+						else
+						{
+							if (BaseStruct != Struct)
+							{
+								bAllTheSame = false;
+								break;
+							}
+							if (!InItemProperty->Identical(BaseAddress, ReadAddress))
+							{
+								bAllTheSame = false;
+								break;
+							}
+						}
+					}
+				}
+			}
+		}
+
+		// If none of the instances have data, treat it as if the instance data was empty.
+		if (!BaseStruct)
+		{
+			bAllTheSame = false;
+		}
+	}
+	else
+	{
+		// Check that all are valid or invalid.
+		const UStruct* BaseStruct = Instances[0].IsValid() ? Instances[0]->GetStruct() : nullptr;
+		for (int32 Index = 1; Index < Instances.Num(); Index++)
+		{
+			const UStruct* Struct = Instances[Index].IsValid() ? Instances[Index]->GetStruct() : nullptr;
+			if (BaseStruct != Struct)
+			{
+				bAllTheSame = false;
+				break;
+			}
+		}
+		
+		// If none of the instances have data, treat it as if the instance data was empty.
+		if (!BaseStruct)
+		{
+			bAllTheSame = false;
+		}
+	}
+
+	if (bAllTheSame && OutAddresses)
+	{
+		for (TSharedPtr<FStructOnScope>& Instance : Instances)
+		{
+			uint8* ReadAddress = Instance.IsValid() ? Instance->GetStructMemory() : nullptr;
+			if (ReadAddress)
+			{
+				OutAddresses->Add(nullptr, InPropertyNode.GetValueBaseAddress(ReadAddress, InPropertyNode.HasNodeFlags(EPropertyNodeFlags::IsSparseClassData) != 0, /*bIsStruct=*/true), /*bIsStruct=*/true);
+			}
+		}
+	}
+
+	return bAllTheSame;
+}
+	
 uint8* FStructurePropertyNode::GetValueBaseAddress(uint8* StartAddress, bool bIsSparseData, bool bIsStruct) const
 {
 	// If called with struct data, we expect that it is compatible with the first structure node down in the property node chain, return the address as is.
diff --git a/Engine/Source/Editor/PropertyEditor/Private/StructurePropertyNode.h b/Engine/Source/Editor/PropertyEditor/Private/StructurePropertyNode.h
--- a/Engine/Source/Editor/PropertyEditor/Private/StructurePropertyNode.h
+++ b/Engine/Source/Editor/PropertyEditor/Private/StructurePropertyNode.h
@@ -77,147 +77,14 @@ public:
 		return StructProvider;
 	}
 
-	bool GetReadAddressUncached(const FPropertyNode& InPropertyNode, FReadAddressListData& OutAddresses) const override
-	{
-		if (!HasValidStructData())
-		{
-			return false;
-		}
-		check(StructProvider.IsValid());
-
-		const FProperty* InItemProperty = InPropertyNode.GetProperty();
-		if (!InItemProperty)
-		{
-			return false;
-		}
-
-		UStruct* OwnerStruct = InItemProperty->GetOwnerStruct();
-		if (!OwnerStruct || OwnerStruct->IsStructTrashed())
-		{
-			// Verify that the property is not part of an invalid trash class
-			return false;
-		}
-
-		TArray<TSharedPtr<FStructOnScope>> Instances;
-		StructProvider->GetInstances(Instances, WeakCachedBaseStruct.Get());
-		bool bHasData = false;
-
-		for (TSharedPtr<FStructOnScope>& Instance : Instances)
-		{
-			uint8* ReadAddress = Instance.IsValid() ? Instance->GetStructMemory() : nullptr;
-			if (ReadAddress)
-			{
-				OutAddresses.Add(nullptr, InPropertyNode.GetValueBaseAddress(ReadAddress, InPropertyNode.HasNodeFlags(EPropertyNodeFlags::IsSparseClassData) != 0, /*bIsStruct=*/true), /*bIsStruct=*/true);
-				bHasData = true;
-			}
-		}
-		return bHasData;
-	}
+	virtual bool GetReadAddressUncached(const FPropertyNode& InPropertyNode, FReadAddressListData& OutAddresses) const override;
 
-	bool GetReadAddressUncached(const FPropertyNode& InPropertyNode,
+	virtual bool GetReadAddressUncached(const FPropertyNode& InPropertyNode,
 		bool InRequiresSingleSelection,
 		FReadAddressListData* OutAddresses,
 		bool bComparePropertyContents,
 		bool bObjectForceCompare,
-		bool bArrayPropertiesCanDifferInSize) const override
-	{
-		if (!HasValidStructData())
-		{
-			return false;
-		}
-		check(StructProvider.IsValid());
-
-		const FProperty* InItemProperty = InPropertyNode.GetProperty();
-		if (!InItemProperty)
-		{
-			return false;
-		}
-
-		const UStruct* OwnerStruct = InItemProperty->GetOwnerStruct();
-		if (!OwnerStruct || OwnerStruct->IsStructTrashed())
-		{
-			// Verify that the property is not part of an invalid trash class
-			return false;
-		}
-
-		bool bAllTheSame = true;
-
-		TArray<TSharedPtr<FStructOnScope>> Instances;
-		StructProvider->GetInstances(Instances, WeakCachedBaseStruct.Get());
-		
-		if (Instances.IsEmpty())
-		{
-			return false;
-		}
-
-		if (bComparePropertyContents || bObjectForceCompare)
-		{
-			const bool bIsSparse = InPropertyNode.HasNodeFlags(EPropertyNodeFlags::IsSparseClassData) != 0;
-			const uint8* BaseAddress = nullptr;
-			const UStruct* BaseStruct = nullptr;
-
-			for (TSharedPtr<FStructOnScope>& Instance : Instances)
-			{
-				if (Instance.IsValid())
-				{
-					if (const UStruct* Struct = Instance->GetStruct())
-					{
-						if (const uint8* ReadAddress = InPropertyNode.GetValueBaseAddress(Instance->GetStructMemory(), bIsSparse, /*bIsStruct=*/true))
-						{
-							if (!BaseAddress)
-							{
-								BaseAddress = ReadAddress;
-								BaseStruct = Struct;
-							}
-							else
-							{
-								if (BaseStruct != Struct)
-								{
-									bAllTheSame = false;
-									break;
-								}
-								if (!InItemProperty->Identical(BaseAddress, ReadAddress))
-								{
-									bAllTheSame = false;
-									break;
-								}
-							}
-						}
-					}
-				}
-			}
-		}
-		else
-		{
-			// Check that all are valid or invalid.
-			const UStruct* BaseStruct = Instances[0].IsValid() ? Instances[0]->GetStruct() : nullptr;
-			for (int32 Index = 1; Index < Instances.Num(); Index++)
-			{
-				const UStruct* Struct = Instances[Index].IsValid() ? Instances[Index]->GetStruct() : nullptr;
-				if (BaseStruct != Struct)
-				{
-					bAllTheSame = false;
-					break;
-				}
-			}
-		}
-
-		bool bHasData = false;
-		if (bAllTheSame && OutAddresses)
-		{
-			for (TSharedPtr<FStructOnScope>& Instance : Instances)
-			{
-				uint8* ReadAddress = Instance.IsValid() ? Instance->GetStructMemory() : nullptr;
-				if (ReadAddress)
-				{
-					OutAddresses->Add(nullptr, InPropertyNode.GetValueBaseAddress(ReadAddress, InPropertyNode.HasNodeFlags(EPropertyNodeFlags::IsSparseClassData) != 0, /*bIsStruct=*/true), /*bIsStruct=*/true);
-					bHasData = true;
-				}
-			}
-		}
-
-		return bAllTheSame && bHasData;
-	}
+		bool bArrayPropertiesCanDifferInSize) const override;
 
 	void GetOwnerPackages(TArray<UPackage*>& OutPackages) const
 	{
