Fix Paste with GetAllowedClasses and GetDisallowedClasses
#rnx

diff --git a/Engine/Source/Editor/PropertyEditor/Private/PropertyHandleImpl.cpp b/Engine/Source/Editor/PropertyEditor/Private/PropertyHandleImpl.cpp
--- a/Engine/Source/Editor/PropertyEditor/Private/PropertyHandleImpl.cpp
+++ b/Engine/Source/Editor/PropertyEditor/Private/PropertyHandleImpl.cpp
@@ -4122,7 +4122,7 @@ FPropertyAccess::Result FPropertyHandleObject::SetValueFromFormattedString(const
 			{
 				for (const FString& ClassName : AllowedClassNames)
 				{
-					UClass* AllowedClass = nullptr;
+					const UClass* AllowedClass = nullptr;
 					if (!FPackageName::IsShortPackageName(ClassName))
 					{
 						AllowedClass = FindObject<UClass>(nullptr, *ClassName);
@@ -4147,6 +4147,44 @@ FPropertyAccess::Result FPropertyHandleObject::SetValueFromFormattedString(const
 				bSupportedObject = true;
 			}
 
+			if (bSupportedObject)
+			{
+				const FString& GetAllowedClassesFunctionName = NodeProperty->GetMetaData("GetAllowedClasses");
+				if (!GetAllowedClassesFunctionName.IsEmpty() && NodeProperty->GetOwnerUObject())
+				{
+					TArray<UObject*> OuterObjects;
+					GetOuterObjects(OuterObjects);
+					for (UObject* Object : OuterObjects)
+					{
+						const UFunction* GetAllowedClassesFunction = Object->FindFunction(*GetAllowedClassesFunctionName);
+						if (GetAllowedClassesFunction)
+						{
+							DECLARE_DELEGATE_RetVal(TArray<UClass*>, FGetAllowedClasses);
+							TArray<UClass*> AllowedClasses = FGetAllowedClasses::CreateUFunction(Object, GetAllowedClassesFunction->GetFName()).Execute();
+							if (AllowedClasses.Num() > 0)
+							{
+								bSupportedObject = false;
+								for (const UClass* AllowedClass : AllowedClasses)
+								{
+									const bool bIsInterface = AllowedClass && AllowedClass->HasAnyClassFlags(CLASS_Interface);
+				
+									// Check if the object is an allowed class type this property supports
+									if ((AllowedClass && QualifiedClass->IsChildOf(AllowedClass)) || (bIsInterface && QualifiedObject->GetClass()->ImplementsInterface(AllowedClass)))
+									{
+										bSupportedObject = true;
+										break;
+									}
+								}
+								if (!bSupportedObject)
+								{
+									break;
+								}
+							}
+						}
+					}
+				}
+			}
+
 			if (bSupportedObject)
 			{
 				const FString& DisallowedClassesString = NodeProperty->GetMetaData("DisallowedClasses");
@@ -4155,7 +4193,7 @@ FPropertyAccess::Result FPropertyHandleObject::SetValueFromFormattedString(const
 
 				for (const FString& DisallowedClassName : DisallowedClassNames)
 				{
-					UClass* DisallowedClass = UClass::TryFindTypeSlow<UClass>(DisallowedClassName);
+					const UClass* DisallowedClass = UClass::TryFindTypeSlow<UClass>(DisallowedClassName);
 					const bool bIsInterface = DisallowedClass && DisallowedClass->HasAnyClassFlags(CLASS_Interface);
 
 					if ((DisallowedClass && QualifiedClass->IsChildOf(DisallowedClass)) || (bIsInterface && QualifiedObject->GetClass()->ImplementsInterface(DisallowedClass)))
@@ -4166,6 +4204,43 @@ FPropertyAccess::Result FPropertyHandleObject::SetValueFromFormattedString(const
 				}
 			}
 
+			if (bSupportedObject)
+			{
+				const FString& GetDisallowedClassesFunctionName = NodeProperty->GetMetaData("GetDisallowedClasses");
+				if (!GetDisallowedClassesFunctionName.IsEmpty() && NodeProperty->GetOwnerUObject())
+				{
+					TArray<UObject*> OuterObjects;
+					GetOuterObjects(OuterObjects);
+					for (UObject* Object : OuterObjects)
+					{
+						const UFunction* GetDisallowedClassesFunction = Object->FindFunction(*GetDisallowedClassesFunctionName);
+						if (GetDisallowedClassesFunction)
+						{
+							DECLARE_DELEGATE_RetVal(TArray<UClass*>, FGetAllowedClasses);
+							TArray<UClass*> DisallowedClasses = FGetAllowedClasses::CreateUFunction(Object, GetDisallowedClassesFunction->GetFName()).Execute();
+							if (DisallowedClasses.Num() > 0)
+							{
+								for (const UClass* DisallowedClass : DisallowedClasses)
+								{
+									const bool bIsInterface = DisallowedClass && DisallowedClass->HasAnyClassFlags(CLASS_Interface);
+					
+									// Check if the object is an allowed class type this property supports
+									if ((DisallowedClass && QualifiedClass->IsChildOf(DisallowedClass)) || (bIsInterface && QualifiedObject->GetClass()->ImplementsInterface(DisallowedClass)))
+									{
+										bSupportedObject = false;
+										break;
+									}
+								}
+								if (!bSupportedObject)
+								{
+									break;
+								}
+							}
+						}
+					}
+				}
+			}
+
 			// Check required class
 			if (RequiredClass && !QualifiedClass->IsChildOf(RequiredClass))
 			{
@@ -4681,6 +4756,8 @@ FPropertyAccess::Result FPropertyHandleVector::SetW(double InValue, EPropertyVal
 	if( VectorComponents.Num() == 4 )
 	{
 		FPropertyAccess::Result Res = VectorComponents[3]->SetValue( InValue, Flags );
+		
+		return Res;
 	}
 
 	return FPropertyAccess::Fail;
@@ -4985,6 +5062,48 @@ FPropertyAccess::Result FPropertyHandleArray::MoveElementTo(int32 OriginalIndex,
 	return Result;
 }
 
+FPropertyAccess::Result FPropertyHandleArray::SetValueFromFormattedString(const FString& InValue, EPropertyValueSetFlags::Type Flags)
+{
+	FPropertyAccess::Result Result = FPropertyAccess::Success;
+
+	if (InValue.StartsWith("(") && InValue.EndsWith(")"))
+	{
+		if (EmptyArray() != FPropertyAccess::Success)
+		{
+			return FPropertyAccess::Fail;
+		}
+
+		TArray<FString> Values;
+
+		// Remove first and last parenthesis  
+		InValue.LeftChop(1).RightChop(1).ParseIntoArrayWS(Values, TEXT(","), true);
+
+		for (const FString& Value : Values)
+		{
+			if (AddItem() != FPropertyAccess::Success)
+			{
+				return FPropertyAccess::Fail;
+			}
+			uint32 NumElements = 0;
+			if (GetNumElements(NumElements) != FPropertyAccess::Success || NumElements == 0)
+			{
+				return FPropertyAccess::Fail;
+			}
+			
+			const TSharedRef<IPropertyHandle> Property = GetElement( NumElements - 1 );
+			if (Property->IsValidHandle() && Property->SetValueFromFormattedString(Value, Flags) == FPropertyAccess::Fail)
+			{
+				Result = FPropertyAccess::Fail;
+			}
+		}
+		return Result;
+	}
+	else
+	{
+		return FPropertyHandleBase::SetValueFromFormattedString(InValue, Flags);
+	}
+}
+
 bool FPropertyHandleArray::IsEditable() const
 {
 	// Property is editable if its a non-const dynamic array
@@ -4999,7 +5118,7 @@ bool FPropertyHandleOptional::Supports(TSharedRef<FPropertyNode> PropertyNode)
 }
 
 /** IPropertyHandleOptional interface */
-FPropertyAccess::Result FPropertyHandleOptional::GetOptionalValue(FProperty* OutValue)
+FPropertyAccess::Result FPropertyHandleOptional::GetOptionalValue(FProperty*& OutValue)
 {
 	const TSharedPtr<FPropertyNode>& PropertyNode = Implementation->GetPropertyNode();
 	FOptionalProperty* OptionalProperty = CastField<FOptionalProperty>(PropertyNode->GetProperty());
@@ -5177,6 +5296,7 @@ FPropertyAccess::Result FPropertyHandleSet::AddItem()
 		if (!HasDefaultElement())
 		{
 			Implementation->AddChild();
+			Implementation->GetPropertyNode()->RebuildChildren();
 			Result = FPropertyAccess::Success;
 		}
 		else
@@ -5247,6 +5367,43 @@ bool FPropertyHandleSet::IsEditable() const
 	return Implementation->HasValidPropertyNode() && !Implementation->IsEditConst() && Implementation->IsPropertyTypeOf(FSetProperty::StaticClass());
 }
 
+FPropertyAccess::Result FPropertyHandleSet::SetValueFromFormattedString(const FString& InValue, EPropertyValueSetFlags::Type Flags)
+{
+	FPropertyAccess::Result Result = FPropertyAccess::Success;
+
+	if (InValue.StartsWith("(") && InValue.EndsWith(")"))
+	{
+		if (Empty() != FPropertyAccess::Success)
+		{
+			return FPropertyAccess::Fail;
+		}
+		
+		TArray<FString> Values;
+
+		// Remove first and last parenthesis  
+		InValue.LeftChop(1).RightChop(1).ParseIntoArrayWS(Values, TEXT(","), true);
+		for (const FString& Value : Values)
+		{
+			AddItem();
+			uint32 NumElements = 0;
+			if (GetNumElements(NumElements) != FPropertyAccess::Success || NumElements == 0)
+			{
+				return FPropertyAccess::Fail;
+			}
+			const TSharedRef<IPropertyHandle> Property = GetElement( NumElements - 1 );
+			if ( Property->IsValidHandle() && Property->SetValueFromFormattedString(Value, Flags) == FPropertyAccess::Fail)
+			{
+				Result = FPropertyAccess::Fail;
+			}
+		}
+		return Result;
+	}
+	else
+	{
+		return FPropertyHandleBase::SetValueFromFormattedString(InValue, Flags);
+	}
+}
+
 // Maps
 
 bool FPropertyHandleMap::Supports(TSharedRef<FPropertyNode> PropertyNode)
diff --git a/Engine/Source/Editor/PropertyEditor/Private/PropertyHandleImpl.h b/Engine/Source/Editor/PropertyEditor/Private/PropertyHandleImpl.h
--- a/Engine/Source/Editor/PropertyEditor/Private/PropertyHandleImpl.h
+++ b/Engine/Source/Editor/PropertyEditor/Private/PropertyHandleImpl.h
@@ -711,6 +711,7 @@ public:
 	virtual TSharedPtr<IPropertyHandleArray> AsArray() override;
 	virtual TSharedRef<IPropertyHandle> GetElement( int32 Index ) const override;
 	virtual FPropertyAccess::Result MoveElementTo(int32 OriginalIndex, int32 NewIndex) override;
+	virtual FPropertyAccess::Result SetValueFromFormattedString(const FString& InValue, EPropertyValueSetFlags::Type Flags = EPropertyValueSetFlags::DefaultFlags) override;
 	virtual bool IsEditable() const override;
 };
 
@@ -720,7 +721,7 @@ public:
 	FPropertyHandleOptional(TSharedRef<FPropertyNode> PropertyNode, FNotifyHook* NotifyHook, TSharedPtr<IPropertyUtilities> PropertyUtilities);
 	static bool Supports(TSharedRef<FPropertyNode> PropertyNode);
 	/** IPropertyHandleOptional interface */
-	virtual FPropertyAccess::Result GetOptionalValue(FProperty* OutValue) override;
+	virtual FPropertyAccess::Result GetOptionalValue(FProperty*& OutValue) override;
 	virtual FPropertyAccess::Result SetOptionalValue(FProperty* NewValue) override;
 	virtual FPropertyAccess::Result ClearOptionalValue() override;
 	virtual TSharedPtr<IPropertyHandleOptional> AsOptional() override;
@@ -753,6 +754,7 @@ public:
 	virtual FPropertyAccess::Result DeleteItem(int32 Index) override;
 	virtual FPropertyAccess::Result GetNumElements(uint32& OutNumElements) override;
 	virtual TSharedRef<IPropertyHandle> GetElement(int32 Index) const override;
+	virtual FPropertyAccess::Result SetValueFromFormattedString(const FString& InValue, EPropertyValueSetFlags::Type Flags = EPropertyValueSetFlags::DefaultFlags) override;
 
 	virtual FDelegateHandle SetOnNumElementsChanged( const FSimpleDelegate& InOnNumElementsChanged ) override;
 	virtual void UnregisterOnNumElementsChanged(FDelegateHandle Handle) override;
diff --git a/Engine/Source/Editor/PropertyEditor/Public/PropertyHandle.h b/Engine/Source/Editor/PropertyEditor/Public/PropertyHandle.h
--- a/Engine/Source/Editor/PropertyEditor/Public/PropertyHandle.h
+++ b/Engine/Source/Editor/PropertyEditor/Public/PropertyHandle.h
@@ -938,7 +938,7 @@ public:
 	 * Get the item
 	 * @param OutValue	The value of the optional if it is set.
 	 */
-	virtual FPropertyAccess::Result GetOptionalValue(FProperty* OutValue) = 0;
+	virtual FPropertyAccess::Result GetOptionalValue(FProperty*& OutValue) = 0;
 
 	/**
 	 * Set the item
