Fix external objects collapsing on edit
#tests Tested in PIE

diff --git a/Engine/Source/Editor/PropertyEditor/Private/DetailLayoutBuilderImpl.cpp b/Engine/Source/Editor/PropertyEditor/Private/DetailLayoutBuilderImpl.cpp
--- a/Engine/Source/Editor/PropertyEditor/Private/DetailLayoutBuilderImpl.cpp
+++ b/Engine/Source/Editor/PropertyEditor/Private/DetailLayoutBuilderImpl.cpp
@@ -26,13 +26,7 @@ FDetailLayoutBuilderImpl::FDetailLayoutBuilderImpl(TSharedPtr<FComplexPropertyNo
 
 FDetailLayoutBuilderImpl::~FDetailLayoutBuilderImpl()
 {
-	if (GetDetailsView())
-	{
-		for (TSharedPtr<FComplexPropertyNode> ExternalRootPropertyNode : ExternalRootPropertyNodes)
-		{
-			GetDetailsView()->SaveExpandedItems(ExternalRootPropertyNode.ToSharedRef());
-		}
-	}
+	ClearExternalRootPropertyNodes();
 }
 
 IDetailCategoryBuilder& FDetailLayoutBuilderImpl::EditCategory(FName CategoryName, const FText& NewLocalizedDisplayName, ECategoryPriority::Type CategoryType)
@@ -140,6 +134,7 @@ TSharedPtr<IPropertyHandle> FDetailLayoutBuilderImpl::AddObjectPropertyData(TCon
 
 		if (TSharedPtr<FPropertyNode> PropertyNode = RootPropertyNode->GenerateSingleChild(PropertyName))
 		{
+			// This is useless as PropertyNode should already be in the child nodes
 			RootPropertyNode->AddChildNode(PropertyNode);
 			PropertyNode->RebuildChildren();
 			Handle = GetPropertyHandle(PropertyNode);
@@ -801,6 +796,19 @@ void FDetailLayoutBuilderImpl::RemoveExternalRootPropertyNode(TSharedRef<FComple
 	}
 }
 
+void FDetailLayoutBuilderImpl::ClearExternalRootPropertyNodes()
+{
+	if (GetDetailsView())
+	{
+		for (const TSharedPtr<FComplexPropertyNode>& ExternalRootPropertyNode : ExternalRootPropertyNodes)
+		{
+			GetDetailsView()->SaveExpandedItems(ExternalRootPropertyNode.ToSharedRef());
+		}
+	}
+
+	ExternalRootPropertyNodes.Empty();
+}
+
 FDelegateHandle FDetailLayoutBuilderImpl::AddNodeVisibilityChangedHandler(FSimpleMulticastDelegate::FDelegate InOnNodeVisibilityChanged)
 {
 	return OnNodeVisibilityChanged.Add(InOnNodeVisibilityChanged);
diff --git a/Engine/Source/Editor/PropertyEditor/Private/DetailLayoutBuilderImpl.h b/Engine/Source/Editor/PropertyEditor/Private/DetailLayoutBuilderImpl.h
--- a/Engine/Source/Editor/PropertyEditor/Private/DetailLayoutBuilderImpl.h
+++ b/Engine/Source/Editor/PropertyEditor/Private/DetailLayoutBuilderImpl.h
@@ -205,6 +205,8 @@ public:
 	*/
 	void RemoveExternalRootPropertyNode(TSharedRef<FComplexPropertyNode> InExternalRootNode);
 
+	void ClearExternalRootPropertyNodes();
+
 	/** @return The details view that owns this layout */
 	IDetailsViewPrivate* GetDetailsView() { return DetailsView; }
 	/** @return The root node for this customization */
diff --git a/Engine/Source/Editor/PropertyEditor/Private/DetailPropertyRow.cpp b/Engine/Source/Editor/PropertyEditor/Private/DetailPropertyRow.cpp
--- a/Engine/Source/Editor/PropertyEditor/Private/DetailPropertyRow.cpp
+++ b/Engine/Source/Editor/PropertyEditor/Private/DetailPropertyRow.cpp
@@ -617,15 +617,24 @@ void FDetailPropertyRow::MakeExternalPropertyRowCustomization(const TArray<UObje
 
 	RootPropertyNode->InitNode(InitParams);
 
-	ParentCategory->GetParentLayoutImpl()->AddExternalRootPropertyNode(RootPropertyNode);
 
 	if (PropertyName != NAME_None)
 	{
 		TSharedPtr<FPropertyNode> PropertyNode = RootPropertyNode->GenerateSingleChild(PropertyName);
 		if(PropertyNode.IsValid())
 		{
+			// This is useless as PropertyNode should already be in the child nodes
 			RootPropertyNode->AddChildNode(PropertyNode);
 
+			if (InitParams.bCreateCategoryNodes)
+			{
+				PropertyNode->SetNodeFlags(EPropertyNodeFlags::ShowCategories, true);
+			}
+			else
+			{
+				PropertyNode->SetNodeFlags(EPropertyNodeFlags::ShowCategories, false);
+			}
+			
 			PropertyNode->RebuildChildren();
 
 			OutCustomization.PropertyRow = MakeShared<FDetailPropertyRow>(PropertyNode, ParentCategory, RootPropertyNode);
@@ -637,6 +646,8 @@ void FDetailPropertyRow::MakeExternalPropertyRowCustomization(const TArray<UObje
 		OutCustomization.PropertyRow = MakeShared<FDetailPropertyRow>(RootPropertyNode, ParentCategory, RootPropertyNode);
 		OutCustomization.PropertyRow->SetCustomExpansionId(Parameters.GetUniqueId());
 	}
+	
+	ParentCategory->GetParentLayoutImpl()->AddExternalRootPropertyNode(RootPropertyNode);
 }
 
 EVisibility FDetailPropertyRow::GetPropertyVisibility() const
diff --git a/Engine/Source/Editor/PropertyEditor/Private/PropertyNode.cpp b/Engine/Source/Editor/PropertyEditor/Private/PropertyNode.cpp
--- a/Engine/Source/Editor/PropertyEditor/Private/PropertyNode.cpp
+++ b/Engine/Source/Editor/PropertyEditor/Private/PropertyNode.cpp
@@ -1,30 +1,24 @@
 // Copyright Epic Games, Inc. All Rights Reserved.
 
 #include "PropertyNode.h"
-#include "Misc/ConfigCacheIni.h"
-#include "Serialization/ArchiveReplaceObjectRef.h"
 #include "Components/ActorComponent.h"
 #include "Containers/Deque.h"
+#include "EditConditionContext.h"
+#include "Editor.h"
 #include "Editor/UnrealEdEngine.h"
+#include "EditorSupportDelegates.h"
 #include "Engine/UserDefinedStruct.h"
-#include "EditConditionContext.h"
-#include "UnrealEdGlobals.h"
-#include "ScopedTransaction.h"
-#include "PropertyRestriction.h"
-#include "Kismet2/StructureEditorUtils.h"
+#include "InstancedReferenceSubobjectHelper.h"
 #include "Kismet2/BlueprintEditorUtils.h"
+#include "Misc/ConfigCacheIni.h"
 #include "Misc/ScopeExit.h"
-#include "Editor.h"
 #include "ObjectPropertyNode.h"
-#include "StructurePropertyNode.h"
 #include "PropertyHandleImpl.h"
+#include "PropertyRestriction.h"
 #include "PropertyTextUtilities.h"
-#include "EditorSupportDelegates.h"
-#include "UObject/ConstructorHelpers.h"
-#include "InstancedReferenceSubobjectHelper.h"
+#include "StringPrefixTree.h"
+#include "StructurePropertyNode.h"
 
-#include "Framework/Notifications/NotificationManager.h"
-#include "Widgets/Notifications/SNotificationList.h"
 #include "UObject/MetaData.h"
 #include "UObject/TextProperty.h"
 #include "UObject/EnumProperty.h"
@@ -267,6 +261,48 @@ void FPropertyNode::InitNode(const FPropertyNodeInitParams& InitParams)
 	PropertyPath = FPropertyNode::CreatePropertyPath(this->AsShared())->ToString();
 }
 
+namespace FPropertyNodeUtils
+{
+	void GetExpandedItems(const TSharedPtr<FPropertyNode>& InPropertyNode, FStringPrefixTree& OutExpandedItems)
+	{
+		if (InPropertyNode->HasNodeFlags(EPropertyNodeFlags::Expanded))
+		{
+			constexpr bool bWithArrayIndex = true;
+			FString Path;
+			Path.Empty(128);
+			InPropertyNode->GetQualifiedName(Path, bWithArrayIndex);
+
+			OutExpandedItems.Insert(Path);
+		}
+
+		for (int32 ChildIndex = 0; ChildIndex < InPropertyNode->GetNumChildNodes(); ++ChildIndex)
+		{
+			GetExpandedItems(InPropertyNode->GetChildNode(ChildIndex), OutExpandedItems);
+		}
+	}
+
+	void SetExpandedItems(const TSharedPtr<FPropertyNode>& InPropertyNode, const FStringPrefixTree& InExpandedItems)
+	{
+		constexpr bool bWithArrayIndex = true;
+		FString Path;
+		Path.Empty(128);
+		InPropertyNode->GetQualifiedName(Path, bWithArrayIndex);
+
+		if (InExpandedItems.Contains(Path))
+		{
+			InPropertyNode->SetNodeFlags(EPropertyNodeFlags::Expanded, true);
+		}
+
+		if (InExpandedItems.AnyStartsWith(Path))
+		{
+			for (int32 NodeIndex = 0; NodeIndex < InPropertyNode->GetNumChildNodes(); ++NodeIndex)
+			{
+				SetExpandedItems(InPropertyNode->GetChildNode(NodeIndex), InExpandedItems);
+			}
+		}
+	}
+}
+
 /**
  * Used for rebuilding a sub portion of the tree
  */
@@ -278,18 +314,12 @@ void FPropertyNode::RebuildChildren()
 	}
 
 	CachedReadAddresses.Reset();
-
-	bool bDestroySelf = false;
-	TSet<FString> OutExpandedChildPropertyPaths;
-
-	for (const TSharedPtr<FPropertyNode>& ChildNode : ChildNodes)
-	{
-		if (ChildNode->HasNodeFlags(EPropertyNodeFlags::Expanded))
-		{
-			OutExpandedChildPropertyPaths.Add(ChildNode->GetPropertyPath());
-		}
-	}
 	
+	FStringPrefixTree ExpandedPropertyItemSet;
+	const TSharedRef<FPropertyNode> ThisAsSharedRef = AsShared();
+	FPropertyNodeUtils::GetExpandedItems(ThisAsSharedRef, ExpandedPropertyItemSet);
+
+	constexpr bool bDestroySelf = false;
 	DestroyTree(bDestroySelf);
 
 	if (MaxChildDepthAllowed != 0)
@@ -300,12 +330,9 @@ void FPropertyNode::RebuildChildren()
 		if (HasNodeFlags(EPropertyNodeFlags::CanBeExpanded) && (ChildNodes.Num() == 0))
 		{
 			InitChildNodes();
-			for (const TSharedPtr<FPropertyNode>& ChildNode : ChildNodes)
+			if (ExpandedPropertyItemSet.Size() > 0)
 			{
-				if ( OutExpandedChildPropertyPaths.Contains(ChildNode->GetPropertyPath()) )
-				{
-					ChildNode->SetNodeFlags(EPropertyNodeFlags::Expanded, true);
-				}
+				FPropertyNodeUtils::SetExpandedItems(ThisAsSharedRef, ExpandedPropertyItemSet);
 			}
 		}
 	}
diff --git a/Engine/Source/Editor/PropertyEditor/Private/SDetailsViewBase.cpp b/Engine/Source/Editor/PropertyEditor/Private/SDetailsViewBase.cpp
--- a/Engine/Source/Editor/PropertyEditor/Private/SDetailsViewBase.cpp
+++ b/Engine/Source/Editor/PropertyEditor/Private/SDetailsViewBase.cpp
@@ -1249,27 +1249,22 @@ void SDetailsViewBase::Tick( const FGeometry& AllottedGeometry, const double InC
 	{
 		for (FDetailLayoutData& LayoutData : DetailLayouts)
 		{
-			FRootPropertyNodeList& ExternalRootPropertyNodes = LayoutData.DetailLayout->GetExternalRootPropertyNodes();
-
-			for (int32 NodeIndex = 0; NodeIndex < ExternalRootPropertyNodes.Num(); ++NodeIndex)
+			for (const TSharedPtr<FPropertyNode>& PropertyNode : LayoutData.DetailLayout->GetExternalRootPropertyNodes())
 			{
-				TSharedPtr<FPropertyNode> PropertyNode = ExternalRootPropertyNodes[NodeIndex];
+				EPropertyDataValidationResult Result = PropertyNode->EnsureDataIsValid();
+				if (Result == EPropertyDataValidationResult::PropertiesChanged || Result == EPropertyDataValidationResult::EditInlineNewValueChanged)
 				{
-					EPropertyDataValidationResult Result = PropertyNode->EnsureDataIsValid();
-					if (Result == EPropertyDataValidationResult::PropertiesChanged || Result == EPropertyDataValidationResult::EditInlineNewValueChanged)
-					{
-						// Note this will invalidate all the external root nodes so there is no need to continue
-						ExternalRootPropertyNodes.Empty();
+					// Note this will invalidate all the external root nodes so there is no need to continue
+					LayoutData.DetailLayout->ClearExternalRootPropertyNodes();
 
-						UpdatePropertyMaps();
-						bUpdateFilteredDetails = true;
+					UpdatePropertyMaps();
+					bUpdateFilteredDetails = true;
 
-						break;
-					}
-					else if (Result == EPropertyDataValidationResult::ArraySizeChanged || Result == EPropertyDataValidationResult::ChildrenRebuilt)
-					{
-						bUpdateFilteredDetails = true;
-					}
+					break;
+				}
+				else if (Result == EPropertyDataValidationResult::ArraySizeChanged || Result == EPropertyDataValidationResult::ChildrenRebuilt)
+				{
+					bUpdateFilteredDetails = true;
 				}
 			}
 		}
