Fix compilation error on Linux
#rnx

diff --git a/Engine/Source/Editor/PropertyEditor/Private/SDetailsViewBase.cpp b/Engine/Source/Editor/PropertyEditor/Private/SDetailsViewBase.cpp
--- a/Engine/Source/Editor/PropertyEditor/Private/SDetailsViewBase.cpp
+++ b/Engine/Source/Editor/PropertyEditor/Private/SDetailsViewBase.cpp
@@ -1247,9 +1247,9 @@ void SDetailsViewBase::Tick( const FGeometry& AllottedGeometry, const double InC
 
 	if (bValidateExternalNodes)
 	{
-		for (FDetailLayoutData& LayoutData : DetailLayouts)
+		for (const FDetailLayoutData& LayoutData : DetailLayouts)
 		{
-			for (const TSharedPtr<FPropertyNode>& PropertyNode : LayoutData.DetailLayout->GetExternalRootPropertyNodes())
+			for (const TSharedPtr<FComplexPropertyNode>& PropertyNode : LayoutData.DetailLayout->GetExternalRootPropertyNodes())
 			{
 				EPropertyDataValidationResult Result = PropertyNode->EnsureDataIsValid();
 				if (Result == EPropertyDataValidationResult::PropertiesChanged || Result == EPropertyDataValidationResult::EditInlineNewValueChanged)
