Warnings in Details display manager in build

Fix the following warnings for details view:
Engine\Source\Editor\PropertyEditor\Public\DetailsViewObjectFilter.h(52): warning V1053: Calling the 'InitializeDetailsDisplayManager' virtual function in the constructor may lead to unexpected result at runtime.
Initialization was moved to constructor, method was removed. Lazy initialization in GetDisplayManager was removed as it it not needed.

diff --git a/Engine/Source/Editor/PropertyEditor/Public/DetailsViewObjectFilter.h b/Engine/Source/Editor/PropertyEditor/Public/DetailsViewObjectFilter.h
--- a/Engine/Source/Editor/PropertyEditor/Public/DetailsViewObjectFilter.h
+++ b/Engine/Source/Editor/PropertyEditor/Public/DetailsViewObjectFilter.h
@@ -49,7 +49,7 @@ public:
 
 	explicit FDetailsViewObjectFilter()
 	{
-		InitializeDetailsDisplayManager();
+		DisplayManager = MakeShared<FDetailsDisplayManager>();
 	}
 
 	virtual ~FDetailsViewObjectFilter()
@@ -76,10 +76,6 @@ public:
 	*/
 	virtual TSharedPtr<FDetailsDisplayManager> GetDisplayManager()
 	{
-		if (!DisplayManager.IsValid())
-		{
-			InitializeDetailsDisplayManager();
-		}
 		return DisplayManager;
 	}
 
@@ -90,14 +86,6 @@ protected:
 	 * details display
 	 */
 	TSharedPtr<FDetailsDisplayManager> DisplayManager;
-
-	/**
-	 * Initializes the Details Display Manager
-	 */
-	virtual void InitializeDetailsDisplayManager()
-	{
-		DisplayManager = MakeShared<FDetailsDisplayManager>();
-	}
 	
 };
  
\ No newline at end of file
