Added structure data provider interface support to the details category builder.

diff --git a/Engine/Source/Editor/PropertyEditor/Private/DetailCategoryBuilderImpl.cpp b/Engine/Source/Editor/PropertyEditor/Private/DetailCategoryBuilderImpl.cpp
--- a/Engine/Source/Editor/PropertyEditor/Private/DetailCategoryBuilderImpl.cpp
+++ b/Engine/Source/Editor/PropertyEditor/Private/DetailCategoryBuilderImpl.cpp
@@ -502,6 +502,27 @@ IDetailPropertyRow* FDetailCategoryImpl::AddExternalStructureProperty(TSharedPtr
 	return NewRow.Get();
 }
 
+IDetailPropertyRow* FDetailCategoryImpl::AddExternalStructureProperty(TSharedPtr<IStructureDataProvider> StructDataProvider, FName PropertyName, EPropertyLocation::Type Location, const FAddPropertyParams& Params)
+{
+	FDetailLayoutCustomization NewCustomization;
+	NewCustomization.bCustom = true;
+	NewCustomization.bAdvanced = Location == EPropertyLocation::Advanced;
+
+	FDetailPropertyRow::MakeExternalPropertyRowCustomization(StructDataProvider, PropertyName, AsShared(), NewCustomization, Params);
+
+	TSharedPtr<FDetailPropertyRow> NewRow = NewCustomization.PropertyRow;
+
+	if (NewRow.IsValid())
+	{
+		TSharedPtr<FPropertyNode> PropertyNode = NewRow->GetPropertyNode();
+		TSharedPtr<FComplexPropertyNode> RootNode = StaticCastSharedRef<FComplexPropertyNode>(PropertyNode->FindComplexParent()->AsShared());
+
+		AddCustomLayout(NewCustomization);
+	}
+
+	return NewRow.Get();
+}
+
 TArray<TSharedPtr<IPropertyHandle>> FDetailCategoryImpl::AddAllExternalStructureProperties(TSharedRef<FStructOnScope> StructData, EPropertyLocation::Type Location, TArray<IDetailPropertyRow*>* OutPropertiesRow)
 {
 	return AddAllExternalStructureProperties(MakeShared<FStructOnScopeStructureDataProvider>(StructData), Location, OutPropertiesRow);
diff --git a/Engine/Source/Editor/PropertyEditor/Private/DetailCategoryBuilderImpl.h b/Engine/Source/Editor/PropertyEditor/Private/DetailCategoryBuilderImpl.h
--- a/Engine/Source/Editor/PropertyEditor/Private/DetailCategoryBuilderImpl.h
+++ b/Engine/Source/Editor/PropertyEditor/Private/DetailCategoryBuilderImpl.h
@@ -174,8 +174,9 @@ public:
 	virtual IDetailPropertyRow* AddExternalObjectProperty(const TArray<UObject*>& Objects, FName PropertyName, EPropertyLocation::Type Location = EPropertyLocation::Default, const FAddPropertyParams& Params = FAddPropertyParams()) override;
 	virtual IDetailPropertyRow* AddExternalStructure(TSharedPtr<FStructOnScope> StructData, EPropertyLocation::Type Location = EPropertyLocation::Default) override;
 	virtual IDetailPropertyRow* AddExternalStructureProperty(TSharedPtr<FStructOnScope> StructData, FName PropertyName, EPropertyLocation::Type Location = EPropertyLocation::Default, const FAddPropertyParams& Params = FAddPropertyParams()) override;
+	virtual IDetailPropertyRow* AddExternalStructureProperty(TSharedPtr<IStructureDataProvider> StructData, FName PropertyName, EPropertyLocation::Type Location = EPropertyLocation::Default, const FAddPropertyParams& Params = FAddPropertyParams()) override;
 	virtual TArray<TSharedPtr<IPropertyHandle>> AddAllExternalStructureProperties(TSharedRef<FStructOnScope> StructData, EPropertyLocation::Type Location = EPropertyLocation::Default, TArray<IDetailPropertyRow*>* OutPropertiesRow = nullptr) override;
-	virtual TArray<TSharedPtr<IPropertyHandle>> AddAllExternalStructureProperties(TSharedPtr<IStructureDataProvider> StructProvider, EPropertyLocation::Type Location, TArray<IDetailPropertyRow*>* OutPropertiesRow) override;
+	virtual TArray<TSharedPtr<IPropertyHandle>> AddAllExternalStructureProperties(TSharedPtr<IStructureDataProvider> StructProvider, EPropertyLocation::Type Location, TArray<IDetailPropertyRow*>* OutPropertiesRow = nullptr) override;
 	virtual bool IsParentLayoutValid() const override { return DetailLayoutBuilder.IsValid(); }
 	virtual IDetailLayoutBuilder& GetParentLayout() const override { return *DetailLayoutBuilder.Pin(); }
 	virtual FDetailWidgetRow& AddCustomRow(const FText& FilterString, bool bForAdvanced = false) override;
diff --git a/Engine/Source/Editor/PropertyEditor/Private/DetailPropertyRow.cpp b/Engine/Source/Editor/PropertyEditor/Private/DetailPropertyRow.cpp
--- a/Engine/Source/Editor/PropertyEditor/Private/DetailPropertyRow.cpp
+++ b/Engine/Source/Editor/PropertyEditor/Private/DetailPropertyRow.cpp
@@ -39,7 +39,7 @@ FDetailPropertyRow::FDetailPropertyRow(TSharedPtr<FPropertyNode> InPropertyNode,
 			MakePropertyEditor(PropertyNodeRef, Utilities, PropertyEditor);
 		}
 		
-		if (PropertyNode->AsObjectNode() && ExternalRootNode.IsValid())
+		if (PropertyNode->AsComplexNode() && ExternalRootNode.IsValid()) // AsComplexNode works both for objects and structs
 		{
 			// We are showing an entirely different object inline.  Generate a layout for it now.
 			if (IDetailsViewPrivate* DetailsView = InParentCategory->GetDetailsView())
@@ -594,6 +594,47 @@ void FDetailPropertyRow::MakeExternalPropertyRowCustomization(TSharedPtr<FStruct
 	}
 }
 
+void FDetailPropertyRow::MakeExternalPropertyRowCustomization(TSharedPtr<IStructureDataProvider> StructDataProvider, FName PropertyName, TSharedRef<FDetailCategoryImpl> ParentCategory, struct FDetailLayoutCustomization& OutCustomization, const FAddPropertyParams& Parameters)
+{
+	TSharedRef<FStructurePropertyNode> RootPropertyNode = MakeShared<FStructurePropertyNode>();
+
+	//SET
+	RootPropertyNode->SetStructure(StructDataProvider);
+
+	FPropertyNodeInitParams InitParams;
+	InitParams.ParentNode = nullptr;
+	InitParams.Property = nullptr;
+	InitParams.ArrayOffset = 0;
+	InitParams.ArrayIndex = INDEX_NONE;
+	InitParams.bAllowChildren = true;
+	InitParams.bForceHiddenPropertyVisibility = Parameters.ShouldForcePropertyVisible() || FPropertySettings::Get().ShowHiddenProperties();
+	InitParams.bCreateCategoryNodes = PropertyName == NAME_None;
+
+	Parameters.OverrideAllowChildren(InitParams.bAllowChildren);
+	Parameters.OverrideCreateCategoryNodes(InitParams.bCreateCategoryNodes);
+
+	RootPropertyNode->InitNode(InitParams);
+
+	ParentCategory->GetParentLayoutImpl()->AddExternalRootPropertyNode(RootPropertyNode);
+
+	if (PropertyName != NAME_None)
+	{
+		TSharedPtr<FPropertyNode> PropertyNode = RootPropertyNode->GenerateSingleChild(PropertyName);
+		if (PropertyNode.IsValid())
+		{
+			PropertyNode->RebuildChildren();
+
+			OutCustomization.PropertyRow = MakeShared<FDetailPropertyRow>(PropertyNode, ParentCategory, RootPropertyNode);
+			OutCustomization.PropertyRow->SetCustomExpansionId(Parameters.GetUniqueId());
+		}
+	}
+	else
+	{
+		OutCustomization.PropertyRow = MakeShared<FDetailPropertyRow>(RootPropertyNode, ParentCategory, RootPropertyNode);
+		OutCustomization.PropertyRow->SetCustomExpansionId(Parameters.GetUniqueId());
+	}
+}
+
 void FDetailPropertyRow::MakeExternalPropertyRowCustomization(const TArray<UObject*>& InObjects, FName PropertyName, TSharedRef<FDetailCategoryImpl> ParentCategory, struct FDetailLayoutCustomization& OutCustomization, const FAddPropertyParams& Parameters)
 {
 	TSharedRef<FObjectPropertyNode> RootPropertyNode = MakeShared<FObjectPropertyNode>();
diff --git a/Engine/Source/Editor/PropertyEditor/Private/DetailPropertyRow.h b/Engine/Source/Editor/PropertyEditor/Private/DetailPropertyRow.h
--- a/Engine/Source/Editor/PropertyEditor/Private/DetailPropertyRow.h
+++ b/Engine/Source/Editor/PropertyEditor/Private/DetailPropertyRow.h
@@ -123,6 +123,7 @@ public:
 	EVisibility GetPropertyVisibility() const;
 
 	static void MakeExternalPropertyRowCustomization(TSharedPtr<FStructOnScope> StructData, FName PropertyName, TSharedRef<FDetailCategoryImpl> ParentCategory, struct FDetailLayoutCustomization& OutCustomization, const FAddPropertyParams& Parameters);
+	static void MakeExternalPropertyRowCustomization(TSharedPtr<IStructureDataProvider> StructDataProvider, FName PropertyName, TSharedRef<FDetailCategoryImpl> ParentCategory, struct FDetailLayoutCustomization& OutCustomization, const FAddPropertyParams& Parameters);
 	static void MakeExternalPropertyRowCustomization(const TArray<UObject*>& InObjects, FName PropertyName, TSharedRef<FDetailCategoryImpl> ParentCategory, struct FDetailLayoutCustomization& OutCustomization, const FAddPropertyParams& Parameters);
 
 private:
diff --git a/Engine/Source/Editor/PropertyEditor/Private/StructurePropertyNode.h b/Engine/Source/Editor/PropertyEditor/Private/StructurePropertyNode.h
--- a/Engine/Source/Editor/PropertyEditor/Private/StructurePropertyNode.h
+++ b/Engine/Source/Editor/PropertyEditor/Private/StructurePropertyNode.h
@@ -20,17 +20,17 @@ public:
 	virtual FStructurePropertyNode* AsStructureNode() override { return this; }
 	virtual const FStructurePropertyNode* AsStructureNode() const override { return this; }
 
-	void RemoveStructure()
+	void RemoveStructure(bool bInDestroySelf = true)
 	{
 		ClearCachedReadAddresses(true);
-		DestroyTree();
+		DestroyTree(bInDestroySelf);
 		StructProvider = nullptr;
 		WeakCachedBaseStruct.Reset();
 	}
 
 	void SetStructure(TSharedPtr<FStructOnScope> InStructData)
 	{
-		RemoveStructure();
+		RemoveStructure(false);
 		if (InStructData)
 		{
 			StructProvider = MakeShared<FStructOnScopeStructureDataProvider>(InStructData);
@@ -39,7 +39,7 @@ public:
 
 	void SetStructure(TSharedPtr<IStructureDataProvider> InStructProvider)
 	{
-		RemoveStructure();
+		RemoveStructure(false);
 		StructProvider = InStructProvider;
 	}
 
diff --git a/Engine/Source/Editor/PropertyEditor/Public/DetailCategoryBuilder.h b/Engine/Source/Editor/PropertyEditor/Public/DetailCategoryBuilder.h
--- a/Engine/Source/Editor/PropertyEditor/Public/DetailCategoryBuilder.h
+++ b/Engine/Source/Editor/PropertyEditor/Public/DetailCategoryBuilder.h
@@ -133,6 +133,7 @@ public:
 	 * @return A property row for customizing the property or NULL if the property could not be found
 	 */
 	virtual IDetailPropertyRow* AddExternalStructureProperty(TSharedPtr<FStructOnScope> StructData, FName PropertyName, EPropertyLocation::Type Location = EPropertyLocation::Default, const FAddPropertyParams& Params = FAddPropertyParams()) = 0;
+	virtual IDetailPropertyRow* AddExternalStructureProperty(TSharedPtr<IStructureDataProvider> StructData, FName PropertyName, EPropertyLocation::Type Location = EPropertyLocation::Default, const FAddPropertyParams& Params = FAddPropertyParams()) = 0;
 
 	/**
 	 * Adds all properties for the specified external structure to this category
