Fixed detail view diff for FSetProperty

diff --git a/Engine/Source/Editor/PropertyEditor/Private/AsyncDetailViewDiff.cpp b/Engine/Source/Editor/PropertyEditor/Private/AsyncDetailViewDiff.cpp
--- a/Engine/Source/Editor/PropertyEditor/Private/AsyncDetailViewDiff.cpp
+++ b/Engine/Source/Editor/PropertyEditor/Private/AsyncDetailViewDiff.cpp
@@ -193,7 +193,7 @@ bool TTreeDiffSpecification<TWeakPtr<FDetailTreeNode>>::AreMatching(const TWeakP
 			}
 			if (ParentPropertyA->IsA<FSetProperty>() || ParentPropertyB->IsA<FSetProperty>())
 			{
-				if (AsyncDetailViewDiffHelpers::SetKeysMatch(ParentPropertyNodeA, ParentPropertyNodeB, ArrayIndexA, ArrayIndexB, OwningObjectA, OwningObjectB))
+				if (!AsyncDetailViewDiffHelpers::SetKeysMatch(ParentPropertyNodeA, ParentPropertyNodeB, ArrayIndexA, ArrayIndexB, OwningObjectA, OwningObjectB))
 				{
 					return false;
 				}
