UE-197108 - "Optionals of class types aren't usable in the editor"
- adds support for setting a default value to an FOptionalPRoperty containing a FObjectProperty type (ie: classes, objects, etc)
- copies over the needed existing code-logic for finding all affected objects and sending them pre/post edit events which will mark them as dirty for set/clear optional calls.

diff --git a/Engine/Source/Editor/PropertyEditor/Private/PropertyHandleImpl.cpp b/Engine/Source/Editor/PropertyEditor/Private/PropertyHandleImpl.cpp
--- a/Engine/Source/Editor/PropertyEditor/Private/PropertyHandleImpl.cpp
+++ b/Engine/Source/Editor/PropertyEditor/Private/PropertyHandleImpl.cpp
@@ -5141,27 +5141,101 @@ FPropertyAccess::Result FPropertyHandleOptional::GetOptionalValue(FProperty*& Ou
 
 FPropertyAccess::Result FPropertyHandleOptional::SetOptionalValue(FProperty* NewValue)
 {
-	const TSharedPtr<FPropertyNode>& PropertyNode = Implementation->GetPropertyNode();
+	TSharedPtr<FPropertyNode> PropertyNode = Implementation->GetPropertyNode();
 	FOptionalProperty* OptionalProperty = CastField<FOptionalProperty>(PropertyNode->GetProperty());
 	
-	FReadAddressList Addresses;
-	if (!PropertyNode->GetReadAddress(Addresses))
+	FReadAddressList ReadAddresses;
+	if (!PropertyNode->GetReadAddress(ReadAddresses))
 	{
 		return FPropertyAccess::Fail;
 	}
+	
+	TArray<TArray<UObject*>> AffectedInstancesPerObject;
+	AffectedInstancesPerObject.SetNum(ReadAddresses.Num());
+
+	// List of top level objects sent to the PropertyChangedEvent
+	TArray<const UObject*> TopLevelObjects;
+	TopLevelObjects.Reserve(ReadAddresses.Num());
 
-	for (int i = 0; i < Addresses.Num(); i++)
+	// Begin a property edit transaction.
+	FScopedTransaction Transaction(NSLOCTEXT("UnrealEd", "SetOptional", "Set Optional"));
+	FObjectPropertyNode* ObjectNode = PropertyNode->FindObjectItemParent();
+
+	for ( int32 i = 0 ; i < ReadAddresses.Num() ; ++i )
 	{
-		void* Optional = Addresses.GetAddress(i);
-		OptionalProperty->MarkSetAndGetInitializedValuePointerToReplace(Optional);
+		void* Addr = ReadAddresses.GetAddress(i);
+		if ( Addr )
+		{
+			if (!FApp::IsGame())
+			{
+				UObject* Obj = ObjectNode ? ObjectNode->GetUObject(i) : nullptr;
+				if (IsTemplate(Obj))
+				{
+					PropertyNode->GatherInstancesAffectedByContainerPropertyChange(Obj, Addr, EPropertyArrayChangeType::Add, AffectedInstancesPerObject[i]);
+				}
+			}
+		}
+	}
 
+	TSet< UObject* > AllAffectedInstances;
+	for (const TArray<UObject*>& AffectedInstances : AffectedInstancesPerObject)
+	{
+		AllAffectedInstances.Append(AffectedInstances);
+	}
+
+	// send the PreEditChange notification to all selected objects
+	FNotifyHook* NotifyHook = Implementation->GetNotifyHook();
+	PropertyNode->NotifyPreChange(OptionalProperty, NotifyHook, AllAffectedInstances);
+
+	for (int i = 0; i < ReadAddresses.Num(); i++)
+	{
+		void* Addr = ReadAddresses.GetAddress(i);
+		OptionalProperty->MarkSetAndGetInitializedValuePointerToReplace(Addr);
+		
 		if (NewValue)
 		{
 			OptionalProperty->SetValueProperty(NewValue);
 		}
+
+		UObject* Obj = ObjectNode ? ObjectNode->GetUObject(i) : nullptr;
+		TopLevelObjects.Add(Obj);
+
+		// If our OptionalValue is a ptr to an object and we are not setting to a passed in value
+		// we need to intialize a default of that object and set the ptr to it.
+		FObjectProperty* ObjectProperty = CastField<FObjectProperty>(OptionalProperty->GetValueProperty());
+		if (ObjectProperty)
+		{
+			UObject* Outer = Obj;
+
+			if (Outer)
+			{
+				UObject* NewDefaultObjectValue = NewObject<UObject>(Outer, ObjectProperty->PropertyClass);
+				if (NewDefaultObjectValue)
+				{
+					void* ObjectPropertyValuePtr = ObjectProperty->ContainerPtrToValuePtr<void>(Addr);
+					ObjectProperty->SetObjectPropertyValue(ObjectPropertyValuePtr, NewDefaultObjectValue);
+				}
+			}
+			else
+			{
+				Implementation->ShowInvalidOperationError(LOCTEXT("SetOptionalElement", "Could not create a default value for optional object as could not determine outer object."));
+			}
+		}
 	}
 
-	// Rebuild our parent as we require a de-draw in the details panel
+	FPropertyChangedEvent ChangeEvent(OptionalProperty, EPropertyChangeType::ValueSet, MakeArrayView(TopLevelObjects));
+	ChangeEvent.SetInstancesChanged(MoveTemp(AllAffectedInstances));
+
+	// send the PostEditChange notification; it will be propagated to all selected objects
+	PropertyNode->NotifyPostChange(ChangeEvent, NotifyHook);
+
+	if (TSharedPtr<IPropertyUtilities> PropertyUtilities = Implementation->GetPropertyUtilities())
+	{
+		PropertyNode->FixPropertiesInEvent(ChangeEvent);
+		PropertyUtilities->NotifyFinishedChangingProperties(ChangeEvent);
+	}
+
+	// Rebuild our parent as we require a re-draw in the details panel
 	if (FPropertyNode* ParentNode = PropertyNode->GetParentNode())
 	{
 		ParentNode->RequestRebuildChildren();
@@ -5185,21 +5259,73 @@ FPropertyAccess::Result FPropertyHandleOptional::ClearOptionalValue()
 		return FPropertyAccess::Fail;
 	}
 
-	FReadAddressList Addresses;
-	if (!PropertyNode->GetReadAddress(Addresses))
+	FReadAddressList ReadAddresses;
+	if (!PropertyNode->GetReadAddress(ReadAddresses))
 	{
 		return FPropertyAccess::Fail;
 	}
+
+	TArray<TArray<UObject*>> AffectedInstancesPerObject;
+	AffectedInstancesPerObject.SetNum(ReadAddresses.Num());
+
+	// List of top level objects sent to the PropertyChangedEvent
+	TArray<const UObject*> TopLevelObjects;
+	TopLevelObjects.Reserve(ReadAddresses.Num());
+
+	// Begin a property edit transaction.
+	FScopedTransaction Transaction(NSLOCTEXT("UnrealEd", "ClearOptional", "Clear Optional"));
+	FObjectPropertyNode* ObjectNode = PropertyNode->FindObjectItemParent();
+
+	for ( int32 i = 0 ; i < ReadAddresses.Num() ; ++i )
+	{
+		void* Addr = ReadAddresses.GetAddress(i);
+		if ( Addr )
+		{
+			if (!FApp::IsGame())
+			{
+				UObject* Obj = ObjectNode ? ObjectNode->GetUObject(i) : nullptr;
+				if (IsTemplate(Obj))
+				{
+					PropertyNode->GatherInstancesAffectedByContainerPropertyChange(Obj, Addr, EPropertyArrayChangeType::Clear, AffectedInstancesPerObject[i]);
+				}
+			}
+		}
+	}
+
+	TSet< UObject* > AllAffectedInstances;
+	for (const TArray<UObject*>& AffectedInstances : AffectedInstancesPerObject)
+	{
+		AllAffectedInstances.Append(AffectedInstances);
+	}
+
+	// send the PreEditChange notification to all selected objects
+	FNotifyHook* NotifyHook = Implementation->GetNotifyHook();
+	PropertyNode->NotifyPreChange(OptionalProperty, NotifyHook, AllAffectedInstances);
 	
-	for (int i = 0; i < Addresses.Num(); i++)
+	for (int i = 0; i < ReadAddresses.Num(); i++)
 	{
-		void* Optional = Addresses.GetAddress(i);
-		OptionalProperty->MarkUnset(Optional);
+		UObject* Obj = ObjectNode ? ObjectNode->GetUObject(i) : nullptr;
+		TopLevelObjects.Add(Obj);
+
+		void* Addr = ReadAddresses.GetAddress(i);
+		OptionalProperty->MarkUnset(Addr);
 	}
 
 	// Could be removed as unecessary (rebuild will do this for us... but removing now makes any future debugging clearer)
 	PropertyNode->GetOptionalValueNode().Reset();
 
+	FPropertyChangedEvent ChangeEvent(OptionalProperty, EPropertyChangeType::ValueSet, MakeArrayView(TopLevelObjects));
+	ChangeEvent.SetInstancesChanged(MoveTemp(AllAffectedInstances));
+
+	// send the PostEditChange notification; it will be propagated to all selected objects
+	PropertyNode->NotifyPostChange(ChangeEvent, NotifyHook);
+
+	if (TSharedPtr<IPropertyUtilities> PropertyUtilities = Implementation->GetPropertyUtilities())
+	{
+		PropertyNode->FixPropertiesInEvent(ChangeEvent);
+		PropertyUtilities->NotifyFinishedChangingProperties(ChangeEvent);
+	}
+
 	// Rebuild our parent as we require a re-draw in the details panel
 	if (FPropertyNode* ParentNode = PropertyNode->GetParentNode())
 	{
