Fixed array indices not being provided for owner of modified UObject
- Changed FPropertyValueImpl::GenerateArrayIndexMapToObjectNode to not stop on the first object, so that we get array indices for the whole property path in case of sub-objects
- Fixed the event copy in FPropertyNode::NotifyPostChange() to retain array indices from the orignal event

diff --git a/Engine/Source/Editor/PropertyEditor/Private/PropertyHandleImpl.cpp b/Engine/Source/Editor/PropertyEditor/Private/PropertyHandleImpl.cpp
--- a/Engine/Source/Editor/PropertyEditor/Private/PropertyHandleImpl.cpp
+++ b/Engine/Source/Editor/PropertyEditor/Private/PropertyHandleImpl.cpp
@@ -211,7 +211,7 @@ void FPropertyValueImpl::GenerateArrayIndexMapToObjectNode( TMap<FString,int32>&
 	if( PropertyNode )
 	{
 		OutArrayIndexMap.Empty();
-		for (FPropertyNode* IterationNode = PropertyNode; (IterationNode != nullptr) && (IterationNode->AsObjectNode() == nullptr); IterationNode = IterationNode->GetParentNode())
+		for (FPropertyNode* IterationNode = PropertyNode; (IterationNode != nullptr); IterationNode = IterationNode->GetParentNode())
 		{
 			FProperty* Property = IterationNode->GetProperty();
 			if (Property)
diff --git a/Engine/Source/Editor/PropertyEditor/Private/PropertyNode.cpp b/Engine/Source/Editor/PropertyEditor/Private/PropertyNode.cpp
--- a/Engine/Source/Editor/PropertyEditor/Private/PropertyNode.cpp
+++ b/Engine/Source/Editor/PropertyEditor/Private/PropertyNode.cpp
@@ -3052,9 +3052,12 @@ void FPropertyNode::NotifyPostChange( FPropertyChangedEvent& InPropertyChangedEv
 				auto ScopePostEditChange = [&PropertyChain, &InPropertyChangedEvent, &CurProperty, CurrentObjectIndex](UObject* Object)
 				{
 					// copy the property changed event
-					FPropertyChangedEvent ChangedEvent = CurProperty != InPropertyChangedEvent.Property ? 
-						FPropertyChangedEvent(CurProperty, InPropertyChangedEvent.ChangeType) : 
-						InPropertyChangedEvent;
+					FPropertyChangedEvent ChangedEvent = InPropertyChangedEvent;
+					if (CurProperty != InPropertyChangedEvent.Property)
+					{
+						ChangedEvent.Property = CurProperty;
+						ChangedEvent.MemberProperty = CurProperty;
+					}
 					ChangedEvent.ObjectIteratorIndex = CurrentObjectIndex;
 
 					if (PropertyChain->Num() == 0)
