[Refactor] Several changes to property editor needed for scene-graph efforts

- Gave  FDetailItemNode::CreatePropertyHandle support for CustomBuilderRow property handles
- Added FPropertyHandleBase::CreateFPropertyPath method
- Added FPropertyHandleMap::GetElement method to parrody FPropertyHandleSet::GetElement
- Made FPropertNode::CreatePropertyPath return an empty/invalid path if any of it's properties are invalid (likely due to a destroyed object)
- Added FPropertyPath::GetPropertyInfoRef method So that an FPropertyPath element can be mutated without having to recreate the entire path

diff --git a/Engine/Source/Editor/PropertyEditor/Private/DetailItemNode.cpp b/Engine/Source/Editor/PropertyEditor/Private/DetailItemNode.cpp
--- a/Engine/Source/Editor/PropertyEditor/Private/DetailItemNode.cpp
+++ b/Engine/Source/Editor/PropertyEditor/Private/DetailItemNode.cpp
@@ -111,6 +111,10 @@ TSharedPtr<IPropertyHandle> FDetailItemNode::CreatePropertyHandle() const
 			return Handles[0];
 		}
 	}
+	else if (Customization.HasCustomBuilder())
+	{
+		return Customization.CustomBuilderRow->GetPropertyHandle();
+	}
 
 	return nullptr;
 }
@@ -848,6 +852,11 @@ FPropertyPath FDetailItemNode::GetPropertyPath() const
 			}
 		}
 	}
+
+	if (const TSharedPtr<IPropertyHandle> PropertyHandle = CreatePropertyHandle())
+	{
+		return *PropertyHandle->CreateFPropertyPath();
+	}
 	return Ret;
 }
 
diff --git a/Engine/Source/Editor/PropertyEditor/Private/PropertyHandleImpl.cpp b/Engine/Source/Editor/PropertyEditor/Private/PropertyHandleImpl.cpp
--- a/Engine/Source/Editor/PropertyEditor/Private/PropertyHandleImpl.cpp
+++ b/Engine/Source/Editor/PropertyEditor/Private/PropertyHandleImpl.cpp
@@ -2278,6 +2278,15 @@ FStringView FPropertyHandleBase::GetPropertyPath() const
 	return FStringView();
 }
 
+TSharedPtr<FPropertyPath> FPropertyHandleBase::CreateFPropertyPath() const
+{
+	if (Implementation->GetPropertyNode().IsValid())
+ 	{
+ 		return FPropertyNode::CreatePropertyPath(Implementation->GetPropertyNode().ToSharedRef());
+ 	}
+	return TSharedPtr<FPropertyPath>();
+}
+
 int32 FPropertyHandleBase::GetArrayIndex() const
 {
 	if (Implementation->GetPropertyNode().IsValid())
@@ -5589,6 +5598,12 @@ FPropertyAccess::Result FPropertyHandleMap::GetNumElements(uint32& OutNumChildre
 	return FPropertyAccess::Success;
 }
 
+TSharedRef<IPropertyHandle> FPropertyHandleMap::GetElement(int32 Index) const
+{
+	TSharedPtr<FPropertyNode> PropertyNode = Implementation->GetChildNode(Index);
+	return PropertyEditorHelpers::GetPropertyHandle(PropertyNode.ToSharedRef(), Implementation->GetNotifyHook(), Implementation->GetPropertyUtilities()).ToSharedRef();
+}
+
 FDelegateHandle FPropertyHandleMap::SetOnNumElementsChanged( const FSimpleDelegate& OnChildrenChanged )
 {
 	return Implementation->SetOnRebuildChildren(OnChildrenChanged);
diff --git a/Engine/Source/Editor/PropertyEditor/Private/PropertyHandleImpl.h b/Engine/Source/Editor/PropertyEditor/Private/PropertyHandleImpl.h
--- a/Engine/Source/Editor/PropertyEditor/Private/PropertyHandleImpl.h
+++ b/Engine/Source/Editor/PropertyEditor/Private/PropertyHandleImpl.h
@@ -513,6 +513,7 @@ public:
 	virtual FName GetDefaultCategoryName() const override;
 	virtual FText GetDefaultCategoryText() const override;
 	virtual FStringView GetPropertyPath() const override;
+	virtual TSharedPtr<FPropertyPath> CreateFPropertyPath() const override;
 	virtual int32 GetArrayIndex() const override;
 	virtual void RequestRebuildChildren() override;
 	virtual bool IsFavorite() const override;
@@ -777,6 +778,7 @@ public:
 	virtual FPropertyAccess::Result Empty() override;
 	virtual FPropertyAccess::Result DeleteItem(int32 Index) override;
 	virtual FPropertyAccess::Result GetNumElements(uint32& OutNumElements) override;
+	virtual TSharedRef<IPropertyHandle> GetElement(int32 Index) const override;
 	virtual FDelegateHandle SetOnNumElementsChanged( const FSimpleDelegate& InOnNumElementsChanged ) override;
 	virtual void UnregisterOnNumElementsChanged(FDelegateHandle Handle) override;
 	virtual bool HasDocumentation() override { return true; }
diff --git a/Engine/Source/Editor/PropertyEditor/Private/PropertyNode.h b/Engine/Source/Editor/PropertyEditor/Private/PropertyNode.h
--- a/Engine/Source/Editor/PropertyEditor/Private/PropertyNode.h
+++ b/Engine/Source/Editor/PropertyEditor/Private/PropertyNode.h
@@ -810,6 +810,10 @@ public:
 			{
 				FPropertyInfo NewPropInfo;
 				NewPropInfo.Property = CurrentNode->GetProperty();
+				if (!NewPropInfo.Property.IsValid())
+				{
+					return MakeShareable(new FPropertyPath());
+				}
 				NewPropInfo.ArrayIndex = CurrentNode->GetArrayIndex();
 
 				Properties.Add(NewPropInfo);
diff --git a/Engine/Source/Editor/PropertyEditor/Public/PropertyHandle.h b/Engine/Source/Editor/PropertyEditor/Public/PropertyHandle.h
--- a/Engine/Source/Editor/PropertyEditor/Public/PropertyHandle.h
+++ b/Engine/Source/Editor/PropertyEditor/Public/PropertyHandle.h
@@ -77,6 +77,11 @@ public:
 	 */
 	virtual FStringView GetPropertyPath() const = 0;
 
+	/**
+	 * Helper to fetch a PropertyPath 
+	 */
+	virtual TSharedPtr<FPropertyPath> CreateFPropertyPath() const = 0;
+
 	/**
 	 * Helper to fetch the ArrayIndex
 	 */
@@ -904,6 +909,11 @@ public:
 	 */
 	virtual FPropertyAccess::Result GetNumElements(uint32& OutNumElements) = 0;
 
+	/**
+	 * @return a handle to the element at the specified index
+	 */
+	virtual TSharedRef<IPropertyHandle> GetElement(int32 Index) const = 0;
+
 	/**
 	 * Sets a delegate to call when the number of elements changes
 	 */
diff --git a/Engine/Source/Editor/PropertyEditor/Public/PropertyPath.h b/Engine/Source/Editor/PropertyEditor/Public/PropertyPath.h
--- a/Engine/Source/Editor/PropertyEditor/Public/PropertyPath.h
+++ b/Engine/Source/Editor/PropertyEditor/Public/PropertyPath.h
@@ -71,6 +71,11 @@ public:
 		return Properties[ index ];
 	}
 
+	FPropertyInfo& GetPropertyInfoRef( int32 index )
+	{
+		return Properties[ index ];
+	}
+
 	const FPropertyInfo& GetLeafMostProperty() const
 	{
 		return Properties[ Properties.Num() - 1 ];
