[Refactor] Several changes to diff utils needed for scene-graph efforts
- Added ShouldInheritEqualFromChildren virtual method to TTreeDiffSpecification. This enables an optimization that can infer the diff result of parent nodes based on the diff result of it's children
- Changed AsyncDetailViewDiff and SDetailsSplitter code to favor usage of property handles since they handle multi-edit better and handle maps/sets better
- added transactions to SDetailsSplitter assignment button so that it can be reverted

diff --git a/Engine/Source/Editor/PropertyEditor/Internal/AsyncDetailViewDiff.h b/Engine/Source/Editor/PropertyEditor/Internal/AsyncDetailViewDiff.h
--- a/Engine/Source/Editor/PropertyEditor/Internal/AsyncDetailViewDiff.h
+++ b/Engine/Source/Editor/PropertyEditor/Internal/AsyncDetailViewDiff.h
@@ -38,6 +38,12 @@ public:
 	 * @param TreeNode node from one of the two user provided trees (guaranteed not to be null)
 	*/
 	virtual bool ShouldMatchByValue(const TWeakPtr<FDetailTreeNode>& TreeNode) const;
+
+	/**
+	 * return true if TreeNode is considered equal when all it's children are equal.
+	 * This avoids an unnecessary call to AreValuesEqual
+	*/
+	virtual bool ShouldInheritEqualFromChildren(const TWeakPtr<FDetailTreeNode>& TreeNodeA, const TWeakPtr<FDetailTreeNode>& TreeNodeB) const;
 };
 
 // asynchronously update a difference tree with changes to details views
diff --git a/Engine/Source/Editor/PropertyEditor/Private/AsyncDetailViewDiff.cpp b/Engine/Source/Editor/PropertyEditor/Private/AsyncDetailViewDiff.cpp
--- a/Engine/Source/Editor/PropertyEditor/Private/AsyncDetailViewDiff.cpp
+++ b/Engine/Source/Editor/PropertyEditor/Private/AsyncDetailViewDiff.cpp
@@ -17,69 +17,6 @@ namespace AsyncDetailViewDiffHelpers
 		}
 		return {};
 	}
-
-	FResolvedProperty GetResolvedProperty(const TSharedPtr<FPropertyNode>& PropertyNode, const UObject* Object)
-	{
-		if (PropertyNode && Object)
-		{
-			const TSharedRef<FPropertyPath> PropertyPath = FPropertyNode::CreatePropertyPath(PropertyNode.ToSharedRef());
-			if (PropertyPath->IsValid())
-			{
-				return FPropertySoftPath(*PropertyPath).Resolve(Object);
-			}
-		}
-		return FResolvedProperty();
-	}
-	
-	bool MapKeysMatch(const TSharedRef<FPropertyNode>& MapPropertyNodeA, const TSharedRef<FPropertyNode>& MapPropertyNodeB, int32 KeyIndexA, int32 KeyIndexB,
-		const UObject* OwningObjectA, const UObject* OwningObjectB)
-	{
-		const FMapProperty* MapPropertyA = CastField<FMapProperty>(MapPropertyNodeA->GetProperty());
-		const FMapProperty* MapPropertyB = CastField<FMapProperty>(MapPropertyNodeB->GetProperty());
-		if (!MapPropertyA || !MapPropertyB)
-		{
-			return false;
-		}
-		
-		const FResolvedProperty ResolvedMapA = AsyncDetailViewDiffHelpers::GetResolvedProperty(MapPropertyNodeA, OwningObjectA);
-		const FResolvedProperty ResolvedMapB = AsyncDetailViewDiffHelpers::GetResolvedProperty(MapPropertyNodeB, OwningObjectB);
-		FScriptMapHelper MapHelperA(MapPropertyA, MapPropertyA->ContainerPtrToValuePtr<UObject*>(ResolvedMapA.Object));
-		FScriptMapHelper MapHelperB(MapPropertyB, MapPropertyB->ContainerPtrToValuePtr<UObject*>(ResolvedMapB.Object));
-		
-		const void* KeyA = MapHelperA.FindNthKeyPtr(KeyIndexA);
-		const void* KeyB = MapHelperB.FindNthKeyPtr(KeyIndexB);
-		
-		if (MapPropertyA->KeyProp->SameType(MapPropertyB->KeyProp))
-		{
-			return MapPropertyA->KeyProp->Identical(KeyA, KeyB, PPF_DeepComparison);
-		}
-		return false;
-	}
-
-	bool SetKeysMatch(const TSharedRef<FPropertyNode>& SetPropertyNodeA, const TSharedRef<FPropertyNode>& SetPropertyNodeB, int32 KeyIndexA, int32 KeyIndexB,
-		const UObject* OwningObjectA, const UObject* OwningObjectB)
-	{
-		const FSetProperty* SetPropertyA = CastField<FSetProperty>(SetPropertyNodeA->GetProperty());
-		const FSetProperty* SetPropertyB = CastField<FSetProperty>(SetPropertyNodeB->GetProperty());
-		if (!SetPropertyA || !SetPropertyB)
-		{
-			return false;
-		}
-		
-		const FResolvedProperty ResolvedSetA = AsyncDetailViewDiffHelpers::GetResolvedProperty(SetPropertyNodeA, OwningObjectA);
-		const FResolvedProperty ResolvedSetB = AsyncDetailViewDiffHelpers::GetResolvedProperty(SetPropertyNodeB, OwningObjectB);
-		FScriptSetHelper SetHelperA(SetPropertyA, SetPropertyA->ContainerPtrToValuePtr<UObject*>(ResolvedSetA.Object));
-		FScriptSetHelper SetHelperB(SetPropertyB, SetPropertyB->ContainerPtrToValuePtr<UObject*>(ResolvedSetB.Object));
-		
-		const void* KeyA = SetHelperA.FindNthElementPtr(KeyIndexA);
-		const void* KeyB = SetHelperB.FindNthElementPtr(KeyIndexB);
-		
-		if (SetPropertyA->ElementProp->SameType(SetPropertyB->ElementProp))
-		{
-			return SetPropertyA->ElementProp->Identical(KeyA, KeyB, PPF_DeepComparison);
-		}
-		return false;
-	}
 }
 
 
@@ -137,80 +74,38 @@ bool TTreeDiffSpecification<TWeakPtr<FDetailTreeNode>>::AreMatching(const TWeakP
 		return PinnedTreeNodeA == PinnedTreeNodeB;
 	}
 
-	
-	TArray<TSharedRef<FPropertyNode>> PropertyNodesA;
-	TArray<TSharedRef<FPropertyNode>> PropertyNodesB;
-	PinnedTreeNodeA->GetAllPropertyNodes(PropertyNodesA);
-	PinnedTreeNodeB->GetAllPropertyNodes(PropertyNodesB);
-	if (PropertyNodesA.Num() != PropertyNodesB.Num())
-	{
-		return false;
-	}
-	if (PropertyNodesA.IsEmpty())
+	const TSharedPtr<IPropertyHandle> PropertyHandleA = PinnedTreeNodeA->CreatePropertyHandle();
+	const TSharedPtr<IPropertyHandle> PropertyHandleB = PinnedTreeNodeB->CreatePropertyHandle();
+	if (!PropertyHandleA || !PropertyHandleB)
 	{
 		// category nodes
 		return PinnedTreeNodeA->GetNodeName() == PinnedTreeNodeB->GetNodeName();
 	}
-
-	const TArray<TWeakObjectPtr<UObject>> OwningObjectsA = AsyncDetailViewDiffHelpers::GetObjects(PinnedTreeNodeA);
-	const TArray<TWeakObjectPtr<UObject>> OwningObjectsB = AsyncDetailViewDiffHelpers::GetObjects(PinnedTreeNodeB);
-	if (OwningObjectsA.IsEmpty() || OwningObjectsB.IsEmpty())
-	{
-		return PinnedTreeNodeA->GetNodeName() == PinnedTreeNodeB->GetNodeName();
-	}
-	
-	// because multi-edit puts matching keys together, we can assume all objects share the same key for this property
-	// and just look at the first instance
-	const UObject* OwningObjectB = OwningObjectsB[0].Get();
-	const UObject* OwningObjectA = OwningObjectsA[0].Get();
 	
-	for (int32 PropNodeIndex = 0; PropNodeIndex < PropertyNodesA.Num(); ++PropNodeIndex)
+	const int32 ArrayIndexA = PropertyHandleA->GetArrayIndex();
+	const int32 ArrayIndexB = PropertyHandleB->GetArrayIndex();
+	if (ArrayIndexA != INDEX_NONE && ArrayIndexB != INDEX_NONE)
 	{
-		const TSharedRef<FPropertyNode>& PropertyNodeA = PropertyNodesA[PropNodeIndex];
-		const TSharedRef<FPropertyNode>& PropertyNodeB = PropertyNodesB[PropNodeIndex];
-
-		// property nodes
-		const int32 ArrayIndexA = PropertyNodeA->GetArrayIndex();
-		const int32 ArrayIndexB = PropertyNodeB->GetArrayIndex();
-		const FProperty* PropertyA = PropertyNodeA->GetProperty();
-		const FProperty* PropertyB = PropertyNodeB->GetProperty();
-		
-		if (ArrayIndexA != INDEX_NONE && ArrayIndexB != INDEX_NONE)
-		{
-			const TSharedRef<FPropertyNode> ParentPropertyNodeA = PropertyNodeA->GetParentNode()->AsShared();
-			const TSharedRef<FPropertyNode> ParentPropertyNodeB = PropertyNodeB->GetParentNode()->AsShared();
-			const FProperty* ParentPropertyA = ParentPropertyNodeA->GetProperty();
-			const FProperty* ParentPropertyB = ParentPropertyNodeB->GetProperty();
-			
-			// sets and maps are stored by index in the property tree so we need to dig their keys out of the data
-			// and compare those instead
-			if (ParentPropertyA->IsA<FMapProperty>() || ParentPropertyB->IsA<FMapProperty>())
-			{
-				if (!AsyncDetailViewDiffHelpers::MapKeysMatch(ParentPropertyNodeA, ParentPropertyNodeB, ArrayIndexA, ArrayIndexB, OwningObjectA, OwningObjectB))
-				{
-					return false;
-				}
-			}
-			if (ParentPropertyA->IsA<FSetProperty>() || ParentPropertyB->IsA<FSetProperty>())
-			{
-				if (!AsyncDetailViewDiffHelpers::SetKeysMatch(ParentPropertyNodeA, ParentPropertyNodeB, ArrayIndexA, ArrayIndexB, OwningObjectA, OwningObjectB))
-				{
-					return false;
-				}
-			}
-		}
-		
-		if (ArrayIndexA != ArrayIndexB)
+		const TSharedPtr<IPropertyHandle> KeyHandleA = PropertyHandleA->GetKeyHandle();
+		const TSharedPtr<IPropertyHandle> KeyHandleB = PropertyHandleB->GetKeyHandle();
+		if (PropertyHandleA->GetKeyHandle() && PropertyHandleB->GetKeyHandle())
 		{
-			return false;
+			const TArray<TWeakObjectPtr<UObject>> OwningObjectsA = AsyncDetailViewDiffHelpers::GetObjects(PinnedTreeNodeA);
+			const TArray<TWeakObjectPtr<UObject>> OwningObjectsB = AsyncDetailViewDiffHelpers::GetObjects(PinnedTreeNodeB);
+			return DiffUtils::Identical(KeyHandleA, KeyHandleB, OwningObjectsA, OwningObjectsB);
 		}
-		const FName PropertyNameA = PropertyA ? PropertyA->GetFName() : NAME_None;
-		const FName PropertyNameB = PropertyB ? PropertyB->GetFName() : NAME_None;
 
-		if (PropertyNameA != PropertyNameB)
+		const TSharedPtr<IPropertyHandleSet> SetHandleA = PropertyHandleA->GetParentHandle()->AsSet();
+		const TSharedPtr<IPropertyHandleSet> SetHandleB = PropertyHandleB->GetParentHandle()->AsSet();
+		if (SetHandleA && SetHandleB)
 		{
-			return false;
+			// match set elements by value
+			const TArray<TWeakObjectPtr<UObject>> OwningObjectsA = AsyncDetailViewDiffHelpers::GetObjects(PinnedTreeNodeA);
+			const TArray<TWeakObjectPtr<UObject>> OwningObjectsB = AsyncDetailViewDiffHelpers::GetObjects(PinnedTreeNodeB);
+			return DiffUtils::Identical(PropertyHandleA, PropertyHandleB, OwningObjectsA, OwningObjectsB);
 		}
+		
+		return ArrayIndexA == ArrayIndexB;
 	}
 	return true;
 }
@@ -249,6 +144,12 @@ bool TTreeDiffSpecification<TWeakPtr<FDetailTreeNode>>::ShouldMatchByValue(const
 	return ParentArrayProperty && ArrayIndex != INDEX_NONE;
 }
 
+bool TTreeDiffSpecification<TWeakPtr<FDetailTreeNode, ESPMode::ThreadSafe>>::ShouldInheritEqualFromChildren(
+	const TWeakPtr<FDetailTreeNode>& TreeNodeA, const TWeakPtr<FDetailTreeNode>& TreeNodeB) const
+{
+	return false; // this theoretically could return true, but leaving it off has helped find false positives and false negatives
+}
+
 FAsyncDetailViewDiff::FAsyncDetailViewDiff(TSharedRef<IDetailsView> InLeftView, TSharedRef<IDetailsView> InRightView)
 	: TAsyncTreeDifferences(RootNodesAttribute(InLeftView), RootNodesAttribute(InRightView))
 	, LeftView(InLeftView)
diff --git a/Engine/Source/Editor/PropertyEditor/Private/SDetailsSplitter.cpp b/Engine/Source/Editor/PropertyEditor/Private/SDetailsSplitter.cpp
--- a/Engine/Source/Editor/PropertyEditor/Private/SDetailsSplitter.cpp
+++ b/Engine/Source/Editor/PropertyEditor/Private/SDetailsSplitter.cpp
@@ -3,10 +3,14 @@
 
 #include "AsyncDetailViewDiff.h"
 #include "DetailTreeNode.h"
+#include "Editor.h"
 #include "IDetailsViewPrivate.h"
 #include "PropertyNode.h"
 #include "Styling/StyleColors.h"
 #include "Framework/Application/SlateApplication.h"
+#include "Serialization/ObjectWriter.h"
+
+#define LOCTEXT_NAMESPACE "DetailsSplitter"
 
 namespace DetailsSplitterHelpers
 {
@@ -120,6 +124,7 @@ namespace DetailsSplitterHelpers
 	void CopyPropertyValueForInsert(const TSharedPtr<FDetailTreeNode>& SourceDetailsNode, const TSharedPtr<FDetailTreeNode>& DestinationDetailsNode)
 	{
 		const TSharedPtr<IPropertyHandle> DestinationHandle = DestinationDetailsNode->CreatePropertyHandle();
+		const TSharedPtr<IPropertyHandle> SourceHandle = SourceDetailsNode->CreatePropertyHandle();
 
 		// Array
 		TArray<TUniquePtr<FScriptArrayHelper>> SourceArrays;
@@ -132,6 +137,8 @@ namespace DetailsSplitterHelpers
 			if (ensure(TryGetDestinationContainer(DestinationDetailsNode, DestinationArrayPropertyNode, DestinationArrays, InsertIndex)))
 			{
 				ensure(SourceArrays.Num() == DestinationArrays.Num());
+				
+				GEditor->BeginTransaction(TEXT("DetailsSplitter"), FText::Format(LOCTEXT("CopyPropertyValueTransaction","Insert {0}"), SourceHandle->GetPropertyDisplayName()), nullptr);
 				DestinationHandle->NotifyPreChange();
 				for (int32 ArrayNum = 0; ArrayNum < SourceArrays.Num(); ++ ArrayNum)
 				{
@@ -145,6 +152,7 @@ namespace DetailsSplitterHelpers
 				
 				DestinationHandle->NotifyPostChange(EPropertyChangeType::ArrayAdd);
 				DestinationHandle->NotifyFinishedChangingProperties();
+				GEditor->EndTransaction();
 			}
 			return;
 		}
@@ -160,6 +168,7 @@ namespace DetailsSplitterHelpers
 			if (ensure(TryGetDestinationContainer(DestinationDetailsNode, DestinationPropertyNode, DestinationSets, InsertIndex)))
 			{
 				ensure(SourceSets.Num() == DestinationSets.Num());
+				GEditor->BeginTransaction(TEXT("DetailsSplitter"), FText::Format(LOCTEXT("CopyPropertyValueTransaction","Insert {0}"), SourceHandle->GetPropertyDisplayName()), nullptr);
 				DestinationHandle->NotifyPreChange();
 				for (int32 SetNum = 0; SetNum < SourceSets.Num(); ++ SetNum)
 				{
@@ -169,6 +178,7 @@ namespace DetailsSplitterHelpers
 				
 				DestinationHandle->NotifyPostChange(EPropertyChangeType::ArrayAdd);
 				DestinationHandle->NotifyFinishedChangingProperties();
+				GEditor->EndTransaction();
 			}
 			return;
 		}
@@ -184,6 +194,7 @@ namespace DetailsSplitterHelpers
 			if (ensure(TryGetDestinationContainer(DestinationDetailsNode, DestinationPropertyNode, DestinationMaps, InsertIndex)))
 			{
 				ensure(SourceMaps.Num() == DestinationMaps.Num());
+				GEditor->BeginTransaction(TEXT("DetailsSplitter"), FText::Format(LOCTEXT("CopyPropertyValueTransaction","Insert {0}"), SourceHandle->GetPropertyDisplayName()), nullptr);
 				DestinationHandle->NotifyPreChange();
 				for (int32 MapNum = 0; MapNum < SourceMaps.Num(); ++ MapNum)
 				{
@@ -195,11 +206,19 @@ namespace DetailsSplitterHelpers
 				
 				DestinationHandle->NotifyPostChange(EPropertyChangeType::ArrayAdd);
 				DestinationHandle->NotifyFinishedChangingProperties();
+				GEditor->EndTransaction();
 			}
 			return;
 		}
 	}
 
+	void AssignPropertyValue(const TSharedPtr<IPropertyHandle>& SourceHandle, const TSharedPtr<IPropertyHandle>& DestinationHandle)
+	{
+		TArray<FString> SourceValues;
+		SourceHandle->GetPerObjectValues(SourceValues);
+		DestinationHandle->SetPerObjectValues(SourceValues);
+	}
+
 	void CopyPropertyValue(const TSharedPtr<FDetailTreeNode>& SourceDetailsNode, const TSharedPtr<FDetailTreeNode>& DestinationDetailsNode, ETreeDiffResult Diff)
 	{
 		switch(Diff)
@@ -221,18 +240,23 @@ namespace DetailsSplitterHelpers
 				return;
 		}
 
-		TArray<void*> SourceData;
-		TArray<void*> DestinationData;
+
 		const TSharedPtr<IPropertyHandle> SourceHandle = SourceDetailsNode->CreatePropertyHandle();
 		const TSharedPtr<IPropertyHandle> DestinationHandle = DestinationDetailsNode->CreatePropertyHandle();
 		if (!ensure(SourceHandle && DestinationHandle))
 		{
 			return;
 		}
-		if (!ensure(SourceHandle->GetProperty()->SameType(DestinationHandle->GetProperty())))
+		
+		if (!SourceHandle->GetProperty()->SameType(DestinationHandle->GetProperty()))
 		{
+			// convert types by assigning via text serialization
+			AssignPropertyValue(SourceHandle, DestinationHandle);
 			return;
 		}
+
+		TArray<void*> SourceData;
+		TArray<void*> DestinationData;
 		SourceHandle->AccessRawData(SourceData);
 		DestinationHandle->AccessRawData(DestinationData);
 		if (!ensure(SourceData.Num() == DestinationData.Num()))
@@ -240,7 +264,7 @@ namespace DetailsSplitterHelpers
 			return;
 		}
 
-		
+		GEditor->BeginTransaction(TEXT("DetailsSplitter"), FText::Format(LOCTEXT("CopyPropertyValueTransaction","Copy {0}"), SourceHandle->GetPropertyDisplayName()), nullptr);
 		DestinationHandle->NotifyPreChange();
 		for (int32 I = 0; I < SourceData.Num(); ++I)
 		{
@@ -255,6 +279,7 @@ namespace DetailsSplitterHelpers
 		}
 		DestinationHandle->NotifyPostChange(EPropertyChangeType::ValueSet);
 		DestinationHandle->NotifyFinishedChangingProperties();
+		GEditor->EndTransaction();
 	}
 
 	// note: DestinationDetailsNode is the node before the position in the tree you wish to insert
@@ -300,6 +325,23 @@ namespace DetailsSplitterHelpers
 		return false;
 	}
 
+	bool CanAssignPropertyValue(const TSharedPtr<IPropertyHandle>& SourceHandle, const TSharedPtr<IPropertyHandle>& DestinationHandle)
+	{
+		TArray<FString> SourceValues;
+		TArray<FString> OldDestinationValues;
+		SourceHandle->GetPerObjectValues(SourceValues);
+		DestinationHandle->GetPerObjectValues(OldDestinationValues);
+		
+		FPropertyAccess::Result Result = DestinationHandle->SetPerObjectValues(SourceValues, EPropertyValueSetFlags::NotTransactable);
+		TArray<FString> ChangedDestinationValues;
+		DestinationHandle->GetPerObjectValues(ChangedDestinationValues);
+
+		// revert changes and query whether anything changed
+		DestinationHandle->SetPerObjectValues(OldDestinationValues, EPropertyValueSetFlags::NotTransactable);
+		
+		return Result != FPropertyAccess::Fail && OldDestinationValues != ChangedDestinationValues;
+	}
+
 	bool CanCopyPropertyValue(const TSharedPtr<FDetailTreeNode>& SourceDetailsNode, const TSharedPtr<FDetailTreeNode>& DestinationDetailsNode, ETreeDiffResult Diff)
 	{
 		if (!SourceDetailsNode || !DestinationDetailsNode)
@@ -319,7 +361,31 @@ namespace DetailsSplitterHelpers
 		{
 		// traditional copy
 		case ETreeDiffResult::DifferentValues:
-			return true;
+		{
+			const TSharedPtr<IPropertyHandle> SourceHandle = SourceDetailsNode->CreatePropertyHandle();
+			const TSharedPtr<IPropertyHandle> DestinationHandle = DestinationDetailsNode->CreatePropertyHandle();
+			if (SourceHandle && DestinationHandle)
+			{
+				TArray<void*> SourceData;
+				TArray<void*> DestinationData;
+				SourceHandle->AccessRawData(SourceData);
+				DestinationHandle->AccessRawData(DestinationData);
+				if (SourceData == DestinationData && SourceHandle->GetProperty() == DestinationHandle->GetProperty())
+                {
+                	// disable copying a value to itself since it's a no-op
+                	return false;
+                }
+				if (SourceHandle->GetProperty()->SameType(DestinationHandle->GetProperty()))
+				{
+					return true;
+				}
+				if (CanAssignPropertyValue(SourceHandle, DestinationHandle))
+				{
+					return true;
+				}
+			}
+			return false;
+		}
 
 		// insert
 		case ETreeDiffResult::MissingFromTree1:
@@ -925,3 +991,5 @@ void SDetailsSplitter::PaintCopyPropertyButton(FSlateWindowElementList& OutDrawE
 		ButtonColor
 	);
 }
+
+#undef LOCTEXT_NAMESPACE
diff --git a/Engine/Source/Editor/UnrealEd/Private/MergeUtils.cpp b/Engine/Source/Editor/UnrealEd/Private/MergeUtils.cpp
--- a/Engine/Source/Editor/UnrealEd/Private/MergeUtils.cpp
+++ b/Engine/Source/Editor/UnrealEd/Private/MergeUtils.cpp
@@ -423,6 +423,11 @@ public:
 		// array elements should match by value
 		return CastField<FArrayProperty>(TreeNodeA.KeyProperty->Owner.ToField()) != nullptr;
 	}
+	
+	virtual bool ShouldInheritEqualFromChildren(const FPropertyInstance& TreeNodeA, const FPropertyInstance& TreeNodeB) const
+	{
+		return true;
+	}
 };
 
 static TAsyncTreeDifferences<FPropertyInstance> ObjectPropTreeDiff(const UObject* ObjectA, const UObject* ObjectB)
