[CIS Fix] fixed duplicate LOCTEXT namespaces in SDetailsSplitter.cpp

diff --git a/Engine/Source/Editor/PropertyEditor/Private/SDetailsSplitter.cpp b/Engine/Source/Editor/PropertyEditor/Private/SDetailsSplitter.cpp
--- a/Engine/Source/Editor/PropertyEditor/Private/SDetailsSplitter.cpp
+++ b/Engine/Source/Editor/PropertyEditor/Private/SDetailsSplitter.cpp
@@ -138,7 +138,7 @@ namespace DetailsSplitterHelpers
 			{
 				ensure(SourceArrays.Num() == DestinationArrays.Num());
 				
-				GEditor->BeginTransaction(TEXT("DetailsSplitter"), FText::Format(LOCTEXT("CopyPropertyValueTransaction","Insert {0}"), SourceHandle->GetPropertyDisplayName()), nullptr);
+				GEditor->BeginTransaction(TEXT("DetailsSplitter"), FText::Format(LOCTEXT("InsertPropertyValueTransaction","Insert {0}"), SourceHandle->GetPropertyDisplayName()), nullptr);
 				DestinationHandle->NotifyPreChange();
 				for (int32 ArrayNum = 0; ArrayNum < SourceArrays.Num(); ++ ArrayNum)
 				{
@@ -168,7 +168,7 @@ namespace DetailsSplitterHelpers
 			if (ensure(TryGetDestinationContainer(DestinationDetailsNode, DestinationPropertyNode, DestinationSets, InsertIndex)))
 			{
 				ensure(SourceSets.Num() == DestinationSets.Num());
-				GEditor->BeginTransaction(TEXT("DetailsSplitter"), FText::Format(LOCTEXT("CopyPropertyValueTransaction","Insert {0}"), SourceHandle->GetPropertyDisplayName()), nullptr);
+				GEditor->BeginTransaction(TEXT("DetailsSplitter"), FText::Format(LOCTEXT("InsertPropertyValueTransaction","Insert {0}"), SourceHandle->GetPropertyDisplayName()), nullptr);
 				DestinationHandle->NotifyPreChange();
 				for (int32 SetNum = 0; SetNum < SourceSets.Num(); ++ SetNum)
 				{
@@ -194,7 +194,7 @@ namespace DetailsSplitterHelpers
 			if (ensure(TryGetDestinationContainer(DestinationDetailsNode, DestinationPropertyNode, DestinationMaps, InsertIndex)))
 			{
 				ensure(SourceMaps.Num() == DestinationMaps.Num());
-				GEditor->BeginTransaction(TEXT("DetailsSplitter"), FText::Format(LOCTEXT("CopyPropertyValueTransaction","Insert {0}"), SourceHandle->GetPropertyDisplayName()), nullptr);
+				GEditor->BeginTransaction(TEXT("DetailsSplitter"), FText::Format(LOCTEXT("InsertPropertyValueTransaction","Insert {0}"), SourceHandle->GetPropertyDisplayName()), nullptr);
 				DestinationHandle->NotifyPreChange();
 				for (int32 MapNum = 0; MapNum < SourceMaps.Num(); ++ MapNum)
 				{
