FPropertyPermissionList: Remove ability to pass NAME_None as owner as that defeats the purpose of having owners system and slows down debugging.
#rnx

diff --git a/Engine/Source/Editor/PropertyEditor/Public/PropertyPermissionList.h b/Engine/Source/Editor/PropertyEditor/Public/PropertyPermissionList.h
--- a/Engine/Source/Editor/PropertyEditor/Public/PropertyPermissionList.h
+++ b/Engine/Source/Editor/PropertyEditor/Public/PropertyPermissionList.h
@@ -58,15 +58,15 @@ public:
 	void UnregisterOwner(const FName Owner);
 
 	/** Add a specific property to a UStruct's AllowList */
-	void AddToAllowList(TSoftObjectPtr<UStruct> Struct, const FName PropertyName, const FName Owner = NAME_None);
+	void AddToAllowList(TSoftObjectPtr<UStruct> Struct, const FName PropertyName, const FName Owner);
 	/** Add a list of properties to a UStruct's AllowList */
-	void AddToAllowList(TSoftObjectPtr<UStruct> Struct, const TArray<FName>& PropertyNames, const FName Owner = NAME_None);
+	void AddToAllowList(TSoftObjectPtr<UStruct> Struct, const TArray<FName>& PropertyNames, const FName Owner);
 	/** Remove a specific property from a UStruct's AllowList */
-	void RemoveFromAllowList(TSoftObjectPtr<UStruct> Struct, const FName PropertyName, const FName Owner = NAME_None);
+	void RemoveFromAllowList(TSoftObjectPtr<UStruct> Struct, const FName PropertyName, const FName Owner);
 	/** Add a specific property to a UStruct's DenyList */
-	void AddToDenyList(TSoftObjectPtr<UStruct> Struct, const FName PropertyName, const FName Owner = NAME_None);
+	void AddToDenyList(TSoftObjectPtr<UStruct> Struct, const FName PropertyName, const FName Owner);
 	/** Remove a specific property from a UStruct's DenyList */
-    void RemoveFromDenyList(TSoftObjectPtr<UStruct> Struct, const FName PropertyName, const FName Owner = NAME_None);
+    void RemoveFromDenyList(TSoftObjectPtr<UStruct> Struct, const FName PropertyName, const FName Owner);
 
 	/** When the PermissionList or DenyList for any struct was added to or removed from. */
     FPermissionListUpdated PermissionListUpdatedDelegate;
