SOL-5617 - "The editor UI for [REDACTED_HANDLE] optionals does not respect property type customization"

- if we encounter a set option in `FDetailPropertyRow` we use its value instead. This allows for cusomizations to just 'work' without any special handling for options.
    - Note: There is some case-handling to ensure the display name is still the options rather than the values.

diff --git a/Engine/Source/Editor/PropertyEditor/Private/DetailItemNode.cpp b/Engine/Source/Editor/PropertyEditor/Private/DetailItemNode.cpp
--- a/Engine/Source/Editor/PropertyEditor/Private/DetailItemNode.cpp
+++ b/Engine/Source/Editor/PropertyEditor/Private/DetailItemNode.cpp
@@ -482,6 +482,16 @@ void FDetailItemNode::OnItemExpansionChanged( bool bInIsExpanded, bool bShouldSa
 	if( Customization.HasPropertyNode() )
 	{
 		Customization.GetPropertyNode()->SetNodeFlags( EPropertyNodeFlags::Expanded, bInIsExpanded );
+		
+		// This is a hack which was needed for optionals to behave correctly (they did not preserve expanded state between clicking on/off an item).
+		// I tried to find how FPropertyNode save or serialize their `EPropertyNodeFlags::Expanded` flag but couldn't find such functionality... 
+		//  
+		// My assumption was that however stuff is saved relies on the FPropertyNode tree structure which OptionalValues are not a part of...
+		// For now, relying on the parent optional instead creates the expected behaviour.
+		if (Customization.GetPropertyNode()->IsOptionalValueNode())
+		{
+			Customization.GetPropertyNode()->GetParentNode()->SetNodeFlags( EPropertyNodeFlags::Expanded, bInIsExpanded );
+		}
 	}
 
 	if (ParentCategory.IsValid() && bShouldSaveState &&
@@ -498,9 +508,14 @@ bool FDetailItemNode::ShouldBeExpanded() const
 	bool bShouldBeExpanded = bIsExpanded || bShouldBeVisibleDueToChildFiltering;
 	if( Customization.HasPropertyNode() )
 	{
-		FPropertyNode& PropertyNode = *Customization.GetPropertyNode();
-		bShouldBeExpanded = PropertyNode.HasNodeFlags( EPropertyNodeFlags::Expanded ) != 0;
-		bShouldBeExpanded |= PropertyNode.HasNodeFlags( EPropertyNodeFlags::IsSeenDueToChildFiltering ) != 0;
+		FPropertyNode* PropertyNode = Customization.GetPropertyNode().Get();
+		// This is a hack... see comment in FDetailItemNode::OnItemExpansionChanged
+		if (Customization.GetPropertyNode()->IsOptionalValueNode())
+		{
+			PropertyNode = Customization.GetPropertyNode()->GetParentNode();
+		}
+		bShouldBeExpanded = PropertyNode->HasNodeFlags( EPropertyNodeFlags::Expanded ) != 0;
+		bShouldBeExpanded |= PropertyNode->HasNodeFlags( EPropertyNodeFlags::IsSeenDueToChildFiltering ) != 0;
 	}
 	return bShouldBeExpanded;
 }
diff --git a/Engine/Source/Editor/PropertyEditor/Private/DetailPropertyRow.cpp b/Engine/Source/Editor/PropertyEditor/Private/DetailPropertyRow.cpp
--- a/Engine/Source/Editor/PropertyEditor/Private/DetailPropertyRow.cpp
+++ b/Engine/Source/Editor/PropertyEditor/Private/DetailPropertyRow.cpp
@@ -15,6 +15,8 @@
 #include "Modules/ModuleManager.h"
 #include "Widgets/Layout/SSpacer.h"
 
+#include "UObject/PropertyOptional.h"
+
 #define LOCTEXT_NAMESPACE	"DetailPropertyRow"
 
 FDetailPropertyRow::FDetailPropertyRow(TSharedPtr<FPropertyNode> InPropertyNode, TSharedRef<FDetailCategoryImpl> InParentCategory, TSharedPtr<FComplexPropertyNode> InExternalRootNode)
@@ -26,6 +28,39 @@ FDetailPropertyRow::FDetailPropertyRow(TSharedPtr<FPropertyNode> InPropertyNode,
 	, bForceAutoExpansion( false )
 	, bCachedCustomTypeInterface(false)
 {
+	// Is this a set optional property?
+	if (TSharedPtr<FPropertyNode>& ValueNode = PropertyNode->GetOptionalValueNode())
+	{
+		// If we are selecting multiple options with different states (ie set/unset) dont use the
+		// Value node so optional multi-select logic is displayed (see SPropertyEditorOptional.h).
+		uint8 MixedValues = 0;
+		FProperty* MyProperty = PropertyNode->GetProperty();
+		if (FOptionalProperty* OptionalProperty = CastField<FOptionalProperty>(MyProperty))
+		{
+			void* Optional = NULL;
+			FReadAddressList Addresses;
+			if (PropertyNode->GetReadAddress(Addresses))
+			{
+				for (int i = 0; i < Addresses.Num(); i++)
+				{
+					Optional = Addresses.GetAddress(i);
+					MixedValues |= OptionalProperty->IsSet(Optional) ? 1 : 2;
+					if (MixedValues == 3)
+					{
+						break;
+					}
+				}
+			}
+		}
+
+		if (MixedValues != 3)
+		{
+			// Swap our property node with the value node.
+			// Note: We still use the display name of the option for this row. 
+			PropertyNode = ValueNode;
+		}
+	}
+
 	PropertyHandle = InParentCategory->GetParentLayoutImpl()->GetPropertyHandle(PropertyNode);
 
 	if (PropertyNode.IsValid())
@@ -872,6 +907,17 @@ void FDetailPropertyRow::MakeNameOrKeyWidget( FDetailWidgetRow& Row, const TShar
 		}
 
 	}
+	else if (PropertyNode->IsOptionalValueNode())
+	{
+		TSharedRef<FPropertyEditor> ParentEditor = FPropertyEditor::Create(
+			PropertyNode->GetParentNode()->AsShared(), 
+			ParentCategory.Pin()->GetParentLayoutImpl()->GetPropertyUtilities()
+		);
+
+		NameWidget = 
+			SNew( SPropertyNameWidget, ParentEditor )
+			.IsEnabled( IsEnabledAttrib );
+	}
 	else if (InCustomRow.IsValid())
 	{
 		NameWidget = 
diff --git a/Engine/Source/Editor/PropertyEditor/Private/Presentation/PropertyEditor/PropertyEditor.cpp b/Engine/Source/Editor/PropertyEditor/Private/Presentation/PropertyEditor/PropertyEditor.cpp
--- a/Engine/Source/Editor/PropertyEditor/Private/Presentation/PropertyEditor/PropertyEditor.cpp
+++ b/Engine/Source/Editor/PropertyEditor/Private/Presentation/PropertyEditor/PropertyEditor.cpp
@@ -323,7 +323,13 @@ void FPropertyEditor::OnSetOptionalValue(FProperty* NewValue)
 
 void FPropertyEditor::OnClearOptionalValue()
 {
-	TSharedPtr<IPropertyHandleOptional> OptionalHandle = PropertyHandle->AsOptional();
+	// This func is called by an options value, not the option itself.
+	if (!PropertyNode->IsOptionalValueNode())
+	{
+		return;
+	}
+
+	TSharedPtr<IPropertyHandleOptional> OptionalHandle = PropertyHandle->GetParentHandle()->AsOptional();
 	if (OptionalHandle.IsValid())
 	{
 		OptionalHandle->ClearOptionalValue();
diff --git a/Engine/Source/Editor/PropertyEditor/Private/PropertyEditorHelpers.cpp b/Engine/Source/Editor/PropertyEditor/Private/PropertyEditorHelpers.cpp
--- a/Engine/Source/Editor/PropertyEditor/Private/PropertyEditorHelpers.cpp
+++ b/Engine/Source/Editor/PropertyEditor/Private/PropertyEditorHelpers.cpp
@@ -704,17 +704,6 @@ namespace PropertyEditorHelpers
 			}
 		}
 
-		//////////////////////////////
-		// Handle an optional property.
-		if (const FOptionalProperty* OptionalProp = CastField<FOptionalProperty>(NodeProperty))
-		{
-			uint8* ValueAddress = nullptr;
-			if (PropertyNode->GetSingleReadAddress(ValueAddress) == FPropertyAccess::Success && OptionalProp->IsSet(ValueAddress))
-			{
-				OutRequiredButtons.Add(EPropertyButton::OptionalClear);
-			}
-		}
-
 		//////////////////////////////
 		// Handle a class property.
 		FClassProperty* ClassProp = CastField<FClassProperty>(NodeProperty);
@@ -850,6 +839,12 @@ namespace PropertyEditorHelpers
 			}
 		}
 
+		//////////////////////////////
+		// Handle an optional value node.
+		if (PropertyNode->IsOptionalValueNode())
+		{
+			OutRequiredButtons.Add(EPropertyButton::OptionalClear);
+		}
 	}
 	
 	void MakeRequiredPropertyButtons( const TSharedRef<FPropertyNode>& PropertyNode, const TSharedRef<IPropertyUtilities>& PropertyUtilities,  TArray< TSharedRef<SWidget> >& OutButtons, const TArray<EPropertyButton::Type>& ButtonsToIgnore, bool bUsingAssetPicker  )
diff --git a/Engine/Source/Editor/PropertyEditor/Private/PropertyNode.cpp b/Engine/Source/Editor/PropertyEditor/Private/PropertyNode.cpp
--- a/Engine/Source/Editor/PropertyEditor/Private/PropertyNode.cpp
+++ b/Engine/Source/Editor/PropertyEditor/Private/PropertyNode.cpp
@@ -3458,6 +3458,11 @@ void FPropertyNode::InvalidateCachedState()
 	bUpdateDiffersFromDefault = true;
 	bUpdateEditConstState = true;
 
+	if (OptionalValueNode.IsValid())
+	{
+		OptionalValueNode->InvalidateCachedState();
+	}
+
 	for( TSharedPtr<FPropertyNode>& ChildNode : ChildNodes )
 	{
 		ChildNode->InvalidateCachedState();
diff --git a/Engine/Source/Editor/PropertyEditor/Private/PropertyNode.h b/Engine/Source/Editor/PropertyEditor/Private/PropertyNode.h
--- a/Engine/Source/Editor/PropertyEditor/Private/PropertyNode.h
+++ b/Engine/Source/Editor/PropertyEditor/Private/PropertyNode.h
@@ -1085,13 +1085,20 @@ public:
 	TSharedPtr<FPropertyNode>& GetOptionalValueNode() { return OptionalValueNode; }
 
 	/**
-	 * Interface function for getting a FOptionalProperty's Node's OptionalValueNode (May construct it if needed).
-	 * @return The optionals value OR null if this is not an optional or is an unset optional.
+	 * Interface function for getting an optionals ValueNode (May construct it if needed).
+	 * @return The optional's value OR null if this is a non-optional/unset-optional.
 	 * 
 	 * Note: This is currently the only method by which an optional's Value FPropertyNode is created (If you change this please update this documentation).
 	 */
 	virtual TSharedPtr<FPropertyNode>& GetOrCreateOptionalValueNode() { return OptionalValueNode; }
 
+	virtual bool IsOptionalValueNode() 
+	{ 
+		return ParentNodeWeakPtr.IsValid() 
+			&& ParentNodeWeakPtr.Pin()->GetOptionalValueNode().IsValid()
+			&& ParentNodeWeakPtr.Pin()->GetOptionalValueNode().Get() == this;
+	}
+
 protected:
 	TSharedRef<FEditPropertyChain> BuildPropertyChain( FProperty* PropertyAboutToChange ) const;
 	TSharedRef<FEditPropertyChain> BuildPropertyChain( FProperty* PropertyAboutToChange, const TSet<UObject*>& InAffectedArchetypeInstances ) const;
diff --git a/Engine/Source/Editor/PropertyEditor/Private/SDetailSingleItemRow.cpp b/Engine/Source/Editor/PropertyEditor/Private/SDetailSingleItemRow.cpp
--- a/Engine/Source/Editor/PropertyEditor/Private/SDetailSingleItemRow.cpp
+++ b/Engine/Source/Editor/PropertyEditor/Private/SDetailSingleItemRow.cpp
@@ -1252,71 +1252,105 @@ void SDetailSingleItemRow::OnCopyProperty()
 
 void SDetailSingleItemRow::OnCopyPropertyDisplayName()
 {
-	if (OwnerTreeNode.IsValid())
+	if (!OwnerTreeNode.IsValid())
 	{
-		TSharedPtr<FPropertyNode> PropertyNode = GetPropertyNode();
-		if (PropertyNode.IsValid())
-		{
-			FPropertyEditorClipboard::ClipboardCopy(*PropertyNode->GetDisplayName().ToString());
-		}
+		return;
+	}
+
+	TSharedPtr<FPropertyNode> PropertyNode = GetPropertyNode();
+	if (!PropertyNode.IsValid())
+	{
+		return;
+	}
+
+	if (PropertyNode->IsOptionalValueNode())
+	{
+		FPropertyEditorClipboard::ClipboardCopy(*PropertyNode->GetParentNode()->GetDisplayName().ToString());
+	}
+	else
+	{
+		FPropertyEditorClipboard::ClipboardCopy(*PropertyNode->GetDisplayName().ToString());
 	}
 }
 
 bool SDetailSingleItemRow::CanCopyPropertyDisplayName()
 {
-	if (OwnerTreeNode.IsValid())
+	if (!OwnerTreeNode.IsValid())
 	{
-		TSharedPtr<FPropertyNode> PropertyNode = GetPropertyNode();
-		if (PropertyNode.IsValid())
-		{
-			if (!PropertyNode->GetDisplayName().IsEmpty())
-			{
-				return true;
-			}
-				
-		}
+		return false;
 	}
-	return false;
+
+	TSharedPtr<FPropertyNode> PropertyNode = GetPropertyNode();
+	if (!PropertyNode.IsValid())
+	{
+		return false;
+	}
+
+	if (PropertyNode->IsOptionalValueNode() && PropertyNode->GetParentNode()->GetDisplayName().IsEmpty())
+	{
+		return false;
+	}
+	else if (PropertyNode->GetDisplayName().IsEmpty())
+	{
+		return false;
+	}
+
+	return true;
 }
 
 void SDetailSingleItemRow::OnCopyPropertyInternalName()
 {
-	if (OwnerTreeNode.IsValid())
+	if (!OwnerTreeNode.IsValid())
 	{
-		TSharedPtr<FPropertyNode> PropertyNode = GetPropertyNode();
-		if (PropertyNode.IsValid())
-		{
-			if (const FProperty* Property = PropertyNode->GetProperty())
-			{
-				if (const UStruct* OwnerStruct = Property->GetOwnerStruct())
-				{
-					FPropertyEditorClipboard::ClipboardCopy(*OwnerStruct->GetAuthoredNameForField(Property));
-				}
-			}
-		}
+		return;
+	}
+
+	TSharedPtr<FPropertyNode> PropertyNode = GetPropertyNode();
+	if (!PropertyNode.IsValid())
+	{
+		return;
+	}
+
+	const FProperty* Property = PropertyNode->IsOptionalValueNode() ? PropertyNode->GetParentNode()->GetProperty() : PropertyNode->GetProperty();
+	if (!Property)
+	{
+		return;
+	}
+
+	if (const UStruct* OwnerStruct = Property->GetOwnerStruct())
+	{
+		FPropertyEditorClipboard::ClipboardCopy(*OwnerStruct->GetAuthoredNameForField(Property));
 	}
 }
 
 bool SDetailSingleItemRow::CanCopyPropertyInternalName()
 {
-	if (OwnerTreeNode.IsValid())
+	if (!OwnerTreeNode.IsValid())
 	{
-		TSharedPtr<FPropertyNode> PropertyNode = GetPropertyNode();
-		if (PropertyNode.IsValid())
+		return false;
+	}
+
+	TSharedPtr<FPropertyNode> PropertyNode = GetPropertyNode();
+	if (!PropertyNode.IsValid())
+	{
+		return false;
+	}
+
+	const FProperty* Property = PropertyNode->IsOptionalValueNode() ? PropertyNode->GetParentNode()->GetProperty() : PropertyNode->GetProperty();
+	if (!Property)
+	{
+		return false;
+	}
+
+	if (const UStruct* OwnerStruct = Property->GetOwnerStruct())
+	{
+		if (OwnerStruct->GetAuthoredNameForField(Property).IsEmpty())
 		{
-			if (const FProperty* Property = PropertyNode->GetProperty())
-			{
-				if (const UStruct* OwnerStruct = Property->GetOwnerStruct())
-				{
-					if (!OwnerStruct->GetAuthoredNameForField(Property).IsEmpty())
-					{
-						return true;
-					}
-				}
-			}
+			return false;
 		}
 	}
-	return false;
+
+	return true;
 }
 
 bool SDetailSingleItemRow::CanPasteProperty() const
diff --git a/Engine/Source/Editor/PropertyEditor/Private/UserInterface/PropertyEditor/SPropertyEditorOptional.cpp b/Engine/Source/Editor/PropertyEditor/Private/UserInterface/PropertyEditor/SPropertyEditorOptional.cpp
--- a/Engine/Source/Editor/PropertyEditor/Private/UserInterface/PropertyEditor/SPropertyEditorOptional.cpp
+++ b/Engine/Source/Editor/PropertyEditor/Private/UserInterface/PropertyEditor/SPropertyEditorOptional.cpp
@@ -6,6 +6,7 @@
 #include "Internationalization/Internationalization.h"
 #include "Layout/Children.h"
 #include "Presentation/PropertyEditor/PropertyEditor.h"
+#include "PropertyEditorModule.h"
 #include "PropertyNode.h"
 #include "UObject/PropertyOptional.h"
 #include "Widgets/Input/SButton.h"
@@ -40,20 +41,6 @@ void SPropertyEditorOptional::Construct(const FArguments& InArgs, const TSharedR
 					.Text(*MultiselectOptions[0])
 			];
 	}
-	else if (const TSharedPtr<FPropertyNode> OptionalValueNode = PropertyNode->GetOrCreateOptionalValueNode())
-	{
-		// create the OptionalValue's PropertyEditor
-		TSharedRef<FPropertyEditor> ValuePropertyEditor = FPropertyEditor::Create(OptionalValueNode.ToSharedRef(), InPropertyUtilities);
-		
-		TSharedPtr<SPropertyValueWidget> Widget = SNew(SPropertyValueWidget, ValuePropertyEditor, InPropertyUtilities)
-			.IsEnabled(InArgs._IsEnabled)
-			.ShowPropertyButtons(true);
-
-		 MinDesiredWidth = Widget->GetMinDesiredWidth();
-		 MaxDesiredWidth = Widget->GetMaxDesiredWidth();
-
-		 ValueEditorWidget = Widget;
-	}
 	else
 	{
 		ValueEditorWidget = PropertyEditorHelpers::MakePropertyButton(EPropertyButton::OptionalSet, InPropertyEditor);
