Patch fix for a consistent crash from Details panel code when selecting actors in level. Author of related code contacted to discuss.

diff --git a/Engine/Source/Editor/PropertyEditor/Private/DetailPropertyRow.cpp b/Engine/Source/Editor/PropertyEditor/Private/DetailPropertyRow.cpp
--- a/Engine/Source/Editor/PropertyEditor/Private/DetailPropertyRow.cpp
+++ b/Engine/Source/Editor/PropertyEditor/Private/DetailPropertyRow.cpp
@@ -29,7 +29,8 @@ FDetailPropertyRow::FDetailPropertyRow(TSharedPtr<FPropertyNode> InPropertyNode,
 	, bCachedCustomTypeInterface(false)
 {
 	// Is this a set optional property?
-	if (TSharedPtr<FPropertyNode>& ValueNode = PropertyNode->GetOptionalValueNode())
+	TSharedPtr<FPropertyNode> ValueNode = PropertyNode.IsValid() ? PropertyNode->GetOptionalValueNode() : nullptr;
+	if (ValueNode.IsValid())
 	{
 		// If we are selecting multiple options with different states (ie set/unset) dont use the
 		// Value node so optional multi-select logic is displayed (see SPropertyEditorOptional.h).
