Overrides support UI in the details panel plugged into the FOverridesManager API.

1. Updated the Category menu to include the override component menu item
2. Updated the Overrides combobutton to take in the component and show the proper widget whether overridden fully or inside. For now it will have no associated menu items
3. Updated DetailsDisplayManager and EntityDetailsDisplayManager to have the property update widgets take in the name of the object (component in this case) that will be watched for override changes to show the status
4. Updated the calls to those in the Detail Category Table Rows to pass in the category object name

diff --git a/Engine/Source/Editor/PropertyEditor/Private/DetailsDisplayManager.cpp b/Engine/Source/Editor/PropertyEditor/Private/DetailsDisplayManager.cpp
--- a/Engine/Source/Editor/PropertyEditor/Private/DetailsDisplayManager.cpp
+++ b/Engine/Source/Editor/PropertyEditor/Private/DetailsDisplayManager.cpp
@@ -151,7 +151,7 @@ bool FDetailsDisplayManager::AddEmptyCategoryToDetailLayoutIfNeeded(TSharedRef<F
 }
 
 TSharedPtr<FPropertyUpdatedWidgetBuilder> FDetailsDisplayManager::GetPropertyUpdatedWidget(FResetToDefault ResetToDefault,
-	bool bIsCategoryUpdateWidget)
+	bool bIsCategoryUpdateWidget, FName InCategoryObjectName)
 {
 	return nullptr;
 }
diff --git a/Engine/Source/Editor/PropertyEditor/Private/SDetailCategoryTableRow.cpp b/Engine/Source/Editor/PropertyEditor/Private/SDetailCategoryTableRow.cpp
--- a/Engine/Source/Editor/PropertyEditor/Private/SDetailCategoryTableRow.cpp
+++ b/Engine/Source/Editor/PropertyEditor/Private/SDetailCategoryTableRow.cpp
@@ -108,7 +108,7 @@ void SDetailCategoryTableRow::Construct(const FArguments& InArgs, TSharedRef<FDe
 	PropertyUpdatedWidgetBuilder = DisplayManager->GetPropertyUpdatedWidget(FExecuteAction::CreateLambda( [this]
 		{
 			UE_LOG(LogTemp, Warning, TEXT("DisplayManager->GetPropertyUpdatedWidget reset"));
-		}), true);
+		}), true, ObjectName );
 	if (PropertyUpdatedWidgetBuilder.IsValid())
 	{
 		TAttribute<bool> IsHovered = TAttribute<bool>::CreateSP( this, &SDetailCategoryTableRow::IsHovered);
@@ -157,35 +157,27 @@ void SDetailCategoryTableRow::Construct(const FArguments& InArgs, TSharedRef<FDe
 			.VAlign(VAlign_Center)
 			.AutoWidth()
 			[
+			PropertyUpdatedWidgetBuilder.IsValid() ?
+								PropertyUpdatedWidgetBuilder->Bind_IsVisible(TAttribute<EVisibility>::CreateLambda([this]
+									{
+										return EVisibility::Visible;
+									})).GenerateWidget().ToSharedRef() :
+									SNullWidget::NullWidget
+			]
+			+SHorizontalBox::Slot()
+			.HAlign(HAlign_Right)
+			.VAlign(VAlign_Center)
+			.AutoWidth()
+			[
 			DisplayManager.IsValid() ?
-			*FCategoryMenuComboButtonBuilder( DisplayManager.ToSharedRef() )
+				*FCategoryMenuComboButtonBuilder( DisplayManager.ToSharedRef() )
 				.Set_OnGetContent(FOnGetContent::CreateLambda([this]
 				{
 					const TSharedPtr<SWidget> Menu = DisplayManager->GetCategoryMenu(ObjectName);
 					return Menu.IsValid() ? Menu.ToSharedRef() : SNullWidget::NullWidget;
-				}))
-				.Bind_IsVisible(TAttribute<EVisibility>::CreateLambda([this]
-				{
-					return IsHovered() ? EVisibility::Visible : EVisibility::Collapsed;
 				})) :
 				SNullWidget::NullWidget
 			]
-			+SHorizontalBox::Slot()
-			.HAlign(HAlign_Right)
-			.VAlign(VAlign_Center)
-			.AutoWidth()
-			[
-				PropertyUpdatedWidgetBuilder.IsValid() ?
-						PropertyUpdatedWidgetBuilder->Bind_IsVisible(TAttribute<EVisibility>::CreateLambda([this]
-							{
-								if (!DisplayManager->GetCategoryHasAnyUpdatedProperties(ObjectName))
-								{
-									return EVisibility::Collapsed;
-								}
-								return EVisibility::Visible;
-							})).GenerateWidget().ToSharedRef() :
-							SNullWidget::NullWidget
-			]
 			]
 			] 
 		]
diff --git a/Engine/Source/Editor/PropertyEditor/Private/UserInterface/Categories/CategoryMenuComboButtonBuilder.cpp b/Engine/Source/Editor/PropertyEditor/Private/UserInterface/Categories/CategoryMenuComboButtonBuilder.cpp
--- a/Engine/Source/Editor/PropertyEditor/Private/UserInterface/Categories/CategoryMenuComboButtonBuilder.cpp
+++ b/Engine/Source/Editor/PropertyEditor/Private/UserInterface/Categories/CategoryMenuComboButtonBuilder.cpp
@@ -27,7 +27,7 @@ TSharedPtr<SWidget> FCategoryMenuComboButtonBuilder::GenerateWidget()
 	return 	DisplayManager->ShouldShowCategoryMenu()  ?
 		        SNew( SComboButton )
 				.ComboButtonStyle( FAppStyle::Get(), "DetailsView.CategoryComboButton" )
-				.Visibility(IsVisible)
+				.Visibility(EVisibility::Visible)
 				.HasDownArrow(true)
 				.OnGetMenuContent(OnGetContent) :
 		        SNullWidget::NullWidget;
diff --git a/Engine/Source/Editor/PropertyEditor/Private/UserInterface/Widgets/OverridesComboButtonBuilder.cpp b/Engine/Source/Editor/PropertyEditor/Private/UserInterface/Widgets/OverridesComboButtonBuilder.cpp
--- a/Engine/Source/Editor/PropertyEditor/Private/UserInterface/Widgets/OverridesComboButtonBuilder.cpp
+++ b/Engine/Source/Editor/PropertyEditor/Private/UserInterface/Widgets/OverridesComboButtonBuilder.cpp
@@ -5,13 +5,16 @@
 #include "DetailsViewStyle.h"
 #include "Styling/SlateBrush.h"
 #include "Brushes/SlateRoundedBoxBrush.h"
+#include "UObject/OverridableManager.h"
 
 FOverridesComboButtonBuilder::FOverridesComboButtonBuilder(
 	TSharedRef<FDetailsDisplayManager> InDetailsDisplayManager,
-	bool bInIsCategoryOverridesComboButton):
+	bool bInIsCategoryOverridesComboButton,
+	TWeakObjectPtr<UObject> InObject ):
 	FPropertyUpdatedWidgetBuilder(),
 	DisplayManager(InDetailsDisplayManager),
-	bIsCategoryOverridesComboButton(bInIsCategoryOverridesComboButton)
+	bIsCategoryOverridesComboButton(bInIsCategoryOverridesComboButton),
+    Object(InObject)
 {
 }
 
@@ -36,17 +39,9 @@ TSharedPtr<SWidget> FOverridesComboButtonBuilder::GenerateWidget()
 				[
 					SNew( SComboButton )
 					.ComboButtonStyle( &DetailsViewStyle->GetOverridesComboButtonStyle(
-										GetOverridesStyleKey(false ), bIsCategoryOverridesComboButton ) )
-					.Visibility_Lambda([this]
-					{
-						if ( IsRowHoveredAttr.IsBound() && IsRowHoveredAttr.Get())
-						{
-							return EVisibility::Collapsed;
-						}
-						return EVisibility::Visible;
-					})
+										&FOverridesWidgetStyleKeys::Here(), bIsCategoryOverridesComboButton ) )
+					.Visibility(this, &FOverridesComboButtonBuilder::GetFullyOverridenVisibility)
 					.HasDownArrow(true)
-					.OnGetMenuContent(OnGetContent)
 				]
 				+SHorizontalBox::Slot()
 				.HAlign(HAlign_Right)
@@ -54,17 +49,11 @@ TSharedPtr<SWidget> FOverridesComboButtonBuilder::GenerateWidget()
 				.AutoWidth()
 				[
 					SNew( SComboButton )
-					.ComboButtonStyle(  &DetailsViewStyle->GetOverridesComboButtonStyle( GetOverridesStyleKey(true), bIsCategoryOverridesComboButton ))
-					.Visibility_Lambda([this]
-					{
-						if (IsRowHoveredAttr.IsBound() && IsRowHoveredAttr.Get())
-						{
-							return EVisibility::Visible;
-						}
-						return EVisibility::Collapsed;
-					})
-					.HasDownArrow(true)
-					.OnGetMenuContent(OnGetContent)
+					.ComboButtonStyle( &DetailsViewStyle->GetOverridesComboButtonStyle(
+					                                    &FOverridesWidgetStyleKeys::Inside()
+                										, bIsCategoryOverridesComboButton ) )
+				                    .Visibility(this, &FOverridesComboButtonBuilder::GetOverridenInsideVisibility)
+                					.HasDownArrow(true)
 				];
 }
 
@@ -73,21 +62,36 @@ TSharedRef<SWidget> FOverridesComboButtonBuilder::operator*()
 	return GenerateWidget().ToSharedRef();
 }
 
-const FOverridesWidgetStyleKey* FOverridesComboButtonBuilder::GetOverridesStyleKey(bool bIsHoveredOver) const
+FOverridesComboButtonBuilder::~FOverridesComboButtonBuilder()
 {
-	if (bIsCategoryOverridesComboButton && !bIsHoveredOver)
-	{
-		return &FOverridesWidgetStyleKeys::Here();
-	} 
-	if (!bIsHoveredOver)
+	OnGetContent.Unbind();
+}
+
+EVisibility FOverridesComboButtonBuilder::GetOverridenInsideVisibility() const
+{
+	static FOverridableManager& Manager = FOverridableManager::Get();
+	
+	if ( UObject* OverridableComponent = Object.Get() )
 	{
-		return &FOverridesWidgetStyleKeys::Here();
+		const EOverriddenState State = Manager.GetOverriddenState(*OverridableComponent);
+		if (State != EOverriddenState::NoOverrides && State != EOverriddenState::AllOverridden)
+		{
+			return EVisibility::Visible;
+		}		
 	}
-	
-	return &FOverridesWidgetStyleKeys::Options();
+	return EVisibility::Collapsed;
 }
 
-FOverridesComboButtonBuilder::~FOverridesComboButtonBuilder()
+EVisibility FOverridesComboButtonBuilder::GetFullyOverridenVisibility() const
 {
-	OnGetContent.Unbind();
-}
+	if ( UObject* OverridableObject = Object.Get() )
+	{
+		static FOverridableManager& Manager = FOverridableManager::Get();
+		const EOverriddenState State = Manager.GetOverriddenState(*OverridableObject);
+		if (State == EOverriddenState::AllOverridden)
+		{
+			return EVisibility::Visible;
+		}
+	}
+	return EVisibility::Collapsed;
+}
\ No newline at end of file
diff --git a/Engine/Source/Editor/PropertyEditor/Public/DetailsDisplayManager.h b/Engine/Source/Editor/PropertyEditor/Public/DetailsDisplayManager.h
--- a/Engine/Source/Editor/PropertyEditor/Public/DetailsDisplayManager.h
+++ b/Engine/Source/Editor/PropertyEditor/Public/DetailsDisplayManager.h
@@ -107,7 +107,7 @@ public:
 	 */	
 	PROPERTYEDITOR_API virtual bool AddEmptyCategoryToDetailLayoutIfNeeded(TSharedRef<FComplexPropertyNode> Node, TSharedRef<FDetailLayoutBuilderImpl> DetailLayoutBuilder);
 
-	virtual TSharedPtr<FPropertyUpdatedWidgetBuilder> GetPropertyUpdatedWidget(FResetToDefault ResetToDefault, bool bIsCategoryUpdateWidget = false);
+	virtual TSharedPtr<FPropertyUpdatedWidgetBuilder> GetPropertyUpdatedWidget(FResetToDefault ResetToDefault, bool bIsCategoryUpdateWidget = false, FName InCategoryObjectName = NAME_None);
 
 	void UpdatePropertyForCategory(FName InCategoryObjectName, FProperty* Property, bool bAddProperty);
 
diff --git a/Engine/Source/Editor/PropertyEditor/Public/UserInterface/Widgets/OverridesComboButtonBuilder.h b/Engine/Source/Editor/PropertyEditor/Public/UserInterface/Widgets/OverridesComboButtonBuilder.h
--- a/Engine/Source/Editor/PropertyEditor/Public/UserInterface/Widgets/OverridesComboButtonBuilder.h
+++ b/Engine/Source/Editor/PropertyEditor/Public/UserInterface/Widgets/OverridesComboButtonBuilder.h
@@ -8,7 +8,7 @@
 #include "Templates/SharedPointer.h"
 
 class FDetailsDisplayManager;
-class FOverridesWidgetStyleKey;
+
 
 DECLARE_DELEGATE_OneParam(FGetIsRowHoveredOver, bool)
 
@@ -34,7 +34,8 @@ public:
 	 */
 	PROPERTYEDITOR_API FOverridesComboButtonBuilder(
 		TSharedRef<FDetailsDisplayManager> InDetailsDisplayManager,
-		bool bInIsCategoryOverridesComboButton);
+		bool bInIsCategoryOverridesComboButton,
+		TWeakObjectPtr<UObject> InObject );
 
 	/**
 	 * Set the OnGetContent for the menu that this button is responsible for
@@ -56,12 +57,17 @@ public:
 	 */
 	TSharedRef<SWidget> operator*();
 
+private:
 	/**
-	 * returns the const FOverridesWidgetStyleKey* which points to the proper FOverridesWidgetStyleKey for this
+	 * Returns the visibility of the image that denotes that a Component is fully overridden
 	 */
-	const FOverridesWidgetStyleKey* GetOverridesStyleKey(bool bIsHoveredOver = false) const;
+	EVisibility GetFullyOverridenVisibility() const;
+
+	/**
+	 * Returns the visibility of the image that denotes a Component is overridden only inside
+	 */
+	EVisibility GetOverridenInsideVisibility() const;
 
-private:
 	/**
 	 * The @code DetailsDisplayManager @endcode which provides an API to manage some of the characteristics of the
 	 * details display
@@ -72,4 +78,9 @@ private:
 	 * if true, this is an overrides combo button for a Category 
 	 */
 	bool bIsCategoryOverridesComboButton;
+
+	/**
+	 * The UObject that will be queried for its override state
+	 */
+	TWeakObjectPtr<UObject> Object;
 };
