[SceneGraph] Workflow tests: tagged DetailsView.CategoryComboButton and fixed tests after recent changes in DetailsView

diff --git a/Engine/Source/Editor/PropertyEditor/Private/UserInterface/Categories/CategoryMenuComboButtonBuilder.cpp b/Engine/Source/Editor/PropertyEditor/Private/UserInterface/Categories/CategoryMenuComboButtonBuilder.cpp
--- a/Engine/Source/Editor/PropertyEditor/Private/UserInterface/Categories/CategoryMenuComboButtonBuilder.cpp
+++ b/Engine/Source/Editor/PropertyEditor/Private/UserInterface/Categories/CategoryMenuComboButtonBuilder.cpp
@@ -29,6 +29,7 @@ TSharedPtr<SWidget> FCategoryMenuComboButtonBuilder::GenerateWidget()
 				.ComboButtonStyle( FAppStyle::Get(), "DetailsView.CategoryComboButton" )
 				.Visibility(EVisibility::Visible)
 				.HasDownArrow(true)
+				.AddMetaData<FTagMetaData>(FTagMetaData(TEXT("DetailsView.CategoryComboButton")))
 				.OnGetMenuContent(OnGetContent) :
 		        SNullWidget::NullWidget;
 }
