[CrashFix] Missing nullptr check in SDetailsSplitter.cpp CanCopyPropertyValue

diff --git a/Engine/Source/Editor/PropertyEditor/Private/SDetailsSplitter.cpp b/Engine/Source/Editor/PropertyEditor/Private/SDetailsSplitter.cpp
--- a/Engine/Source/Editor/PropertyEditor/Private/SDetailsSplitter.cpp
+++ b/Engine/Source/Editor/PropertyEditor/Private/SDetailsSplitter.cpp
@@ -364,7 +364,7 @@ namespace DetailsSplitterHelpers
 		{
 			const TSharedPtr<IPropertyHandle> SourceHandle = SourceDetailsNode->CreatePropertyHandle();
 			const TSharedPtr<IPropertyHandle> DestinationHandle = DestinationDetailsNode->CreatePropertyHandle();
-			if (SourceHandle && DestinationHandle)
+			if (SourceHandle && DestinationHandle && SourceHandle->GetProperty() && DestinationHandle->GetProperty())
 			{
 				TArray<void*> SourceData;
 				TArray<void*> DestinationData;
