Add GetValidEnumValues to be able to filter enums dynamically

#rnx

diff --git a/Engine/Source/Editor/PropertyEditor/Private/PropertyEditorHelpers.cpp b/Engine/Source/Editor/PropertyEditor/Private/PropertyEditorHelpers.cpp
--- a/Engine/Source/Editor/PropertyEditor/Private/PropertyEditorHelpers.cpp
+++ b/Engine/Source/Editor/PropertyEditor/Private/PropertyEditorHelpers.cpp
@@ -1129,23 +1129,66 @@ namespace PropertyEditorHelpers
 		
 	}
 
-	TArray<FName> GetValidEnumsFromPropertyOverride(const FProperty* Property, const UEnum* InEnum)
+	TArray<FName> GetValidEnumsFromPropertyOverride(const TArray<UObject*>& ObjectList, const FProperty* Property, const UEnum* InEnum)
 	{
 		TArray<FName> ValidEnumValues;
+		
+		bool bMergeValidEnums = false;
 
 		static const FName ValidEnumValuesName("ValidEnumValues");
+		static const FName GetValidEnumValuesName("GetValidEnumValues");
+		TArray<FString> ValidEnumValuesAsString;
+		
 		if(Property->HasMetaData(ValidEnumValuesName))
 		{
-			TArray<FString> ValidEnumValuesAsString;
-
+			bMergeValidEnums = true;
 			Property->GetMetaData(ValidEnumValuesName).ParseIntoArray(ValidEnumValuesAsString, TEXT(","));
-			for(FString& Value : ValidEnumValuesAsString)
+		}
+
+		if(Property->HasMetaData(GetValidEnumValuesName))
+		{
+			const FString GetValidEnumValuesNameFunctionName = Property->GetMetaData(GetValidEnumValuesName);
+			if (!GetValidEnumValuesNameFunctionName.IsEmpty())
 			{
-				Value.TrimStartInline();
-				ValidEnumValues.Add(*InEnum->GenerateFullEnumName(*Value));
+				for (UObject* Object : ObjectList)
+				{
+					const UFunction* GetValidEnumValuesNameFunction = Object ? Object->FindFunction(*GetValidEnumValuesNameFunctionName) : nullptr;
+					if (GetValidEnumValuesNameFunction)
+					{
+						DECLARE_DELEGATE_RetVal(TArray<FString>, FGetValidEnumValuesNameF);
+						if (!bMergeValidEnums)
+						{
+							ValidEnumValuesAsString = FGetValidEnumValuesNameF::CreateUFunction(Object, GetValidEnumValuesNameFunction->GetFName()).Execute();
+							bMergeValidEnums = true;
+						}
+						else
+						{
+							TArray<FString> ValidEnumValuesAsStringToMerge = FGetValidEnumValuesNameF::CreateUFunction(Object, GetValidEnumValuesNameFunction->GetFName()).Execute();
+							TArray<FString> OldValidEnumValuesAsString = MoveTemp(ValidEnumValuesAsString);
+							for (const FString& OldValidEnumValue : OldValidEnumValuesAsString)
+							{
+								if (ValidEnumValuesAsStringToMerge.Contains(OldValidEnumValue))
+								{
+									ValidEnumValuesAsString.Add(OldValidEnumValue);
+								}
+							}
+						}
+						if (ValidEnumValuesAsString.IsEmpty())
+						{
+							ValidEnumValues.Add(NAME_None);
+							return ValidEnumValues;
+						}
+					}
+				}
 			}
 		}
-
+		
+		for(FString& Value : ValidEnumValuesAsString)
+		{
+			Value.TrimStartInline();
+			ValidEnumValues.Add(*InEnum->GenerateFullEnumName(*Value));
+		}
+			
 		return ValidEnumValues;
 	}
 
diff --git a/Engine/Source/Editor/PropertyEditor/Private/PropertyEditorHelpers.h b/Engine/Source/Editor/PropertyEditor/Private/PropertyEditorHelpers.h
--- a/Engine/Source/Editor/PropertyEditor/Private/PropertyEditorHelpers.h
+++ b/Engine/Source/Editor/PropertyEditor/Private/PropertyEditorHelpers.h
@@ -256,11 +256,12 @@ namespace PropertyEditorHelpers
 	/**
 	 * Returns any enums that are explicitly allowed by the "ValidEnumValues" metadata on FProperty using the specified enum.
 	 *
+	 * @param ObjectList The list of objects currently edited
 	 * @param Property	The property which may contain the "ValidEnumValues" metadata
 	 * @param InEnum	The enum to search
 	 * @return The array of allowed enums.  NOTE: If an empty array is returned all enum values are allowed. It is an error for a property to hide all enum values so that state is undefined here.
 	 */
-	TArray<FName> GetValidEnumsFromPropertyOverride(const FProperty* Property, const UEnum* InEnum);
+	TArray<FName> GetValidEnumsFromPropertyOverride(const TArray<UObject*>& ObjectList, const FProperty* Property, const UEnum* InEnum);
 
 	/**
 	 * Returns any enums that are explicitly Disallowed by the "InvalidEnumValues" metadata on FProperty using the specified enum.
diff --git a/Engine/Source/Editor/PropertyEditor/Private/PropertyHandleImpl.cpp b/Engine/Source/Editor/PropertyEditor/Private/PropertyHandleImpl.cpp
--- a/Engine/Source/Editor/PropertyEditor/Private/PropertyHandleImpl.cpp
+++ b/Engine/Source/Editor/PropertyEditor/Private/PropertyHandleImpl.cpp
@@ -3079,7 +3079,10 @@ bool FPropertyHandleBase::GeneratePossibleValues(TArray< TSharedPtr<FString> >&
 
 	if( Enum )
 	{
-		const TArray<FName> ValidEnumValues = PropertyEditorHelpers::GetValidEnumsFromPropertyOverride(Property, Enum);
+		TArray<UObject*> OuterObjects;
+		GetOuterObjects(OuterObjects);		
+		
+		const TArray<FName> ValidEnumValues = PropertyEditorHelpers::GetValidEnumsFromPropertyOverride(OuterObjects, Property, Enum);
 		const TArray<FName> InvalidEnumValues = PropertyEditorHelpers::GetInvalidEnumsFromPropertyOverride(Property, Enum);
 		const TMap<FName, FText> EnumValueDisplayNameOverrides = PropertyEditorHelpers::GetEnumValueDisplayNamesFromPropertyOverride(Property, Enum);
 
@@ -3148,10 +3151,10 @@ bool FPropertyHandleBase::GeneratePossibleValues(TArray< TSharedPtr<FString> >&
 		FString GetOptionsFunctionName = Property->GetOwnerProperty()->GetMetaData(MetaDataKey);
 		if (!GetOptionsFunctionName.IsEmpty())
 		{
-			TArray<UObject*> OutObjects;
-			GetOuterObjects(OutObjects);
+			TArray<UObject*> OuterObjects;
+			GetOuterObjects(OuterObjects);
 
-			PropertyEditorUtils::GetPropertyOptions(OutObjects, GetOptionsFunctionName, OutOptionStrings);
+			PropertyEditorUtils::GetPropertyOptions(OuterObjects, GetOptionsFunctionName, OutOptionStrings);
 		}
 	}
 	else if( Property->IsA(FClassProperty::StaticClass()) || Property->IsA(FSoftClassProperty::StaticClass()) )		
diff --git a/Engine/Source/Editor/PropertyEditor/Private/UserInterface/PropertyEditor/SPropertyEditorCombo.cpp b/Engine/Source/Editor/PropertyEditor/Private/UserInterface/PropertyEditor/SPropertyEditorCombo.cpp
--- a/Engine/Source/Editor/PropertyEditor/Private/UserInterface/PropertyEditor/SPropertyEditorCombo.cpp
+++ b/Engine/Source/Editor/PropertyEditor/Private/UserInterface/PropertyEditor/SPropertyEditorCombo.cpp
@@ -260,7 +260,10 @@ void SPropertyEditorCombo::GenerateComboBoxStrings( TArray< TSharedPtr<FString>
 
 			if(Enum)
 			{
-				TArray<FName> ValidPropertyEnums = PropertyEditorHelpers::GetValidEnumsFromPropertyOverride(Property, Enum);
+				TArray<UObject*> ObjectList;
+				ComboArgs.PropertyHandle->GetOuterObjects(ObjectList);
+				
+				TArray<FName> ValidPropertyEnums = PropertyEditorHelpers::GetValidEnumsFromPropertyOverride(ObjectList, Property, Enum);
 				TArray<FName> InvalidPropertyEnums = PropertyEditorHelpers::GetInvalidEnumsFromPropertyOverride(Property, Enum);
 				
 				// Get enum doc link (not just GetDocumentationLink as that is the documentation for the struct we're in, not the enum documentation)
diff --git a/Engine/Source/Editor/PropertyEditor/Private/UserInterface/PropertyEditor/SPropertyEditorNumeric.h b/Engine/Source/Editor/PropertyEditor/Private/UserInterface/PropertyEditor/SPropertyEditorNumeric.h
--- a/Engine/Source/Editor/PropertyEditor/Private/UserInterface/PropertyEditor/SPropertyEditorNumeric.h
+++ b/Engine/Source/Editor/PropertyEditor/Private/UserInterface/PropertyEditor/SPropertyEditorNumeric.h
@@ -62,7 +62,7 @@ public:
 		{
 			auto CreateBitmaskFlagsArray = [PropertyHandle]()
 			{
-				const int32 BitmaskBitCount = sizeof(NumericType) << 3;
+				constexpr int32 BitmaskBitCount = sizeof(NumericType) << 3;
 
 				TArray<FBitmaskFlagInfo> Result;
 				Result.Empty(BitmaskBitCount);
@@ -79,7 +79,7 @@ public:
 				{
 					const TMap<FName, FText> EnumValueDisplayNameOverrides = PropertyEditorHelpers::GetEnumValueDisplayNamesFromPropertyOverride(PropertyHandle->GetProperty(), BitmaskEnum);
 
-					const bool bUseEnumValuessAsMaskValues = BitmaskEnum->GetBoolMetaData(PropertyEditorConstants::MD_UseEnumValuesAsMaskValuesInEditor);
+					const bool bUseEnumValuesAsMaskValues = BitmaskEnum->GetBoolMetaData(PropertyEditorConstants::MD_UseEnumValuesAsMaskValuesInEditor);
 					auto AddNewBitmaskFlagLambda = [BitmaskEnum, &Result, &EnumValueDisplayNameOverrides](int32 InEnumIndex, int64 InFlagValue)
 					{
 						Result.Emplace();
@@ -98,7 +98,10 @@ public:
 						}
 					};
 
-					const TArray<FName> AllowedPropertyEnums = PropertyEditorHelpers::GetValidEnumsFromPropertyOverride(PropertyHandle->GetProperty(), BitmaskEnum);
+					TArray<UObject*> OuterObjects;
+					PropertyHandle->GetOuterObjects(OuterObjects);
+					
+					const TArray<FName> AllowedPropertyEnums = PropertyEditorHelpers::GetValidEnumsFromPropertyOverride(OuterObjects, PropertyHandle->GetProperty(), BitmaskEnum);
 					const TArray<FName> DisallowedPropertyEnums = PropertyEditorHelpers::GetInvalidEnumsFromPropertyOverride(PropertyHandle->GetProperty(), BitmaskEnum);
 					// Note: This loop doesn't include (BitflagsEnum->NumEnums() - 1) in order to skip the implicit "MAX" value that gets added to the enum type at compile time.
 					for (int32 BitmaskEnumIndex = 0; BitmaskEnumIndex < BitmaskEnum->NumEnums() - 1; ++BitmaskEnumIndex)
@@ -119,7 +122,7 @@ public:
 						}
 						if (EnumValue >= 0 && !bShouldBeHidden)
 						{
-							if (bUseEnumValuessAsMaskValues)
+							if (bUseEnumValuesAsMaskValues)
 							{
 								if (EnumValue < MAX_int64 && FMath::IsPowerOfTwo(EnumValue))
 								{
diff --git a/Engine/Source/Runtime/CoreUObject/Public/UObject/ObjectMacros.h b/Engine/Source/Runtime/CoreUObject/Public/UObject/ObjectMacros.h
--- a/Engine/Source/Runtime/CoreUObject/Public/UObject/ObjectMacros.h
+++ b/Engine/Source/Runtime/CoreUObject/Public/UObject/ObjectMacros.h
@@ -1490,6 +1490,15 @@ namespace UM
 
 		/// [PropertyMetadata] Used for enum properties to define the subset of valid values as a comma-separated string; values outside of this subset are not made available.
 		ValidEnumValues,
+
+		/// [PropertyMetadata] Used for enum properties to define the subset of invalid values as a comma-separated string; values within this subset are not made available.
+		InvalidEnumValues,
+
+		/// [PropertyMetadata] Used for enum properties to define the subset of valid values through a UFunction, e.g. meta=(GetValidEnumValues="FuncName"). 
+		///
+		/// UFUNCTION()
+		/// TArray<FString> FuncName() const;
+		GetValidEnumValues,
 		
 		/// [PropertyMetadata] Causes assets to be filtered through a UFunction, e.g. meta=(GetAssetFilter="FuncName"). The UFunction should return true to exclude the asset.
 		///
@@ -1507,10 +1516,7 @@ namespace UM
 		///
 		/// UFUNCTION()
 		/// TArray<UClass*> FuncName() const;
-		GetDisallowedClasses,
-
-	    /// [PropertyMetadata] Used for enum properties to define the subset of invalid values as a comma-separated string; values within this subset are not made available.
-        InvalidEnumValues,
+		GetDisallowedClasses
 	};
 
 	// Metadata usable in UPROPERTY for customizing the behavior of Persona and UMG
