Replace GetEnabledEnumValues by GetRestrictedEnumValues
#tests PIE

diff --git a/Engine/Source/Editor/PropertyEditor/Private/PropertyEditorHelpers.cpp b/Engine/Source/Editor/PropertyEditor/Private/PropertyEditorHelpers.cpp
--- a/Engine/Source/Editor/PropertyEditor/Private/PropertyEditorHelpers.cpp
+++ b/Engine/Source/Editor/PropertyEditor/Private/PropertyEditorHelpers.cpp
@@ -128,7 +128,6 @@ TSharedRef<SWidget> SPropertyValueWidget::ConstructPropertyEditorWidget( TShared
 	//const TSharedRef<IPropertyUtilities> PropertyUtilitiesRef = InPropertyUtilities.ToSharedRef();
 
 	const TSharedRef< FPropertyNode > PropertyNode = PropertyEditorRef->GetPropertyNode();
-	const int32 NodeArrayIndex = PropertyNode->GetArrayIndex();
 	FProperty* Property = PropertyNode->GetProperty();
 	
 	FSlateFontInfo FontStyle = FAppStyle::GetFontStyle( PropertyEditorConstants::PropertyFontStyle );
@@ -1117,7 +1116,7 @@ namespace PropertyEditorHelpers
 
 	void CollectObjectNodes( TSharedPtr<FPropertyNode> StartNode, TArray<FObjectPropertyNode*>& OutObjectNodes )
 	{
-		if( StartNode->AsObjectNode() != NULL )
+		if( StartNode->AsObjectNode() != nullptr )
 		{
 			OutObjectNodes.Add( StartNode->AsObjectNode() );
 		}
@@ -1126,69 +1125,26 @@ namespace PropertyEditorHelpers
 		{
 			CollectObjectNodes( StartNode->GetChildNode( ChildIndex ), OutObjectNodes );
 		}
-		
 	}
 
-	TArray<FName> GetValidEnumsFromPropertyOverride(const TArray<UObject*>& ObjectList, const FProperty* Property, const UEnum* InEnum)
+	TArray<FName> GetValidEnumsFromPropertyOverride(const FProperty* Property, const UEnum* InEnum)
 	{
 		TArray<FName> ValidEnumValues;
-		
-		bool bMergeValidEnums = false;
 
 		static const FName ValidEnumValuesName("ValidEnumValues");
-		static const FName GetValidEnumValuesName("GetValidEnumValues");
-		TArray<FString> ValidEnumValuesAsString;
 		
-		if(Property->HasMetaData(ValidEnumValuesName))
+		if (Property->HasMetaData(ValidEnumValuesName))
 		{
-			bMergeValidEnums = true;
+			TArray<FString> ValidEnumValuesAsString;
 			Property->GetMetaData(ValidEnumValuesName).ParseIntoArray(ValidEnumValuesAsString, TEXT(","));
-		}
-
-		if(Property->HasMetaData(GetValidEnumValuesName))
-		{
-			const FString GetValidEnumValuesNameFunctionName = Property->GetMetaData(GetValidEnumValuesName);
-			if (!GetValidEnumValuesNameFunctionName.IsEmpty())
+			
+			for (FString& Value : ValidEnumValuesAsString)
 			{
-				for (UObject* Object : ObjectList)
-				{
-					const UFunction* GetValidEnumValuesNameFunction = Object ? Object->FindFunction(*GetValidEnumValuesNameFunctionName) : nullptr;
-					if (GetValidEnumValuesNameFunction)
-					{
-						DECLARE_DELEGATE_RetVal(TArray<FString>, FGetValidEnumValuesNameF);
-						if (!bMergeValidEnums)
-						{
-							ValidEnumValuesAsString = FGetValidEnumValuesNameF::CreateUFunction(Object, GetValidEnumValuesNameFunction->GetFName()).Execute();
-							bMergeValidEnums = true;
-						}
-						else
-						{
-							TArray<FString> ValidEnumValuesAsStringToMerge = FGetValidEnumValuesNameF::CreateUFunction(Object, GetValidEnumValuesNameFunction->GetFName()).Execute();
-							TArray<FString> OldValidEnumValuesAsString = MoveTemp(ValidEnumValuesAsString);
-							for (const FString& OldValidEnumValue : OldValidEnumValuesAsString)
-							{
-								if (ValidEnumValuesAsStringToMerge.Contains(OldValidEnumValue))
-								{
-									ValidEnumValuesAsString.Add(OldValidEnumValue);
-								}
-							}
-						}
-						if (ValidEnumValuesAsString.IsEmpty())
-						{
-							ValidEnumValues.Add(NAME_None);
-							return ValidEnumValues;
-						}
-					}
-				}
+				Value.TrimStartInline();
+				ValidEnumValues.Add(*InEnum->GenerateFullEnumName(*Value));
 			}
 		}
-		
-		for(FString& Value : ValidEnumValuesAsString)
-		{
-			Value.TrimStartInline();
-			ValidEnumValues.Add(*InEnum->GenerateFullEnumName(*Value));
-		}
-			
+
 		return ValidEnumValues;
 	}
 
@@ -1197,12 +1153,12 @@ namespace PropertyEditorHelpers
 		TArray<FName> InvalidEnumValues;
 
 		static const FName InvalidEnumValuesName("InvalidEnumValues");
-		if(Property->HasMetaData(InvalidEnumValuesName))
+		if (Property->HasMetaData(InvalidEnumValuesName))
 		{
 			TArray<FString> InvalidEnumValuesAsString;
 
 			Property->GetMetaData(InvalidEnumValuesName).ParseIntoArray(InvalidEnumValuesAsString, TEXT(","));
-			for(FString& Value : InvalidEnumValuesAsString)
+			for (FString& Value : InvalidEnumValuesAsString)
 			{
 				Value.TrimStartInline();
 				InvalidEnumValues.Add(*InEnum->GenerateFullEnumName(*Value));
@@ -1211,6 +1167,40 @@ namespace PropertyEditorHelpers
 
 		return InvalidEnumValues;
 	}
+	
+	TArray<FName> GetRestrictedEnumsFromPropertyOverride(TArrayView<UObject*> ObjectList, const FProperty* Property, const UEnum* InEnum)
+	{
+		TArray<FName> RestrictedEnumValues;
+
+		static const FName GetRestrictedEnumValuesName("GetRestrictedEnumValues");
+		TArray<FString> ValidEnumValuesAsString;
+		
+		if (Property->HasMetaData(GetRestrictedEnumValuesName))
+		{
+			const FString GetRestrictedEnumValuesNameFunctionName = Property->GetMetaData(GetRestrictedEnumValuesName);
+			if (!GetRestrictedEnumValuesNameFunctionName.IsEmpty())
+			{
+				for (UObject* Object : ObjectList)
+				{
+					const UFunction* GetRestrictedEnumValuesNameFunction = Object ? Object->FindFunction(*GetRestrictedEnumValuesNameFunctionName) : nullptr;
+					if (GetRestrictedEnumValuesNameFunction)
+					{
+						DECLARE_DELEGATE_RetVal(TArray<FString>, FGetValidEnumValuesNameF);
+
+						ValidEnumValuesAsString.Append(FGetValidEnumValuesNameF::CreateUFunction(Object, GetRestrictedEnumValuesNameFunction->GetFName()).Execute());
+					}
+				}
+			}
+
+			for (FString& Value : ValidEnumValuesAsString)
+			{
+				Value.TrimStartInline();
+				RestrictedEnumValues.AddUnique(*InEnum->GenerateFullEnumName(*Value));
+			}
+		}
+		
+		return RestrictedEnumValues;
+	}
 
 	TMap<FName, FText> GetEnumValueDisplayNamesFromPropertyOverride(const FProperty* Property, const UEnum* InEnum)
 	{
diff --git a/Engine/Source/Editor/PropertyEditor/Private/PropertyEditorHelpers.h b/Engine/Source/Editor/PropertyEditor/Private/PropertyEditorHelpers.h
--- a/Engine/Source/Editor/PropertyEditor/Private/PropertyEditorHelpers.h
+++ b/Engine/Source/Editor/PropertyEditor/Private/PropertyEditorHelpers.h
@@ -256,12 +256,11 @@ namespace PropertyEditorHelpers
 	/**
 	 * Returns any enums that are explicitly allowed by the "ValidEnumValues" metadata on FProperty using the specified enum.
 	 *
-	 * @param ObjectList The list of objects currently edited
 	 * @param Property	The property which may contain the "ValidEnumValues" metadata
 	 * @param InEnum	The enum to search
 	 * @return The array of allowed enums.  NOTE: If an empty array is returned all enum values are allowed. It is an error for a property to hide all enum values so that state is undefined here.
 	 */
-	TArray<FName> GetValidEnumsFromPropertyOverride(const TArray<UObject*>& ObjectList, const FProperty* Property, const UEnum* InEnum);
+	TArray<FName> GetValidEnumsFromPropertyOverride(const FProperty* Property, const UEnum* InEnum);
 
 	/**
 	 * Returns any enums that are explicitly Disallowed by the "InvalidEnumValues" metadata on FProperty using the specified enum.
@@ -272,6 +271,16 @@ namespace PropertyEditorHelpers
 	 */
 	TArray<FName> GetInvalidEnumsFromPropertyOverride(const FProperty* Property, const UEnum* InEnum);
 
+	/**
+	 * Returns any enums that are explicitly restricted by the "GetRestrictedEnumValues" metadata on FProperty using the specified enum.
+	 *
+	 * @param ObjectList The list of objects currently edited
+	 * @param Property	The property which may contain the "GetRestrictedEnumValues" metadata
+	 * @param InEnum	The enum to search
+	 * @return The array of restricted enums.  
+	 */
+	TArray<FName> GetRestrictedEnumsFromPropertyOverride(TArrayView<UObject*> ObjectList, const FProperty* Property, const UEnum* InEnum);
+
 	/**
 	 * Returns any enums that are have an overridden display name from the "EnumValueDisplayNameOverrides" metadata on FProperty using the specified enum.
 	 *
diff --git a/Engine/Source/Editor/PropertyEditor/Private/PropertyHandleImpl.cpp b/Engine/Source/Editor/PropertyEditor/Private/PropertyHandleImpl.cpp
--- a/Engine/Source/Editor/PropertyEditor/Private/PropertyHandleImpl.cpp
+++ b/Engine/Source/Editor/PropertyEditor/Private/PropertyHandleImpl.cpp
@@ -3082,8 +3082,10 @@ bool FPropertyHandleBase::GeneratePossibleValues(TArray< TSharedPtr<FString> >&
 		TArray<UObject*> OuterObjects;
 		GetOuterObjects(OuterObjects);		
 		
-		const TArray<FName> ValidEnumValues = PropertyEditorHelpers::GetValidEnumsFromPropertyOverride(OuterObjects, Property, Enum);
+		const TArray<FName> ValidEnumValues = PropertyEditorHelpers::GetValidEnumsFromPropertyOverride(Property, Enum);
 		const TArray<FName> InvalidEnumValues = PropertyEditorHelpers::GetInvalidEnumsFromPropertyOverride(Property, Enum);
+		const TArray<FName> RestrictedEnumValues = PropertyEditorHelpers::GetRestrictedEnumsFromPropertyOverride(OuterObjects, Property, Enum);
+		
 		const TMap<FName, FText> EnumValueDisplayNameOverrides = PropertyEditorHelpers::GetEnumValueDisplayNamesFromPropertyOverride(Property, Enum);
 
 		//NumEnums() - 1, because the last item in an enum is the _MAX item
@@ -3095,12 +3097,12 @@ bool FPropertyHandleBase::GeneratePossibleValues(TArray< TSharedPtr<FString> >&
 			{
 				if(ValidEnumValues.Num() > 0)
 				{
-					bShouldBeHidden = ValidEnumValues.Find(Enum->GetNameByIndex(EnumIndex)) == INDEX_NONE;
+					bShouldBeHidden = !ValidEnumValues.Contains(Enum->GetNameByIndex(EnumIndex));
 				}
 				// If both are specified, InvalidEnumValues takes precedence
 				else if(InvalidEnumValues.Num() > 0)
 				{
-					bShouldBeHidden = InvalidEnumValues.Find(Enum->GetNameByIndex(EnumIndex)) != INDEX_NONE;
+					bShouldBeHidden = InvalidEnumValues.Contains(Enum->GetNameByIndex(EnumIndex));
 				}
 			}
 
@@ -3120,7 +3122,7 @@ bool FPropertyHandleBase::GeneratePossibleValues(TArray< TSharedPtr<FString> >&
 				}
 
 				FText RestrictionTooltip;
-				const bool bIsRestricted = GenerateRestrictionToolTip(EnumName, RestrictionTooltip);
+				const bool bIsRestricted = GenerateRestrictionToolTip(EnumName, RestrictionTooltip) || RestrictedEnumValues.Contains(Enum->GetNameByIndex(EnumIndex));
 				OutRestrictedItems.Add(bIsRestricted);
 
 				if (EnumDisplayName.Len() == 0)
diff --git a/Engine/Source/Editor/PropertyEditor/Private/UserInterface/PropertyEditor/SPropertyEditorCombo.cpp b/Engine/Source/Editor/PropertyEditor/Private/UserInterface/PropertyEditor/SPropertyEditorCombo.cpp
--- a/Engine/Source/Editor/PropertyEditor/Private/UserInterface/PropertyEditor/SPropertyEditorCombo.cpp
+++ b/Engine/Source/Editor/PropertyEditor/Private/UserInterface/PropertyEditor/SPropertyEditorCombo.cpp
@@ -260,10 +260,7 @@ void SPropertyEditorCombo::GenerateComboBoxStrings( TArray< TSharedPtr<FString>
 
 			if(Enum)
 			{
-				TArray<UObject*> ObjectList;
-				ComboArgs.PropertyHandle->GetOuterObjects(ObjectList);
-				
-				TArray<FName> ValidPropertyEnums = PropertyEditorHelpers::GetValidEnumsFromPropertyOverride(ObjectList, Property, Enum);
+				TArray<FName> ValidPropertyEnums = PropertyEditorHelpers::GetValidEnumsFromPropertyOverride(Property, Enum);
 				TArray<FName> InvalidPropertyEnums = PropertyEditorHelpers::GetInvalidEnumsFromPropertyOverride(Property, Enum);
 				
 				// Get enum doc link (not just GetDocumentationLink as that is the documentation for the struct we're in, not the enum documentation)
diff --git a/Engine/Source/Editor/PropertyEditor/Private/UserInterface/PropertyEditor/SPropertyEditorNumeric.h b/Engine/Source/Editor/PropertyEditor/Private/UserInterface/PropertyEditor/SPropertyEditorNumeric.h
--- a/Engine/Source/Editor/PropertyEditor/Private/UserInterface/PropertyEditor/SPropertyEditorNumeric.h
+++ b/Engine/Source/Editor/PropertyEditor/Private/UserInterface/PropertyEditor/SPropertyEditorNumeric.h
@@ -98,10 +98,7 @@ public:
 						}
 					};
 
-					TArray<UObject*> OuterObjects;
-					PropertyHandle->GetOuterObjects(OuterObjects);
-					
-					const TArray<FName> AllowedPropertyEnums = PropertyEditorHelpers::GetValidEnumsFromPropertyOverride(OuterObjects, PropertyHandle->GetProperty(), BitmaskEnum);
+					const TArray<FName> AllowedPropertyEnums = PropertyEditorHelpers::GetValidEnumsFromPropertyOverride(PropertyHandle->GetProperty(), BitmaskEnum);
 					const TArray<FName> DisallowedPropertyEnums = PropertyEditorHelpers::GetInvalidEnumsFromPropertyOverride(PropertyHandle->GetProperty(), BitmaskEnum);
 					// Note: This loop doesn't include (BitflagsEnum->NumEnums() - 1) in order to skip the implicit "MAX" value that gets added to the enum type at compile time.
 					for (int32 BitmaskEnumIndex = 0; BitmaskEnumIndex < BitmaskEnum->NumEnums() - 1; ++BitmaskEnumIndex)
