PropertyRowGenerator : exposing filter function for generated rows

diff --git a/Engine/Source/Editor/PropertyEditor/Private/PropertyRowGenerator.cpp b/Engine/Source/Editor/PropertyEditor/Private/PropertyRowGenerator.cpp
--- a/Engine/Source/Editor/PropertyEditor/Private/PropertyRowGenerator.cpp
+++ b/Engine/Source/Editor/PropertyEditor/Private/PropertyRowGenerator.cpp
@@ -154,6 +154,7 @@ FPropertyRowGenerator::FPropertyRowGenerator(const FPropertyRowGeneratorArgs& In
 	, PropertyUtilities(new FPropertyRowGeneratorUtilities(*this))
 	, PropertyGenerationUtilities(new FPropertyRowGeneratorGenerationUtilities(*this))
 {
+	CurrentFilter.bShowAllAdvanced = true;
 }
 
 FPropertyRowGenerator::FPropertyRowGenerator(const FPropertyRowGeneratorArgs& InArgs, TSharedPtr<FAssetThumbnailPool> InThumbnailPool)
@@ -161,7 +162,7 @@ FPropertyRowGenerator::FPropertyRowGenerator(const FPropertyRowGeneratorArgs& In
 	, PropertyUtilities(new FPropertyRowGeneratorUtilities(*this))
 	, PropertyGenerationUtilities(new FPropertyRowGeneratorGenerationUtilities(*this))
 {
-
+	CurrentFilter.bShowAllAdvanced = true;
 }
 
 FPropertyRowGenerator::~FPropertyRowGenerator()
@@ -355,6 +356,12 @@ void FPropertyRowGenerator::InvalidateCachedState()
 	}
 }
 
+void FPropertyRowGenerator::FilterNodes(const TArray<FString>& InFilterStrings)
+{
+	CurrentFilter.FilterStrings = InFilterStrings;
+	UpdateDetailRows();
+}
+
 void FPropertyRowGenerator::Tick(float DeltaTime)
 {
 	for (TSharedPtr<IDetailCustomization>& Customization : CustomizationClassInstancesPendingDelete)
@@ -569,9 +576,6 @@ void FPropertyRowGenerator::UpdateDetailRows()
 
 	//NumVisbleTopLevelObjectNodes = 0;
 
-	FDetailFilter CurrentFilter;
-	CurrentFilter.bShowAllAdvanced = true;
-
 	for (int32 RootNodeIndex = 0; RootNodeIndex < RootPropertyNodes.Num(); ++RootNodeIndex)
 	{
 		TSharedPtr<FComplexPropertyNode>& RootPropertyNode = RootPropertyNodes[RootNodeIndex];
diff --git a/Engine/Source/Editor/PropertyEditor/Private/PropertyRowGenerator.h b/Engine/Source/Editor/PropertyEditor/Private/PropertyRowGenerator.h
--- a/Engine/Source/Editor/PropertyEditor/Private/PropertyRowGenerator.h
+++ b/Engine/Source/Editor/PropertyEditor/Private/PropertyRowGenerator.h
@@ -3,6 +3,7 @@
 #pragma once
 
 #include "CoreMinimal.h"
+#include "DetailFilter.h"
 #include "UObject/WeakObjectPtr.h"
 #include "PropertyPath.h"
 #include "TickableEditorObject.h"
@@ -10,7 +11,6 @@
 #include "PropertyEditorModule.h"
 #include "IPropertyRowGenerator.h"
 
-
 class FDetailCategoryImpl;
 class FDetailLayoutBuilderImpl;
 class FPropertyNode;
@@ -100,6 +100,7 @@ public:
 	/** This function sets property paths to generate PropertyNodes.This improves the performance for cases where PropertyView is only showing a few properties of the object by not generating all other PropertyNodes */
 	virtual void SetPropertyGenerationAllowListPaths(const TSet<FString>& InPropertyGenerationAllowListPaths) override;
 	virtual void InvalidateCachedState() override;
+	virtual void FilterNodes(const TArray<FString>& InFilterStrings) override;
 
 	/** FTickableEditorObject interface */
 	virtual ETickableTickType GetTickableTickType() const override { return ETickableTickType::Always; }
@@ -161,6 +162,8 @@ private:
 	FOnValidatePropertyRowGeneratorNodes CustomValidatePropertyNodesFunction;
 	/** If specified only nodes for these properties will be generated */
 	TSet<FString> PropertyGenerationAllowListPaths;
+	/** Filter to apply on generated rows nodes */
+	FDetailFilter CurrentFilter;
 
 	bool bViewingClassDefaultObject = false;
 	bool bRefreshPending = false;
diff --git a/Engine/Source/Editor/PropertyEditor/Public/IPropertyRowGenerator.h b/Engine/Source/Editor/PropertyEditor/Public/IPropertyRowGenerator.h
--- a/Engine/Source/Editor/PropertyEditor/Public/IPropertyRowGenerator.h
+++ b/Engine/Source/Editor/PropertyEditor/Public/IPropertyRowGenerator.h
@@ -120,4 +120,7 @@ public:
 
 	/* Invalidates internal cached state.  Common use of this API is to synchronize the viewed object with changes made by external code. */
 	virtual void InvalidateCachedState() = 0;
+
+	/** Apply filters to generated row tree nodes */
+	virtual void FilterNodes(const TArray<FString>& InFilterStrings) = 0;
 };
\ No newline at end of file
