More unreachable code fixes.

diff --git a/Engine/Source/Editor/PropertyEditor/Private/Presentation/PropertyEditor/PropertyEditor.cpp b/Engine/Source/Editor/PropertyEditor/Private/Presentation/PropertyEditor/PropertyEditor.cpp
--- a/Engine/Source/Editor/PropertyEditor/Private/Presentation/PropertyEditor/PropertyEditor.cpp
+++ b/Engine/Source/Editor/PropertyEditor/Private/Presentation/PropertyEditor/PropertyEditor.cpp
@@ -57,14 +57,10 @@ FText FPropertyEditor::GetDisplayName() const
 	{
 		return ItemPropertyNode->GetDisplayName();
 	}
-	else
-	{
-		FString DisplayName;
-		PropertyNode->GetQualifiedName( DisplayName, true );
-		return FText::FromString(DisplayName);
-	}
 
-	return FText::GetEmpty();
+	FString DisplayName;
+	PropertyNode->GetQualifiedName( DisplayName, true );
+	return FText::FromString(DisplayName);
 }
 
 FText FPropertyEditor::GetToolTipText() const
diff --git a/Engine/Source/Editor/Sequencer/Private/SequencerNodeTree.cpp b/Engine/Source/Editor/Sequencer/Private/SequencerNodeTree.cpp
--- a/Engine/Source/Editor/Sequencer/Private/SequencerNodeTree.cpp
+++ b/Engine/Source/Editor/Sequencer/Private/SequencerNodeTree.cpp
@@ -379,7 +379,6 @@ bool FSequencerNodeTree::IsNodeFiltered(const TSharedPtr<UE::Sequencer::FViewMod
 
 	TViewModelPtr<IOutlinerExtension> OutlinerItem = CastViewModel<IOutlinerExtension>(Node);
 	return OutlinerItem && FilteredNodes.Contains(OutlinerItem);
-	return FilteredNodes.Contains(CastViewModel<IOutlinerExtension>(Node));
 }
 
 TSharedPtr<UE::Sequencer::FSectionModel> FSequencerNodeTree::GetSectionModel(const UMovieSceneSection* Section) const
diff --git a/Engine/Source/Editor/StructViewer/Private/SStructViewer.cpp b/Engine/Source/Editor/StructViewer/Private/SStructViewer.cpp
--- a/Engine/Source/Editor/StructViewer/Private/SStructViewer.cpp
+++ b/Engine/Source/Editor/StructViewer/Private/SStructViewer.cpp
@@ -1527,8 +1527,6 @@ TSharedRef<SWidget> SStructViewer::GetViewButtonContent()
 	MenuBuilder.EndSection();
 
 	return MenuBuilder.MakeWidget();
-
-	return MenuBuilder.MakeWidget();
 }
 
 void SStructViewer::SetCurrentDeveloperViewType(EStructViewerDeveloperType NewType)
diff --git a/Engine/Source/Runtime/Engine/Private/EngineUtils.cpp b/Engine/Source/Runtime/Engine/Private/EngineUtils.cpp
--- a/Engine/Source/Runtime/Engine/Private/EngineUtils.cpp
+++ b/Engine/Source/Runtime/Engine/Private/EngineUtils.cpp
@@ -68,9 +68,9 @@ bool HActor::AlwaysAllowsTranslucentPrimitives() const
 {
 #if WITH_EDITOR
 	return PrimComponent->bAlwaysAllowTranslucentSelect;
-#endif
-
+#else
 	return false;
+#endif
 }
 
 EMouseCursor::Type HTranslucentActor::GetMouseCursor()
diff --git a/Engine/Source/Runtime/LevelSequence/Public/LevelSequence.h b/Engine/Source/Runtime/LevelSequence/Public/LevelSequence.h
--- a/Engine/Source/Runtime/LevelSequence/Public/LevelSequence.h
+++ b/Engine/Source/Runtime/LevelSequence/Public/LevelSequence.h
@@ -161,8 +161,9 @@ public:
 #if WITH_EDITORONLY_DATA
 		auto const* Found = MetaDataObjects.FindByPredicate([InClass](UObject* In) { return In && In->GetClass() == InClass; });
 		return Found ? CastChecked<UObject>(*Found) : nullptr;
-#endif
+#else
 		return nullptr;
+#endif
 	}
 
 	/**
@@ -181,8 +182,9 @@ public:
 			MetaDataObjects.Add(Found);
 		}
 		return Found;
-#endif
+#else
 		return nullptr;
+#endif
 	}
 
 	/**
