PropertyEditor Truncation fixes

diff --git a/Engine/Source/Editor/PropertyEditor/Private/Presentation/PropertyTable/PropertyTableColumn.cpp b/Engine/Source/Editor/PropertyEditor/Private/Presentation/PropertyTable/PropertyTableColumn.cpp
--- a/Engine/Source/Editor/PropertyEditor/Private/Presentation/PropertyTable/PropertyTableColumn.cpp
+++ b/Engine/Source/Editor/PropertyEditor/Private/Presentation/PropertyTable/PropertyTableColumn.cpp
@@ -348,7 +348,7 @@ FORCEINLINE int32 FCompareRowByColumnAscending<FStructProperty>::CompareProperty
 
 		if ( LhsPropertyHandle->GetValue(LhsVector) != FPropertyAccess::Fail && RhsPropertyHandle->GetValue(RhsVector) != FPropertyAccess::Fail )
 		{
-			return RhsVector.SizeSquared() - LhsVector.SizeSquared();
+			return static_cast<int32>(RhsVector.SizeSquared() - LhsVector.SizeSquared());
 		}
 
 		FVector2D LhsVector2D;
@@ -356,7 +356,7 @@ FORCEINLINE int32 FCompareRowByColumnAscending<FStructProperty>::CompareProperty
 
 		if ( LhsPropertyHandle->GetValue(LhsVector2D) != FPropertyAccess::Fail && RhsPropertyHandle->GetValue(RhsVector2D) != FPropertyAccess::Fail )
 		{
-			return RhsVector2D.SizeSquared() - LhsVector2D.SizeSquared();
+			return static_cast<int32>(RhsVector2D.SizeSquared() - LhsVector2D.SizeSquared());
 		}
 
 		FVector4 LhsVector4;
@@ -364,7 +364,7 @@ FORCEINLINE int32 FCompareRowByColumnAscending<FStructProperty>::CompareProperty
 
 		if ( LhsPropertyHandle->GetValue(LhsVector4) != FPropertyAccess::Fail && RhsPropertyHandle->GetValue(RhsVector4) != FPropertyAccess::Fail )
 		{
-			return RhsVector4.SizeSquared() - LhsVector4.SizeSquared();
+			return static_cast<int32>(RhsVector4.SizeSquared() - LhsVector4.SizeSquared());
 		}
 	}
 
diff --git a/Engine/Source/Editor/PropertyEditor/Private/PropertyHandleImpl.cpp b/Engine/Source/Editor/PropertyEditor/Private/PropertyHandleImpl.cpp
--- a/Engine/Source/Editor/PropertyEditor/Private/PropertyHandleImpl.cpp
+++ b/Engine/Source/Editor/PropertyEditor/Private/PropertyHandleImpl.cpp
@@ -925,7 +925,7 @@ static Type ClampIntegerValueFromMetaData(Type InValue, FPropertyHandleBase& InP
 		FObjectPropertyNode* ObjectPropertyNode = InPropertyNode.FindObjectItemParent();
 		if (ObjectPropertyNode && ObjectPropertyNode->GetNumObjects() == 1)
 		{
-			Type LastValidIndex = GetArrayPropertyLastValidIndex(ObjectPropertyNode, ArrayClampString);
+			Type LastValidIndex = static_cast<Type>(GetArrayPropertyLastValidIndex(ObjectPropertyNode, ArrayClampString));
 			RetVal = FMath::Clamp<Type>(RetVal, 0, LastValidIndex);
 		}
 		else
@@ -3839,7 +3839,7 @@ FPropertyAccess::Result FPropertyHandleByte::GetValue( uint8& OutValue ) const
 		else
 		{
 			check(PropertyNodePin.IsValid());
-			OutValue = CastFieldChecked<FEnumProperty>(Property)->GetUnderlyingProperty()->GetUnsignedIntPropertyValue(PropValue);
+			OutValue = static_cast<uint8>(CastFieldChecked<FEnumProperty>(Property)->GetUnderlyingProperty()->GetUnsignedIntPropertyValue(PropValue));
 		}
 	}
 
@@ -4457,7 +4457,7 @@ FPropertyAccess::Result FPropertyHandleMixed::GetValue(float& OutValue) const
 {
 	double AsDouble;
 	FPropertyAccess::Result Res = GetValue(AsDouble);
-	OutValue = AsDouble;
+	OutValue = static_cast<float>(AsDouble);
 	return Res;
 }
 
diff --git a/Engine/Source/Editor/PropertyEditor/Private/SConstrainedBox.cpp b/Engine/Source/Editor/PropertyEditor/Private/SConstrainedBox.cpp
--- a/Engine/Source/Editor/PropertyEditor/Private/SConstrainedBox.cpp
+++ b/Engine/Source/Editor/PropertyEditor/Private/SConstrainedBox.cpp
@@ -31,7 +31,7 @@ FVector2D SConstrainedBox::ComputeDesiredSize(float LayoutScaleMultiplier) const
 	{
 		FVector2D ChildSize = ChildSlot.GetWidget()->GetDesiredSize();
 
-		float XVal = FMath::Max(MinWidthVal, ChildSize.X);
+		double XVal = FMath::Max(MinWidthVal, ChildSize.X);
 		if (MaxWidthVal > MinWidthVal)
 		{
 			XVal = FMath::Min(MaxWidthVal, XVal);
diff --git a/Engine/Source/Editor/PropertyEditor/Private/UserInterface/PropertyEditor/SPropertyEditorNumeric.h b/Engine/Source/Editor/PropertyEditor/Private/UserInterface/PropertyEditor/SPropertyEditorNumeric.h
--- a/Engine/Source/Editor/PropertyEditor/Private/UserInterface/PropertyEditor/SPropertyEditorNumeric.h
+++ b/Engine/Source/Editor/PropertyEditor/Private/UserInterface/PropertyEditor/SPropertyEditorNumeric.h
@@ -80,7 +80,7 @@ public:
 					const TMap<FName, FText> EnumValueDisplayNameOverrides = PropertyEditorHelpers::GetEnumValueDisplayNamesFromPropertyOverride(PropertyHandle->GetProperty(), BitmaskEnum);
 
 					const bool bUseEnumValuesAsMaskValues = BitmaskEnum->GetBoolMetaData(PropertyEditorConstants::MD_UseEnumValuesAsMaskValuesInEditor);
-					auto AddNewBitmaskFlagLambda = [BitmaskEnum, &Result, &EnumValueDisplayNameOverrides](int32 InEnumIndex, int64 InFlagValue)
+					auto AddNewBitmaskFlagLambda = [BitmaskEnum, &Result, &EnumValueDisplayNameOverrides](int32 InEnumIndex, NumericType InFlagValue)
 					{
 						Result.Emplace();
 						FBitmaskFlagInfo* BitmaskFlag = &Result.Last();
@@ -103,7 +103,6 @@ public:
 					// Note: This loop doesn't include (BitflagsEnum->NumEnums() - 1) in order to skip the implicit "MAX" value that gets added to the enum type at compile time.
 					for (int32 BitmaskEnumIndex = 0; BitmaskEnumIndex < BitmaskEnum->NumEnums() - 1; ++BitmaskEnumIndex)
 					{
-						const int64 EnumValue = BitmaskEnum->GetValueByIndex(BitmaskEnumIndex);
 						bool bShouldBeHidden = BitmaskEnum->HasMetaData(TEXT("Hidden"), BitmaskEnumIndex);
 						if (!bShouldBeHidden)
 						{
@@ -117,11 +116,15 @@ public:
 								bShouldBeHidden = DisallowedPropertyEnums.Find(BitmaskEnum->GetNameByIndex(BitmaskEnumIndex)) != INDEX_NONE;
 							}
 						}
+
+						const int64 EnumValue64 = BitmaskEnum->GetValueByIndex(BitmaskEnumIndex);
+						const NumericType EnumValue = static_cast<NumericType>(EnumValue64);
+
 						if (EnumValue >= 0 && !bShouldBeHidden)
 						{
 							if (bUseEnumValuesAsMaskValues)
 							{
-								if (EnumValue < MAX_int64 && FMath::IsPowerOfTwo(EnumValue))
+								if (EnumValue64 < static_cast<int64>(TNumericLimits<NumericType>::Max()) &&FMath::IsPowerOfTwo(EnumValue64))
 								{
 									AddNewBitmaskFlagLambda(BitmaskEnumIndex, EnumValue);
 								}
