PropertyEditor  Truncation warning fixes

diff --git a/Engine/Source/Editor/PropertyEditor/Private/DetailAdvancedDropdownNode.cpp b/Engine/Source/Editor/PropertyEditor/Private/DetailAdvancedDropdownNode.cpp
--- a/Engine/Source/Editor/PropertyEditor/Private/DetailAdvancedDropdownNode.cpp
+++ b/Engine/Source/Editor/PropertyEditor/Private/DetailAdvancedDropdownNode.cpp
@@ -52,7 +52,7 @@ public:
 		+ SHorizontalBox::Slot()
 		.HAlign(HAlign_Left)
 		.VAlign(VAlign_Center)
-		.Padding(2, 0, 0, 0)
+		.Padding(2.0f, 0.0f, 0.0f, 0.0f)
 		.AutoWidth()
 		[
 			SAssignNew(ExpanderButton, SButton)
@@ -62,7 +62,7 @@ public:
 			.ClickMethod(EButtonClickMethod::MouseDown)
 			.OnClicked(InArgs._OnClicked)
 			.IsEnabled(InArgs._IsButtonEnabled)
-			.ContentPadding(0)
+			.ContentPadding(0.0f)
 			.IsFocusable(false)
 			.ToolTipText(this, &SAdvancedDropdownRow::GetAdvancedPulldownToolTipText )
 			[
@@ -74,7 +74,7 @@ public:
 		+ SHorizontalBox::Slot()
 		.VAlign(VAlign_Center)
 		.HAlign(HAlign_Fill)
-		.Padding(4, 0, 0, 0)
+		.Padding(4.0f, 0.0f, 0.0f, 0.0f)
 		[
 			SNew(SBox)
 			.VAlign(VAlign_Center)
@@ -108,7 +108,7 @@ public:
 		[
 			SNew(SBorder)
 			.BorderImage(FAppStyle::Get().GetBrush("DetailsView.GridLine"))
-			.Padding( 0, 0, 0, 1 )
+			.Padding( 0.0f, 0.0f, 0.0f, 1.0f )
 			[
 				SNew(SHorizontalBox)
 				+ SHorizontalBox::Slot()
@@ -117,7 +117,7 @@ public:
 					SNew(SBorder)
 					.BorderImage(FAppStyle::Get().GetBrush("DetailsView.CategoryMiddle"))
 					.BorderBackgroundColor(this, &SAdvancedDropdownRow::GetRowBackgroundColor)
-					.Padding(0)
+					.Padding(0.0f)
 					[
 						ContentWidget.ToSharedRef()
 					]
diff --git a/Engine/Source/Editor/PropertyEditor/Private/DetailMultiTopLevelObjectRootNode.cpp b/Engine/Source/Editor/PropertyEditor/Private/DetailMultiTopLevelObjectRootNode.cpp
--- a/Engine/Source/Editor/PropertyEditor/Private/DetailMultiTopLevelObjectRootNode.cpp
+++ b/Engine/Source/Editor/PropertyEditor/Private/DetailMultiTopLevelObjectRootNode.cpp
@@ -38,7 +38,7 @@ void SDetailMultiTopLevelObjectTableRow::SetContent(TSharedRef<SWidget> InConten
 			]
 			+ SHorizontalBox::Slot()
 			.Expose(ContentSlot)
-			.Padding(FMargin(0, 0, 0, 16))
+			.Padding(FMargin(0.0f, 0.0f, 0.0f, 16.0f))
 			[
 				InContent
 			]
diff --git a/Engine/Source/Editor/PropertyEditor/Private/MaterialList.cpp b/Engine/Source/Editor/PropertyEditor/Private/MaterialList.cpp
--- a/Engine/Source/Editor/PropertyEditor/Private/MaterialList.cpp
+++ b/Engine/Source/Editor/PropertyEditor/Private/MaterialList.cpp
@@ -320,13 +320,13 @@ TSharedRef<SWidget> FMaterialItemView::MakeBrowseNaniteOverrideMaterialButton()
 		SNew(SBox)
 		.HAlign(HAlign_Center)
 		.VAlign(VAlign_Center)
-		.WidthOverride(22)
-		.HeightOverride(22)
+		.WidthOverride(22.0f)
+		.HeightOverride(22.0f)
 		.ToolTipText(LOCTEXT("BrowseToNaniteOverride_Tip", "Browse to the Nanite Override Material in Content Browser"))
 		[
 			SNew(SButton)
 			.ButtonStyle(FAppStyle::Get(), "SimpleButton")
-			.ContentPadding(0)
+			.ContentPadding(0.0f)
 			.IsFocusable(false)
 			.OnClicked(FOnClicked::CreateLambda([WeakMaterial = MaterialItem.Material]()
 				{
diff --git a/Engine/Source/Editor/PropertyEditor/Private/PropertyCustomizationHelpers.cpp b/Engine/Source/Editor/PropertyEditor/Private/PropertyCustomizationHelpers.cpp
--- a/Engine/Source/Editor/PropertyEditor/Private/PropertyCustomizationHelpers.cpp
+++ b/Engine/Source/Editor/PropertyEditor/Private/PropertyCustomizationHelpers.cpp
@@ -67,14 +67,14 @@ namespace PropertyCustomizationHelpers
 				SNew(SBox)
 				.HAlign(HAlign_Center)
 				.VAlign(VAlign_Center)
-				.WidthOverride(22)
-				.HeightOverride(22)
+				.WidthOverride(22.0f)
+				.HeightOverride(22.0f)
 				.ToolTipText(InArgs._Text)
 				[
 					SNew(SButton)
 					.ButtonStyle( FAppStyle::Get(), "SimpleButton" )
 					.OnClicked( this, &SPropertyEditorButton::OnClick )
-					.ContentPadding(0)
+					.ContentPadding(0.0f)
 					.IsFocusable(InArgs._IsFocusable)
 					[ 
 						SNew( SImage )
@@ -194,7 +194,7 @@ namespace PropertyCustomizationHelpers
 					return FReply::Handled();
 				})
 				.Text(LOCTEXT("SetButtonText", "Set to Value"))
-				.ContentPadding(0)
+				.ContentPadding(0.0f)
 				.IsFocusable(false)
 			];
 	}
@@ -286,7 +286,7 @@ namespace PropertyCustomizationHelpers
 		return
 			SNew(SComboButton)
 			.ComboButtonStyle( FAppStyle::Get(), "SimpleComboButton" )
-			.ContentPadding(2)
+			.ContentPadding(2.0f)
 			.ForegroundColor( FSlateColor::UseForeground() )
 			.HasDownArrow(true)
 			.MenuContent()
@@ -1093,7 +1093,7 @@ public:
 								[
 									SNew(SHorizontalBox)
 									+ SHorizontalBox::Slot()
-									.Padding(0)
+									.Padding(0.0f)
 									.VAlign(VAlign_Center)
 									.AutoWidth()
 									[
@@ -1110,7 +1110,7 @@ public:
 									+ SHorizontalBox::Slot()
 									.VAlign(VAlign_Center)
 									.FillWidth(1.0f)
-									.Padding(5, 0, 0, 0)
+									.Padding(5.0f, 0.0f, 0.0f, 0.0f)
 									[
 										SNew(SBox)
 										.HAlign(HAlign_Fill)
@@ -1121,7 +1121,7 @@ public:
 											SNew(SComboButton)
 											.OnGetMenuContent(this, &FSectionItemView::OnGetMaterialSlotNameMenuForSection)
 											.VAlign(VAlign_Center)
-											.ContentPadding(2)
+											.ContentPadding(2.0f)
 											.IsEnabled(!SectionItem.IsSectionUsingCloth)
 											.ButtonContent()
 											[
diff --git a/Engine/Source/Editor/PropertyEditor/Private/PropertyEditorHelpers.cpp b/Engine/Source/Editor/PropertyEditor/Private/PropertyEditorHelpers.cpp
--- a/Engine/Source/Editor/PropertyEditor/Private/PropertyEditorHelpers.cpp
+++ b/Engine/Source/Editor/PropertyEditor/Private/PropertyEditorHelpers.cpp
@@ -42,7 +42,7 @@ void SPropertyNameWidget::Construct( const FArguments& InArgs, TSharedPtr<FPrope
 	[
 		SAssignNew(HorizontalBox, SHorizontalBox)
 		+SHorizontalBox::Slot()
-		.Padding( FMargin( 0, 1, 0, 1 ) )
+		.Padding( FMargin( 0.0f, 1.0f, 0.0f, 1.0f) )
 		.FillWidth(1)
 		[
 			SNew(SBorder)
diff --git a/Engine/Source/Editor/PropertyEditor/Private/PropertyEditorToolkit.cpp b/Engine/Source/Editor/PropertyEditor/Private/PropertyEditorToolkit.cpp
--- a/Engine/Source/Editor/PropertyEditor/Private/PropertyEditorToolkit.cpp
+++ b/Engine/Source/Editor/PropertyEditor/Private/PropertyEditorToolkit.cpp
@@ -199,7 +199,7 @@ TSharedRef<SDockTab> FPropertyEditorToolkit::SpawnTab_PropertyTree( const FSpawn
 		.Content()
 		[
 			SNew(SBorder)
-			.Padding(4)
+			.Padding(4.0f)
 			.BorderImage( FAppStyle::GetBrush( "ToolPanel.GroupBorder" ) )
 			.Content()
 			[
@@ -279,7 +279,7 @@ void FPropertyEditorToolkit::ConstructTreeColumns( const TSharedRef< class SHead
 	.FixedWidth(24)
 	[
 		SNew(SBorder)
-		.Padding( 0 )
+		.Padding( 0.0f )
 		.BorderImage( FAppStyle::GetBrush( "NoBorder" ) )
 		.HAlign( HAlign_Center )
 		.VAlign( VAlign_Center )
diff --git a/Engine/Source/Editor/PropertyEditor/Private/SDetailCategoryTableRow.cpp b/Engine/Source/Editor/PropertyEditor/Private/SDetailCategoryTableRow.cpp
--- a/Engine/Source/Editor/PropertyEditor/Private/SDetailCategoryTableRow.cpp
+++ b/Engine/Source/Editor/PropertyEditor/Private/SDetailCategoryTableRow.cpp
@@ -65,7 +65,7 @@ void SDetailCategoryTableRow::Construct(const FArguments& InArgs, TSharedRef<FDe
 		+ SHorizontalBox::Slot()
 		.HAlign(HAlign_Left)
 		.VAlign(VAlign_Center)
-		.Padding(2, 0, 0, 0)
+		.Padding(2.0f, 0.0f, 0.0f, 0.0f)
 		.AutoWidth()
 		[
 			SNew(SDetailExpanderArrow, SharedThis(this))
@@ -84,7 +84,7 @@ void SDetailCategoryTableRow::Construct(const FArguments& InArgs, TSharedRef<FDe
 	{
 		HeaderBox->AddSlot()
 			.VAlign(VAlign_Center)
-			.Padding(4, 0, 0, 0)
+			.Padding(4.0f, 0.0f, 0.0f, 0.0f)
 			.FillWidth(1)
 			[
 				SNew(STextBlock)
@@ -130,7 +130,7 @@ void SDetailCategoryTableRow::Construct(const FArguments& InArgs, TSharedRef<FDe
 				SNew(SBorder)
 				.BorderImage(this, &SDetailCategoryTableRow::GetBackgroundImage)
 				.BorderBackgroundColor(this, &SDetailCategoryTableRow::GetInnerBackgroundColor)
-					.Padding(0)
+					.Padding(0.0f)
 				[
 					SNew(SBox)
 					.MinDesiredHeight(PropertyEditorConstants::PropertyRowHeight)
diff --git a/Engine/Source/Editor/PropertyEditor/Private/SDetailExpanderArrow.cpp b/Engine/Source/Editor/PropertyEditor/Private/SDetailExpanderArrow.cpp
--- a/Engine/Source/Editor/PropertyEditor/Private/SDetailExpanderArrow.cpp
+++ b/Engine/Source/Editor/PropertyEditor/Private/SDetailExpanderArrow.cpp
@@ -13,7 +13,7 @@ void SDetailExpanderArrow::Construct(const FArguments& InArgs, TSharedRef<SDetai
 	ChildSlot
 	[
 		SNew(SConstrainedBox)
-		.MinWidth(12)
+		.MinWidth(12.0f)
 		.Visibility(this, &SDetailExpanderArrow::GetExpanderVisibility)
 		[
 			SAssignNew(ExpanderArrow, SButton)
@@ -22,7 +22,7 @@ void SDetailExpanderArrow::Construct(const FArguments& InArgs, TSharedRef<SDetai
 			.HAlign(HAlign_Center)
 			.ClickMethod(EButtonClickMethod::MouseDown)
 			.OnClicked(this, &SDetailExpanderArrow::OnExpanderClicked)
-			.ContentPadding(0)
+			.ContentPadding(0.0f)
 			.IsFocusable(false)
 			[
 				SNew(SImage)
diff --git a/Engine/Source/Editor/PropertyEditor/Private/SDetailNameArea.cpp b/Engine/Source/Editor/PropertyEditor/Private/SDetailNameArea.cpp
--- a/Engine/Source/Editor/PropertyEditor/Private/SDetailNameArea.cpp
+++ b/Engine/Source/Editor/PropertyEditor/Private/SDetailNameArea.cpp
@@ -115,7 +115,7 @@ TSharedRef< SWidget > SDetailNameArea::BuildObjectNameArea( const TArray< TWeakO
 		.AutoWidth()
 		.HAlign(HAlign_Left)
 		.VAlign(VAlign_Center)
-		.Padding(0)
+		.Padding(0.0f)
 		[
 			SNew(SImage)
 			.Image(ActorIcon)
@@ -134,8 +134,8 @@ TSharedRef< SWidget > SDetailNameArea::BuildObjectNameArea( const TArray< TWeakO
 			TSharedRef<IObjectNameEditableTextBox> ObjectNameBox = EdWidgetsModule.CreateObjectNameEditableTextBox(SelectedObjects);
 
 			ObjectNameArea->AddSlot()
-				.FillWidth(1.0)
-				.Padding(8, 0, 3, 0)
+				.FillWidth(1.0f)
+				.Padding(8.0f, 0.0f, 3.0f, 0.0f)
 				[
 					SNew(SBox)
 					.WidthOverride(200.0f)
@@ -161,7 +161,7 @@ TSharedRef< SWidget > SDetailNameArea::BuildObjectNameArea( const TArray< TWeakO
 			ObjectNameArea->AddSlot()
 				.HAlign(HAlign_Right)
 				.VAlign(VAlign_Center)
-				.Padding(4,0,4,0)
+				.Padding(4.0f,0.0f,4.0f,0.0f)
 				.AutoWidth()
 				[
 					SNew( SButton )
diff --git a/Engine/Source/Editor/PropertyEditor/Private/SDetailSingleItemRow.cpp b/Engine/Source/Editor/PropertyEditor/Private/SDetailSingleItemRow.cpp
--- a/Engine/Source/Editor/PropertyEditor/Private/SDetailSingleItemRow.cpp
+++ b/Engine/Source/Editor/PropertyEditor/Private/SDetailSingleItemRow.cpp
@@ -431,7 +431,7 @@ void SDetailSingleItemRow::Construct( const FArguments& InArgs, FDetailLayoutCus
 			Widget = SNew(SBorder)
 				.BorderImage(FAppStyle::Get().GetBrush("DetailsView.CategoryMiddle"))
 				.BorderBackgroundColor(this, &SDetailSingleItemRow::GetInnerBackgroundColor)
-				.Padding(0)
+				.Padding(0.0f)
 				[
 					Splitter
 				];
@@ -445,7 +445,7 @@ void SDetailSingleItemRow::Construct( const FArguments& InArgs, FDetailLayoutCus
 			NameColumnBox->AddSlot()
 				.HAlign(HAlign_Left)
 				.VAlign(VAlign_Fill)
-				.Padding(0)
+				.Padding(0.0f)
 				.AutoWidth()
 				[
 					SNew(SDetailRowIndent, SharedThis(this))
@@ -459,7 +459,7 @@ void SDetailSingleItemRow::Construct( const FArguments& InArgs, FDetailLayoutCus
 				NameColumnBox->AddSlot()
 					.HAlign(HAlign_Left)
 					.VAlign(VAlign_Center)
-					.Padding(-4, 0, -10, 0)
+					.Padding(-4.0f, 0.0f, -10.0f, 0.0f)
 					.AutoWidth()
 					[
 						ReorderHandle
@@ -479,7 +479,7 @@ void SDetailSingleItemRow::Construct( const FArguments& InArgs, FDetailLayoutCus
 					NameColumnBox->AddSlot()
 						.HAlign(HAlign_Left)
 						.VAlign(VAlign_Center)
-						.Padding(-4, 0, -10, 0)
+						.Padding(-4.0f, 0.0f, -10.0f, 0.0f)
 						.AutoWidth()
 						[
 							ArrayHandle
@@ -501,7 +501,7 @@ void SDetailSingleItemRow::Construct( const FArguments& InArgs, FDetailLayoutCus
 			NameColumnBox->AddSlot()
 				.HAlign(HAlign_Left)
 				.VAlign(VAlign_Center)
-				.Padding(2,0,0,0)
+				.Padding(2.0f,0.0f,0.0f,0.0f)
 				.AutoWidth()
 				[
 					SNew(SDetailExpanderArrow, SharedThis(this))
@@ -510,7 +510,7 @@ void SDetailSingleItemRow::Construct( const FArguments& InArgs, FDetailLayoutCus
 			NameColumnBox->AddSlot()
 				.VAlign(VAlign_Center)
 				.HAlign(HAlign_Left)
-				.Padding(2,0,0,0)
+				.Padding(2.0f,0.0f,0.0f,0.0f)
 				.AutoWidth()
 				[
 					SNew(SEditConditionWidget)
@@ -523,7 +523,7 @@ void SDetailSingleItemRow::Construct( const FArguments& InArgs, FDetailLayoutCus
 				NameColumnBox->AddSlot()
 					.HAlign(WidgetRow.NameWidget.HorizontalAlignment)
 					.VAlign(WidgetRow.NameWidget.VerticalAlignment)
-					.Padding(2,0,0,0)
+					.Padding(2.0f,0.0f,0.0f,0.0f)
 					[
 						NameWidget.ToSharedRef()
 					];
@@ -556,7 +556,7 @@ void SDetailSingleItemRow::Construct( const FArguments& InArgs, FDetailLayoutCus
 						+ SHorizontalBox::Slot()
 						.HAlign(WidgetRow.ExtensionWidget.HorizontalAlignment)
 						.VAlign(WidgetRow.ExtensionWidget.VerticalAlignment)
-						.Padding(5,0,0,0)
+						.Padding(5.0f,0.0f,0.0f,0.0f)
 						.AutoWidth()
 						[
 							ExtensionWidget.ToSharedRef()
@@ -571,7 +571,7 @@ void SDetailSingleItemRow::Construct( const FArguments& InArgs, FDetailLayoutCus
 				NameColumnBox->AddSlot()
 					.HAlign(WidgetRow.WholeRowWidget.HorizontalAlignment)
 					.VAlign(WidgetRow.WholeRowWidget.VerticalAlignment)
-					.Padding(2,0,0,0)
+					.Padding(2.0f,0.0f,0.0f,0.0f)
 					[
 						WidgetRow.WholeRowWidget.Widget
 					];
@@ -657,7 +657,7 @@ void SDetailSingleItemRow::Construct( const FArguments& InArgs, FDetailLayoutCus
 				.BorderBackgroundColor(this, &SDetailSingleItemRow::GetOuterBackgroundColor)
 				.HAlign(HAlign_Right)
 				.VAlign(VAlign_Center)
-				.Padding(0)
+				.Padding(0.0f)
 				[
 					PropertyUpdatedWidgetBuilder.IsValid() ?
 					PropertyUpdatedWidgetBuilder
@@ -721,7 +721,7 @@ void SDetailSingleItemRow::Construct( const FArguments& InArgs, FDetailLayoutCus
 		[
 			SNew( SBorder )
 			.BorderImage(FAppStyle::Get().GetBrush("DetailsView.GridLine"))
-			.Padding(FMargin(0,0,0,1))
+			.Padding(FMargin(0.0f,0.0f,0.0f,1.0f))
 			.Clipping(EWidgetClipping::ClipToBounds)
 			[
 				SNew(SBox)
@@ -739,7 +739,7 @@ void SDetailSingleItemRow::Construct( const FArguments& InArgs, FDetailLayoutCus
 							SNew( SBorder )
 							.BorderImage(FAppStyle::Get().GetBrush("DetailsView.CategoryMiddle"))
 							.BorderBackgroundColor(this, &SDetailSingleItemRow::GetOuterBackgroundColor)
-							.Padding(0)
+							.Padding(0.0f)
 							[
 								Widget
 							]
diff --git a/Engine/Source/Editor/PropertyEditor/Private/SDetailsView.cpp b/Engine/Source/Editor/PropertyEditor/Private/SDetailsView.cpp
--- a/Engine/Source/Editor/PropertyEditor/Private/SDetailsView.cpp
+++ b/Engine/Source/Editor/PropertyEditor/Private/SDetailsView.cpp
@@ -272,7 +272,7 @@ void SDetailsView::Construct(const FArguments& InArgs, const FDetailsViewArgs& I
 		SNew( SHorizontalBox )
 		+SHorizontalBox::Slot()
 		.Padding(6.f)
-		.FillWidth(1)
+		.FillWidth(1.0f)
 		[
 			// Create the search box
 			SAssignNew(SearchBox, SSearchBox)
@@ -284,7 +284,7 @@ void SDetailsView::Construct(const FArguments& InArgs, const FDetailsViewArgs& I
 	
 	{
 		FilterRowHBox->AddSlot()
-			.Padding(0)
+			.Padding(0.0f)
 			.HAlign(HAlign_Right)
 			.VAlign(VAlign_Center)
 			.AutoWidth()
@@ -307,7 +307,7 @@ void SDetailsView::Construct(const FArguments& InArgs, const FDetailsViewArgs& I
 	if (DetailsViewArgs.bAllowFavoriteSystem)
 	{
 		FilterRowHBox->AddSlot()
-			.Padding(0)
+			.Padding(0.0f)
 			.HAlign(HAlign_Right)
 			.VAlign(VAlign_Center)
 			.AutoWidth()
@@ -335,14 +335,14 @@ void SDetailsView::Construct(const FArguments& InArgs, const FDetailsViewArgs& I
 		FilterImage->AddLayer(TAttribute<const FSlateBrush*>(this, &SDetailsView::GetViewOptionsBadgeIcon));
 		
 		FilterRowHBox->AddSlot()
-			.Padding(0)
+			.Padding(0.0f)
 			.HAlign(HAlign_Right)
 			.VAlign(VAlign_Center)
 			.AutoWidth()
 			[
 				SNew( SComboButton )
 				.HasDownArrow(false)
-				.ContentPadding(0)
+				.ContentPadding(0.0f)
 				.ForegroundColor( FSlateColor::UseForeground() )
 				.ButtonStyle( FAppStyle::Get(), "SimpleButton" )
 				.AddMetaData<FTagMetaData>(FTagMetaData(TEXT("ViewOptions")))
@@ -366,12 +366,12 @@ void SDetailsView::Construct(const FArguments& InArgs, const FDetailsViewArgs& I
 		];
 
 	FilterRowVBox->AddSlot()
-		.Padding(8, 2, 8, 7)
+		.Padding(8.0f, 2.0f, 8.0f, 7.0f)
 		.AutoHeight()
 		[
 			SAssignNew(SectionSelectorBox, SWrapBox)
 			.UseAllottedSize(true)
-			.InnerSlotPadding(FVector2D(4,4))
+			.InnerSlotPadding(FVector2D(4.0f,4.0f))
 		];
 
 	RebuildSectionSelector();
@@ -392,8 +392,8 @@ void SDetailsView::Construct(const FArguments& InArgs, const FDetailsViewArgs& I
 	}
 
 	VerticalBox->AddSlot()
-	.FillHeight(1)
-	.Padding(0)
+	.FillHeight(1.0f)
+	.Padding(0.0f)
 	[
 		SNew(SOverlay)
 		+ SOverlay::Slot()
@@ -1318,7 +1318,7 @@ void SDetailsView::RebuildSectionSelector()
 		-> TSharedRef<SWidget>
 	{
 		return SNew(SBox)
-			.Padding(FMargin(0))
+			.Padding(FMargin(0.0f))
 			[
 				SNew(SCheckBox)
 				.Style(FAppStyle::Get(), "DetailsView.SectionButton")
diff --git a/Engine/Source/Editor/PropertyEditor/Private/SPropertyTreeViewImpl.cpp b/Engine/Source/Editor/PropertyEditor/Private/SPropertyTreeViewImpl.cpp
--- a/Engine/Source/Editor/PropertyEditor/Private/SPropertyTreeViewImpl.cpp
+++ b/Engine/Source/Editor/PropertyEditor/Private/SPropertyTreeViewImpl.cpp
@@ -200,7 +200,7 @@ void SPropertyTreeViewImpl::ConstructPropertyTree()
 			SNew( SHorizontalBox )
 			+ SHorizontalBox::Slot()
 			.HAlign( HAlign_Fill )
-			.FillWidth(1)
+			.FillWidth(1.0f)
 			.Padding(0,0,3,0)
 			[
 				SAssignNew( FilterTextBox, SSearchBox )
@@ -215,7 +215,7 @@ void SPropertyTreeViewImpl::ConstructPropertyTree()
 				SNew( SButton )
 				.Visibility( bFavoritesAllowed ? EVisibility::Visible : EVisibility::Collapsed )
 				.OnClicked( this, &SPropertyTreeViewImpl::OnToggleFavoritesClicked )
-				.ContentPadding(1)
+				.ContentPadding(1.0f)
 				.ButtonStyle( FAppStyle::Get(), "NoBorder" )
 				[
 					SNew( SImage )
@@ -230,7 +230,7 @@ void SPropertyTreeViewImpl::ConstructPropertyTree()
 				SNew( SButton )
 				.Visibility( bLockable ? EVisibility::Visible : EVisibility::Collapsed )
 				.OnClicked( this, &SPropertyTreeViewImpl::OnLockButtonClicked )
-				.ContentPadding(1)
+				.ContentPadding(1.0f)
 				.ButtonStyle( FAppStyle::Get(), "NoBorder" )
 				[
 					SNew( SImage )
@@ -265,7 +265,7 @@ void SPropertyTreeViewImpl::ConstructPropertyTree()
 		]
 		+SVerticalBox::Slot()
 			.VAlign( VAlign_Fill )
-			.FillHeight(1)
+			.FillHeight(1.0f)
 			.Padding( 0.0f, PaddingAfterFilter, 0.0f, 0.0f )
 		[
 			SAssignNew( PropertyTree, SPropertyTree )
@@ -282,7 +282,7 @@ void SPropertyTreeViewImpl::ConstructPropertyTree()
 					.FillWidth(InitialNameColumnWidth)
 					[
 						SNew(SBorder)
-						.Padding(3)
+						.Padding(3.0f)
 						.BorderImage( FAppStyle::GetBrush("NoBorder") )
 						[
 							SNew(STextBlock)
@@ -293,7 +293,7 @@ void SPropertyTreeViewImpl::ConstructPropertyTree()
 					.FillWidth(1.0f)
 					[
 						SNew(SBorder)
-						.Padding(3)
+						.Padding(3.0f)
 						.BorderImage( FAppStyle::GetBrush("NoBorder") )
 						[
 							SNew(STextBlock)
diff --git a/Engine/Source/Editor/PropertyEditor/Private/SResetToDefaultMenu.cpp b/Engine/Source/Editor/PropertyEditor/Private/SResetToDefaultMenu.cpp
--- a/Engine/Source/Editor/PropertyEditor/Private/SResetToDefaultMenu.cpp
+++ b/Engine/Source/Editor/PropertyEditor/Private/SResetToDefaultMenu.cpp
@@ -41,7 +41,7 @@ void SResetToDefaultMenu::Construct( const FArguments& InArgs )
 		.ToolTipText(NSLOCTEXT( "PropertyEditor", "ResetToDefault", "Reset to Default" ))
 		.HasDownArrow(false)
 		.ButtonStyle( FAppStyle::Get(), "NoBorder" )
-		.ContentPadding(0) 
+		.ContentPadding(0.0f)
 		.Visibility( this, &SResetToDefaultMenu::GetResetToDefaultVisibility )
 		.OnGetMenuContent( this, &SResetToDefaultMenu::OnGenerateResetToDefaultMenuContent )
 		.ButtonContent()
diff --git a/Engine/Source/Editor/PropertyEditor/Private/SStructureDetailsView.cpp b/Engine/Source/Editor/PropertyEditor/Private/SStructureDetailsView.cpp
--- a/Engine/Source/Editor/PropertyEditor/Private/SStructureDetailsView.cpp
+++ b/Engine/Source/Editor/PropertyEditor/Private/SStructureDetailsView.cpp
@@ -147,7 +147,7 @@ void SStructureDetailsView::Construct(const FArguments& InArgs)
 			.AutoWidth()
 			[
 				SNew( SComboButton )
-				.ContentPadding(0)
+				.ContentPadding(0.0f)
 				.ForegroundColor( FSlateColor::UseForeground() )
 				.ButtonStyle( FAppStyle::Get(), "ToggleButton" )
 				.MenuContent()
@@ -185,8 +185,8 @@ void SStructureDetailsView::Construct(const FArguments& InArgs)
 				FilterBoxRow
 			]
 			+ SVerticalBox::Slot()
-			.FillHeight(1)
-			.Padding(0)
+			.FillHeight(1.0f)
+			.Padding(0.0f)
 			[
 				SNew( SOverlay )
 				+ SOverlay::Slot()
diff --git a/Engine/Source/Editor/PropertyEditor/Private/UserInterface/PropertyEditor/SPropertyAssetPicker.cpp b/Engine/Source/Editor/PropertyEditor/Private/UserInterface/PropertyEditor/SPropertyAssetPicker.cpp
--- a/Engine/Source/Editor/PropertyEditor/Private/UserInterface/PropertyEditor/SPropertyAssetPicker.cpp
+++ b/Engine/Source/Editor/PropertyEditor/Private/UserInterface/PropertyEditor/SPropertyAssetPicker.cpp
@@ -34,7 +34,7 @@ void SPropertyAssetPicker::Construct( const FArguments& InArgs )
 			.ButtonStyle( FAppStyle::Get(), "HoverHintOnly" )
 			.OnClicked( this, &SPropertyAssetPicker::OnClicked )
 			.ToolTipText(LOCTEXT("PickButtonLabel", "Pick Asset"))
-			.ContentPadding(0)
+			.ContentPadding(0.0f)
 			.ForegroundColor( FSlateColor::UseForeground() )
 			.IsFocusable(false)
 			[ 
@@ -83,8 +83,8 @@ TSharedRef<SWidget> SPropertyAssetPicker::OnGenerateAssetPicker()
 	
 	TSharedRef<SWidget> MenuContent = 
 		SNew(SBox)
-		.HeightOverride(300)
-		.WidthOverride(300)
+		.HeightOverride(300.0f)
+		.WidthOverride(300.0f)
 		[
 			SNew( SBorder )
 			.BorderImage( FAppStyle::GetBrush("Menu.Background") )
diff --git a/Engine/Source/Editor/PropertyEditor/Private/UserInterface/PropertyEditor/SPropertyEditorAsset.cpp b/Engine/Source/Editor/PropertyEditor/Private/UserInterface/PropertyEditor/SPropertyEditorAsset.cpp
--- a/Engine/Source/Editor/PropertyEditor/Private/UserInterface/PropertyEditor/SPropertyEditorAsset.cpp
+++ b/Engine/Source/Editor/PropertyEditor/Private/UserInterface/PropertyEditor/SPropertyEditorAsset.cpp
@@ -559,18 +559,18 @@ void SPropertyEditorAsset::Construct(const FArguments& InArgs, const TSharedPtr<
 		}
 
 		ValueContentBox->AddSlot()
-		.Padding(0,3,5,0)
+		.Padding(0.0f,3.0f,5.0f,0.0f)
 		.AutoWidth()
 		.VAlign(VAlign_Center)
 		[
 			SNew(SBorder)
 			.Visibility(EVisibility::SelfHitTestInvisible)
-			.Padding(FMargin(0, 0, 4, 4))
+			.Padding(FMargin(0.0f, 0.0f, 4.0f, 4.0f))
 			.BorderImage(FAppStyle::Get().GetBrush("PropertyEditor.AssetTileItem.DropShadow"))
 			[
 				SNew(SOverlay)
 				+SOverlay::Slot()
-				.Padding(1)
+				.Padding(1.0f)
 				[
 					SAssignNew(ThumbnailBorder, SBorder)
 					.Padding(0)
@@ -579,8 +579,8 @@ void SPropertyEditorAsset::Construct(const FArguments& InArgs, const TSharedPtr<
 					[
 						SNew(SBox)
 						.ToolTipText(TooltipAttribute)
-						.WidthOverride(InArgs._ThumbnailSize.X)
-						.HeightOverride(InArgs._ThumbnailSize.Y)
+						.WidthOverride(static_cast<float>(InArgs._ThumbnailSize.X))
+						.HeightOverride(static_cast<float>(InArgs._ThumbnailSize.Y))
 						[
 							AssetThumbnail->MakeThumbnailWidget(AssetThumbnailConfig)
 						]
@@ -597,7 +597,7 @@ void SPropertyEditorAsset::Construct(const FArguments& InArgs, const TSharedPtr<
 
 	
 		ValueContentBox->AddSlot()
-		.Padding(0)
+		.Padding(0.0f)
 		.VAlign(VAlign_Center)
 		[
 			SNew(SVerticalBox)
diff --git a/Engine/Source/Editor/PropertyEditor/Private/UserInterface/PropertyEditor/SPropertyEditorClass.cpp b/Engine/Source/Editor/PropertyEditor/Private/UserInterface/PropertyEditor/SPropertyEditorClass.cpp
--- a/Engine/Source/Editor/PropertyEditor/Private/UserInterface/PropertyEditor/SPropertyEditorClass.cpp
+++ b/Engine/Source/Editor/PropertyEditor/Private/UserInterface/PropertyEditor/SPropertyEditorClass.cpp
@@ -318,12 +318,12 @@ TSharedRef<SWidget> SPropertyEditorClass::GenerateClassPicker()
 	FOnClassPicked OnPicked(FOnClassPicked::CreateSP(this, &SPropertyEditorClass::OnClassPicked));
 
 	return SNew(SBox)
-		.WidthOverride(280)
+		.WidthOverride(280.0f)
 		[
 			SNew(SVerticalBox)
 			+SVerticalBox::Slot()
 			.AutoHeight()
-			.MaxHeight(500)
+			.MaxHeight(500.0f)
 			[
 				FModuleManager::LoadModuleChecked<FClassViewerModule>("ClassViewer").CreateClassViewer(ClassViewerOptions, OnPicked)
 			]			
diff --git a/Engine/Source/Editor/PropertyEditor/Private/UserInterface/PropertyEditor/SPropertyEditorEditInline.cpp b/Engine/Source/Editor/PropertyEditor/Private/UserInterface/PropertyEditor/SPropertyEditorEditInline.cpp
--- a/Engine/Source/Editor/PropertyEditor/Private/UserInterface/PropertyEditor/SPropertyEditorEditInline.cpp
+++ b/Engine/Source/Editor/PropertyEditor/Private/UserInterface/PropertyEditor/SPropertyEditorEditInline.cpp
@@ -155,7 +155,7 @@ void SPropertyEditorEditInline::Construct( const FArguments& InArgs, const TShar
 		SAssignNew(ComboButton, SComboButton)
 		.IsEnabled(this, &SPropertyEditorEditInline::IsValueEnabled, WeakHandlePtr)
 		.OnGetMenuContent(this, &SPropertyEditorEditInline::GenerateClassPicker)
-		.ContentPadding(0)
+		.ContentPadding(0.0f)
 		.ToolTipText(InPropertyEditor, &FPropertyEditor::GetValueAsText )
 		.ButtonContent()
 		[
diff --git a/Engine/Source/Editor/PropertyEditor/Private/UserInterface/PropertyEditor/SPropertyEditorStruct.cpp b/Engine/Source/Editor/PropertyEditor/Private/UserInterface/PropertyEditor/SPropertyEditorStruct.cpp
--- a/Engine/Source/Editor/PropertyEditor/Private/UserInterface/PropertyEditor/SPropertyEditorStruct.cpp
+++ b/Engine/Source/Editor/PropertyEditor/Private/UserInterface/PropertyEditor/SPropertyEditorStruct.cpp
@@ -195,12 +195,12 @@ TSharedRef<SWidget> SPropertyEditorStruct::GenerateStructPicker()
 	FOnStructPicked OnPicked(FOnStructPicked::CreateRaw(this, &SPropertyEditorStruct::OnStructPicked));
 
 	return SNew(SBox)
-		.WidthOverride(280)
+		.WidthOverride(280.0f)
 		[
 			SNew(SVerticalBox)
 			+SVerticalBox::Slot()
 			.AutoHeight()
-			.MaxHeight(500)
+			.MaxHeight(500.0f)
 			[
 				FModuleManager::LoadModuleChecked<FStructViewerModule>("StructViewer").CreateStructViewer(Options, OnPicked)
 			]			
diff --git a/Engine/Source/Editor/PropertyEditor/Private/UserInterface/PropertyEditor/SPropertyEditorTableRow.cpp b/Engine/Source/Editor/PropertyEditor/Private/UserInterface/PropertyEditor/SPropertyEditorTableRow.cpp
--- a/Engine/Source/Editor/PropertyEditor/Private/UserInterface/PropertyEditor/SPropertyEditorTableRow.cpp
+++ b/Engine/Source/Editor/PropertyEditor/Private/UserInterface/PropertyEditor/SPropertyEditorTableRow.cpp
@@ -72,13 +72,13 @@ TSharedRef< SWidget > SPropertyEditorTableRow::ConstructNameColumnWidget()
 		SNew(SHorizontalBox)
 		+SHorizontalBox::Slot()
 		.AutoWidth()
-		.Padding( FMargin( 0, 1, 0, 1 ) )
+		.Padding( FMargin( 0.0f, 1.0f, 0.0f, 1.0f ) )
 		.VAlign(VAlign_Center)
 		[
 			SNew(SExpanderArrow, SharedThis( this ) )
 		]
 		+SHorizontalBox::Slot()
-		.Padding( FMargin( 0, 1, 0, 1 ) )
+		.Padding( FMargin( 0.0f, 1.0f, 0.0f, 1.0f ) )
 		.AutoWidth()
 		.VAlign(VAlign_Center)
 		[
@@ -110,7 +110,7 @@ TSharedRef< SWidget > SPropertyEditorTableRow::ConstructValueColumnWidget()
 		.FillWidth(1) // Fill the entire width if possible
 		.VAlign(VAlign_Center)
 		.HAlign(HAlign_Left)
-		.Padding(20, 0, 0, 0)
+		.Padding(20.0f, 0.0f, 0.0f, 0.0f)
 		[
 			SNew(SBox)
 			.MinDesiredWidth( MinWidth )
@@ -130,7 +130,7 @@ TSharedRef< SWidget > SPropertyEditorTableRow::ConstructValueColumnWidget()
 			.ButtonStyle( FAppStyle::Get(), "NoBorder" )
 			.Visibility( this, &SPropertyEditorTableRow::OnGetFavoritesVisibility )
 			.OnClicked( this, &SPropertyEditorTableRow::OnToggleFavoriteClicked )
-			.ContentPadding(0)
+			.ContentPadding(0.0f)
 			[
 				SNew( SImage )
 				.Image( this, &SPropertyEditorTableRow::OnGetFavoriteImage )
@@ -153,7 +153,7 @@ TSharedRef< SWidget > SPropertyEditorTableRow::ConstructValueColumnWidget()
 	}
 
 	return SNew(SBorder)
-		.Padding( FMargin( 0, 1, 0, 1 ) )
+		.Padding( FMargin( 0.0f, 1.0f, 0.0f, 1.0f ) )
 		.BorderImage_Static( &PropertyEditorConstants::GetOverlayBrush, PropertyEditor.ToSharedRef() )
 		.VAlign(VAlign_Fill)
 		[
diff --git a/Engine/Source/Editor/PropertyEditor/Private/UserInterface/PropertyEditor/SPropertyMenuActorPicker.cpp b/Engine/Source/Editor/PropertyEditor/Private/UserInterface/PropertyEditor/SPropertyMenuActorPicker.cpp
--- a/Engine/Source/Editor/PropertyEditor/Private/UserInterface/PropertyEditor/SPropertyMenuActorPicker.cpp
+++ b/Engine/Source/Editor/PropertyEditor/Private/UserInterface/PropertyEditor/SPropertyMenuActorPicker.cpp
@@ -94,8 +94,8 @@ void SPropertyMenuActorPicker::Construct( const FArguments& InArgs )
 		
 		MenuContent =
 			SNew(SBox)
-			.WidthOverride(PropertyEditorAssetConstants::SceneOutlinerWindowSize.X)
-			.HeightOverride(PropertyEditorAssetConstants::SceneOutlinerWindowSize.Y)
+			.WidthOverride(static_cast<float>(PropertyEditorAssetConstants::SceneOutlinerWindowSize.X))
+			.HeightOverride(static_cast<float>(PropertyEditorAssetConstants::SceneOutlinerWindowSize.Y))
 			[
 				SceneOutlinerModule.CreateActorPicker(InitOptions, FOnActorPicked::CreateSP(this, &SPropertyMenuActorPicker::OnActorSelected), nullptr, !bAllowPickingLevelInstanceContent)
 			];
diff --git a/Engine/Source/Editor/PropertyEditor/Private/UserInterface/PropertyEditor/SPropertyMenuAssetPicker.cpp b/Engine/Source/Editor/PropertyEditor/Private/UserInterface/PropertyEditor/SPropertyMenuAssetPicker.cpp
--- a/Engine/Source/Editor/PropertyEditor/Private/UserInterface/PropertyEditor/SPropertyMenuAssetPicker.cpp
+++ b/Engine/Source/Editor/PropertyEditor/Private/UserInterface/PropertyEditor/SPropertyMenuAssetPicker.cpp
@@ -164,8 +164,8 @@ void SPropertyMenuAssetPicker::Construct( const FArguments& InArgs )
 
 		TSharedRef<SWidget> MenuContent =
 			SNew(SBox)
-			.WidthOverride(PropertyEditorAssetConstants::ContentBrowserWindowSize.X)
-			.HeightOverride(PropertyEditorAssetConstants::ContentBrowserWindowSize.Y)
+			.WidthOverride(static_cast<float>(PropertyEditorAssetConstants::ContentBrowserWindowSize.X))
+			.HeightOverride(static_cast<float>(PropertyEditorAssetConstants::ContentBrowserWindowSize.Y))
 			[
 				AssetPickerWidget.ToSharedRef()
 			];
diff --git a/Engine/Source/Editor/PropertyEditor/Private/UserInterface/PropertyEditor/SPropertyMenuComponentPicker.cpp b/Engine/Source/Editor/PropertyEditor/Private/UserInterface/PropertyEditor/SPropertyMenuComponentPicker.cpp
--- a/Engine/Source/Editor/PropertyEditor/Private/UserInterface/PropertyEditor/SPropertyMenuComponentPicker.cpp
+++ b/Engine/Source/Editor/PropertyEditor/Private/UserInterface/PropertyEditor/SPropertyMenuComponentPicker.cpp
@@ -126,8 +126,8 @@ void SPropertyMenuComponentPicker::Construct(const FArguments& InArgs)
 
 		MenuContent =
 			SNew(SBox)
-			.WidthOverride(PropertyEditorAssetConstants::SceneOutlinerWindowSize.X)
-			.HeightOverride(PropertyEditorAssetConstants::SceneOutlinerWindowSize.Y)
+			.WidthOverride(static_cast<float>(PropertyEditorAssetConstants::SceneOutlinerWindowSize.X))
+			.HeightOverride(static_cast<float>(PropertyEditorAssetConstants::SceneOutlinerWindowSize.Y))
 			[
 				SNew(SBorder)
 				.BorderImage(FAppStyle::GetBrush("Menu.Background"))
diff --git a/Engine/Source/Editor/PropertyEditor/Private/UserInterface/PropertyEditor/SPropertySceneOutliner.cpp b/Engine/Source/Editor/PropertyEditor/Private/UserInterface/PropertyEditor/SPropertySceneOutliner.cpp
--- a/Engine/Source/Editor/PropertyEditor/Private/UserInterface/PropertyEditor/SPropertySceneOutliner.cpp
+++ b/Engine/Source/Editor/PropertyEditor/Private/UserInterface/PropertyEditor/SPropertySceneOutliner.cpp
@@ -27,13 +27,13 @@ void SPropertySceneOutliner::Construct( const FArguments& InArgs )
 			.OnGetMenuContent( this, &SPropertySceneOutliner::OnGenerateSceneOutliner )
 		]
 		+ SVerticalBox::Slot()
-		.FillHeight(1)
+		.FillHeight(1.0f)
 		[
 			SNew( SButton )
 			.ButtonStyle( FAppStyle::Get(), "HoverHintOnly" )
 			.OnClicked( this, &SPropertySceneOutliner::OnClicked )
 			.ToolTipText(LOCTEXT("PickButtonLabel", "Pick Actor"))
-			.ContentPadding(0)
+			.ContentPadding(0.0f)
 			.ForegroundColor( FSlateColor::UseForeground() )
 			.IsFocusable(false)
 			[ 
@@ -60,8 +60,8 @@ TSharedRef<SWidget> SPropertySceneOutliner::OnGenerateSceneOutliner()
 
 	TSharedRef<SWidget> MenuContent = 
 		SNew(SBox)
-		.HeightOverride(300)
-		.WidthOverride(300)
+		.HeightOverride(300.0f)
+		.WidthOverride(300.0f)
 		[
 			SNew( SBorder )
 			.BorderImage( FAppStyle::GetBrush("Menu.Background") )
diff --git a/Engine/Source/Editor/PropertyEditor/Private/UserInterface/PropertyEditor/SResetToDefaultPropertyEditor.cpp b/Engine/Source/Editor/PropertyEditor/Private/UserInterface/PropertyEditor/SResetToDefaultPropertyEditor.cpp
--- a/Engine/Source/Editor/PropertyEditor/Private/UserInterface/PropertyEditor/SResetToDefaultPropertyEditor.cpp
+++ b/Engine/Source/Editor/PropertyEditor/Private/UserInterface/PropertyEditor/SResetToDefaultPropertyEditor.cpp
@@ -44,7 +44,7 @@ void SResetToDefaultPropertyEditor::Construct(const FArguments& InArgs, const TS
 		.IsFocusable(false)
 		.ToolTipText(this, &SResetToDefaultPropertyEditor::GetResetToolTip)
 		.ButtonStyle(FAppStyle::Get(), "SimpleButton")
-		.ContentPadding(0) 
+		.ContentPadding(0.0f)
 		.Visibility(this, &SResetToDefaultPropertyEditor::GetDiffersFromDefaultAsVisibility )
 		.OnClicked(this, &SResetToDefaultPropertyEditor::OnResetClicked)
 		.Content()
diff --git a/Engine/Source/Editor/PropertyEditor/Private/UserInterface/PropertyTree/SPropertyTreeCategoryRow.h b/Engine/Source/Editor/PropertyEditor/Private/UserInterface/PropertyTree/SPropertyTreeCategoryRow.h
--- a/Engine/Source/Editor/PropertyEditor/Private/UserInterface/PropertyTree/SPropertyTreeCategoryRow.h
+++ b/Engine/Source/Editor/PropertyEditor/Private/UserInterface/PropertyTree/SPropertyTreeCategoryRow.h
@@ -30,7 +30,7 @@ public:
 	{
 		STableRow< TSharedPtr< class FPropertyNode* > >::Construct(
 			STableRow< TSharedPtr< class FPropertyNode* > >::FArguments()
-				.Padding(3)
+				.Padding(3.0f)
 				[
 					SNew( STextBlock )
 					.Text( InArgs._DisplayName )
