Remove property buttons on the header row of statically sized arrays of class properties

diff --git a/Engine/Source/Editor/PropertyEditor/Private/PropertyEditorHelpers.cpp b/Engine/Source/Editor/PropertyEditor/Private/PropertyEditorHelpers.cpp
--- a/Engine/Source/Editor/PropertyEditor/Private/PropertyEditorHelpers.cpp
+++ b/Engine/Source/Editor/PropertyEditor/Private/PropertyEditorHelpers.cpp
@@ -692,6 +692,9 @@ namespace PropertyEditorHelpers
 		const FSetProperty* OuterSetProp = NodeProperty->GetOwner<FSetProperty>();
 		const FMapProperty* OuterMapProp = NodeProperty->GetOwner<FMapProperty>();
 
+		// Some buttons should be skipped for statically sized arrays
+		bool bStaticSizedArray = (NodeProperty->ArrayDim > 1) && (PropertyNode->GetArrayIndex() == -1);
+
 		//////////////////////////////
 		// Handle a container property.
 		if( NodeProperty->IsA(FArrayProperty::StaticClass()) || NodeProperty->IsA(FSetProperty::StaticClass()) || NodeProperty->IsA(FMapProperty::StaticClass()) )
@@ -709,31 +712,34 @@ namespace PropertyEditorHelpers
 		FSoftClassProperty* SoftClassProp = CastField<FSoftClassProperty>(NodeProperty);
 		if( ClassProp || SoftClassProp || IsSoftClassPath(NodeProperty))
 		{
-			OutRequiredButtons.Add( EPropertyButton::Use );			
-			OutRequiredButtons.Add( EPropertyButton::Browse );
-
-			UClass* Class = nullptr;
-			if (ClassProp)
-			{
-				Class = ClassProp->MetaClass;
-			}
-			else if (SoftClassProp)
-			{
-				Class = SoftClassProp->MetaClass;
-			}
-			else
+			if (!bStaticSizedArray)
 			{
-				Class = NodeProperty->GetOwnerProperty()->GetClassMetaData(TEXT("MetaClass"));
-			}
+				OutRequiredButtons.Add(EPropertyButton::Use);
+				OutRequiredButtons.Add(EPropertyButton::Browse);
 
-			if (Class && FKismetEditorUtilities::CanCreateBlueprintOfClass(Class) && !NodeProperty->HasMetaData("DisallowCreateNew"))
-			{
-				OutRequiredButtons.Add(EPropertyButton::NewBlueprint);
-			}
+				UClass* Class = nullptr;
+				if (ClassProp)
+				{
+					Class = ClassProp->MetaClass;
+				}
+				else if (SoftClassProp)
+				{
+					Class = SoftClassProp->MetaClass;
+				}
+				else
+				{
+					Class = NodeProperty->GetOwnerProperty()->GetClassMetaData(TEXT("MetaClass"));
+				}
 
-			if( !(NodeProperty->PropertyFlags & CPF_NoClear) )
-			{
-				OutRequiredButtons.Add( EPropertyButton::Clear );
+				if (Class && FKismetEditorUtilities::CanCreateBlueprintOfClass(Class) && !NodeProperty->HasMetaData("DisallowCreateNew"))
+				{
+					OutRequiredButtons.Add(EPropertyButton::NewBlueprint);
+				}
+
+				if (!(NodeProperty->PropertyFlags & CPF_NoClear))
+				{
+					OutRequiredButtons.Add(EPropertyButton::Clear);
+				}
 			}
 		}
 
@@ -756,7 +762,6 @@ namespace PropertyEditorHelpers
 		if( SupportsObjectPropertyButtons( NodeProperty, bUsingAssetPicker ) )
 		{
 			//ignore this node if the consistency check should happen for the children
-			bool bStaticSizedArray = (NodeProperty->ArrayDim > 1) && (PropertyNode->GetArrayIndex() == -1);
 			if (!bStaticSizedArray)
 			{
 				if( PropertyNode->HasNodeFlags(EPropertyNodeFlags::EditInlineNew) )
