Speculative fix Object name crash when no objects selected.

diff --git a/Engine/Source/Editor/PropertyEditor/Private/DetailCategoryBuilderImpl.cpp b/Engine/Source/Editor/PropertyEditor/Private/DetailCategoryBuilderImpl.cpp
--- a/Engine/Source/Editor/PropertyEditor/Private/DetailCategoryBuilderImpl.cpp
+++ b/Engine/Source/Editor/PropertyEditor/Private/DetailCategoryBuilderImpl.cpp
@@ -884,10 +884,13 @@ void FDetailCategoryImpl::InitializeObjectName()
 	{
 		if (const FObjectPropertyNode* Object = Node->AsObjectNode())
 		{
-			if (const UObject* CategoryObject = Object->GetUObject(0))
+			if (Object->GetNumObjects() > 0)
 			{
-				const FName Name{CategoryObject->GetName()};
-				ObjectName = Name;
+				if (const UObject* CategoryObject = Object->GetUObject(0))
+				{
+					const FName Name{ CategoryObject->GetName() };
+					ObjectName = Name;
+				}
 			}
 		}
 	}	
