Overrides widget support phase 2

- enabled code in details single item row to show the overrides widget
- added code to the overrides style keys to provide the visibility based on overrides state of the edit property chain (for properties) and object inputs
- Updates to the details display managers to get the property update widgets by providing the edit property chain, needed for override state for properties
- updated the overrides combo button code to iterate through the style keys and show any the have been initialized to show for the current override state
- currently conntains API changes to provide the override state for properties

diff --git a/Engine/Source/Editor/PropertyEditor/Private/DetailsDisplayManager.cpp b/Engine/Source/Editor/PropertyEditor/Private/DetailsDisplayManager.cpp
--- a/Engine/Source/Editor/PropertyEditor/Private/DetailsDisplayManager.cpp
+++ b/Engine/Source/Editor/PropertyEditor/Private/DetailsDisplayManager.cpp
@@ -150,8 +150,12 @@ bool FDetailsDisplayManager::AddEmptyCategoryToDetailLayoutIfNeeded(TSharedRef<F
 	return false;
 }
 
-TSharedPtr<FPropertyUpdatedWidgetBuilder> FDetailsDisplayManager::GetPropertyUpdatedWidget(FResetToDefault ResetToDefault,
-	bool bIsCategoryUpdateWidget, FName InCategoryObjectName)
+TSharedPtr<FPropertyUpdatedWidgetBuilder> FDetailsDisplayManager::GetPropertyUpdatedWidget(FResetToDefault ResetToDefault, bool bIsCategoryUpdateWidget, FName InCategoryObjectName )
+{
+	return nullptr;
+}
+
+TSharedPtr<FPropertyUpdatedWidgetBuilder> FDetailsDisplayManager::GetPropertyUpdatedWidget(FResetToDefault ResetToDefault, TSharedRef<FEditPropertyChain> Chain, FName InCategoryObjectName)
 {
 	return nullptr;
 }
diff --git a/Engine/Source/Editor/PropertyEditor/Private/DetailsViewStyle.cpp b/Engine/Source/Editor/PropertyEditor/Private/DetailsViewStyle.cpp
--- a/Engine/Source/Editor/PropertyEditor/Private/DetailsViewStyle.cpp
+++ b/Engine/Source/Editor/PropertyEditor/Private/DetailsViewStyle.cpp
@@ -7,18 +7,20 @@
 #include "Styling/StyleColors.h"
 #include "Styling/StarshipCoreStyle.h"
 #include "Brushes/SlateRoundedBoxBrush.h"
+#include "UObject/OverridableManager.h"
+#include "Layout/Visibility.h"
+
+TArray< TSharedRef< const FOverridesWidgetStyleKey >> FOverridesWidgetStyleKeys::OverridesWidgetStyleKeys;
 
 const FSlateBrush& FOverridesWidgetStyleKey::GetConstStyleBrush() const
 {
 	return ImageBrush;
 }
 
-const FComboButtonStyle& FDetailsViewStyle::GetOverridesComboButtonStyle(
-										const FOverridesWidgetStyleKey* OverridesWidgetStyleKey,
-										const bool bIsOverridesWidgetForOuterCategory) const
+const FComboButtonStyle& FOverridesWidgetStyleKey::GetComboButtonStyle( const bool bIsOverridesWidgetForOuterCategory ) const
 {
 	static TMap<const FOverridesWidgetStyleKey*, const FComboButtonStyle> OverridesKeyToComboButtonStyleMap;
-	const FComboButtonStyle* ComboButtonStylePtr = OverridesKeyToComboButtonStyleMap.Find(OverridesWidgetStyleKey);
+	const FComboButtonStyle* ComboButtonStylePtr = OverridesKeyToComboButtonStyleMap.Find( this );
 
 	if ( ComboButtonStylePtr )
 	{
@@ -39,21 +41,31 @@ const FComboButtonStyle& FDetailsViewStyle::GetOverridesComboButtonStyle(
 	                                     .SetPressedPadding(FMargin(2.f, 0.f, 0.f, 0.f));
 			
 
-	OverridesKeyToComboButtonStyleMap.Add(OverridesWidgetStyleKey, FComboButtonStyle(FStarshipCoreStyle::GetCoreStyle().GetWidgetStyle<FComboButtonStyle>("ComboButton"))
+	OverridesKeyToComboButtonStyleMap.Add( this, FComboButtonStyle(FStarshipCoreStyle::GetCoreStyle().GetWidgetStyle<FComboButtonStyle>("ComboButton"))
 												   .SetButtonStyle(OverridesButton)
-												   .SetDownArrowImage(OverridesWidgetStyleKey->GetConstStyleBrush())
+												   .SetDownArrowImage(GetConstStyleBrush())
 												   .SetDownArrowPadding(FMargin(2.f, 5.f, 3.f, 5.f)));
 
 		
-	return *OverridesKeyToComboButtonStyleMap.Find(OverridesWidgetStyleKey);
+	return *OverridesKeyToComboButtonStyleMap.Find( this );
 }
 
 const FOverridesWidgetStyleKey& FOverridesWidgetStyleKeys::Here()
 {
-	static const FOverridesWidgetStyleKey Here{"OverrideHere"};
+	static const FOverridesWidgetStyleKey Here{"OverrideHere", EOverriddenPropertyOperation::Replace, EOverriddenState::AllOverridden };
 	return Here;
 }
 
+void FOverridesWidgetStyleKeys::Initialize()
+{
+	OverridesWidgetStyleKeys.Add( MakeShared<FOverridesWidgetStyleKey>(Here()));
+	OverridesWidgetStyleKeys.Add( MakeShared<FOverridesWidgetStyleKey>(Inside()));
+	OverridesWidgetStyleKeys.Add( MakeShared<FOverridesWidgetStyleKey>(HereInside()));
+	OverridesWidgetStyleKeys.Add( MakeShared<FOverridesWidgetStyleKey>(Added()));
+	OverridesWidgetStyleKeys.Add( MakeShared<FOverridesWidgetStyleKey>(Removed()));
+	OverridesWidgetStyleKeys.Add( MakeShared<FOverridesWidgetStyleKey>(Options()));	
+}
+
 const FOverridesWidgetStyleKey& FOverridesWidgetStyleKeys::Added()
 {
 	static const FOverridesWidgetStyleKey Added{"OverrideAdded"};
@@ -74,17 +86,22 @@ const FOverridesWidgetStyleKey& FOverridesWidgetStyleKeys::Removed()
 
 const FOverridesWidgetStyleKey& FOverridesWidgetStyleKeys::Inside()
 {
-	static const FOverridesWidgetStyleKey Inside{"OverrideInside"};
+	static const FOverridesWidgetStyleKey Inside{"OverrideInside", EOverriddenPropertyOperation::Modified, EOverriddenState::HasOverrides };
 	return Inside;
 }
 
 const FOverridesWidgetStyleKey& FOverridesWidgetStyleKeys::HereInside()
 {
-	static const FOverridesWidgetStyleKey HereInside{"OverrideHereInside"};
+	static const FOverridesWidgetStyleKey HereInside{"OverrideHereInside" };
 	return HereInside;
 }
 
-FOverridesWidgetStyleKey::FOverridesWidgetStyleKey(FName InName) : Name{InName}
+TArray< TSharedRef< const FOverridesWidgetStyleKey >> FOverridesWidgetStyleKeys::GetKeys()
+{
+	return OverridesWidgetStyleKeys;
+}
+
+void FOverridesWidgetStyleKey::Construct()
 {
 	static const FVector2D Icon16x16{16.0f, 16.0f};
 	static const FString Path = FPaths::ConvertRelativePathToFull(FPaths::Combine(FPaths::EngineContentDir(), TEXT("Slate/Starship/Common/")));
@@ -94,6 +111,50 @@ FOverridesWidgetStyleKey::FOverridesWidgetStyleKey(FName InName) : Name{InName}
 	ImageBrush = *SlateBrushPtr;
 }
 
+
+TAttribute<EVisibility> FOverridesWidgetStyleKey::GetVisibilityAttribute(const TSharedPtr<FEditPropertyChain>& PropertyChain, TWeakObjectPtr<UObject>& OverriddenObjectWeakPtr) const
+{
+	static FOverridableManager& Manager = FOverridableManager::Get();
+	
+	return TAttribute<EVisibility>::CreateLambda( [ this, PropertyChain,  OverriddenObjectWeakPtr ] ()
+	{
+		const bool bIsProperty = PropertyChain.IsValid();
+		
+	    if ( OverriddenObjectWeakPtr.IsValid() )
+	    {
+	    	UObject& OverriddenObject = *OverriddenObjectWeakPtr.Get();
+
+	    	const bool bIsPropertyVisible = bIsProperty &&
+												VisibleOverriddenPropertyOperation == Manager.GetOverriddenPropertyOperation(OverriddenObject, *PropertyChain.Get());
+
+			const bool bIsComponentVisible = !bIsProperty && VisibleOverriddenState == Manager.GetOverriddenState(OverriddenObject);
+
+			if ( bIsComponentVisible || bIsPropertyVisible )
+			{
+				return EVisibility::Visible;
+			}
+	    }
+
+		return EVisibility::Collapsed;
+	});
+}
+
+FOverridesWidgetStyleKey::FOverridesWidgetStyleKey(FName InName) : Name{InName}
+{
+     Construct();
+}
+
+FOverridesWidgetStyleKey::FOverridesWidgetStyleKey(FName InName,
+																				EOverriddenPropertyOperation InOverriddenPropertyOperation,
+																				EOverriddenState InOverriddenState) :
+   Name(InName),
+   VisibleOverriddenPropertyOperation(InOverriddenPropertyOperation),
+   VisibleOverriddenState(InOverriddenState),
+   bCanBeVisible(true)
+{
+	Construct();
+}
+
 FDetailsViewStyle::FDetailsViewStyle()
 {
 }
@@ -229,6 +290,8 @@ void FDetailsViewStyle::InitializeDetailsViewStyles()
 	
 	static FDetailsViewStyle ClassicStyle{FDetailsViewStyleKeys::Classic(), 0.0f };
 	StyleKeyToStyleTemplateMap.Add(FDetailsViewStyleKeys::Classic().GetName(), &ClassicStyle);
+
+	FOverridesWidgetStyleKeys::Initialize();
 }
 
 const FSlateBrush* FDetailsViewStyle::GetBackgroundImageForScrollBarWell(
diff --git a/Engine/Source/Editor/PropertyEditor/Private/DetailsViewStyle.h b/Engine/Source/Editor/PropertyEditor/Private/DetailsViewStyle.h
--- a/Engine/Source/Editor/PropertyEditor/Private/DetailsViewStyle.h
+++ b/Engine/Source/Editor/PropertyEditor/Private/DetailsViewStyle.h
@@ -10,36 +10,87 @@
 #include "Misc/Paths.h"
 #include "Styling/SlateWidgetStyle.h"
 
+enum class EOverriddenState : uint8;
+struct EVisibility;
+
 /**
  * A class which provides a key with the information to create the Overrides widget style (including the Icom)
  */
-class FOverridesWidgetStyleKey
+class FOverridesWidgetStyleKey : public TSharedFromThis< FOverridesWidgetStyleKey >
 {
 public:
 
 	/**
-	 * The constructor for the style key.
+	 * The default constructor for the style key.
 	 *
 	 * @param InName the FName which is the name of the key
 	 */
 	PROPERTYEDITOR_API FOverridesWidgetStyleKey(FName InName);
 
+	/**
+ * The constructor for the style key that initializes the Overridden state members
+ *
+ * @param InName the FName which is the name of the key
+ * @param InOverriddenPropertyOperation the overridden property operation for which this style is visible
+ * @param InOverriddenState the overridden state for components for which this style is visible
+ */
+	PROPERTYEDITOR_API FOverridesWidgetStyleKey(FName InName,
+																			EOverriddenPropertyOperation InOverriddenPropertyOperation,
+																			EOverriddenState InOverriddenState
+																			);
+
+	/**
+	 * returns the const FSlateBrush& that creates the icon for this widget style
+	 */
 	const FSlateBrush& GetConstStyleBrush() const;
 
+	/**
+	 * Creates a TAttribute<EVisibility> which indicates the EVisibility for this widget style
+	 * 
+	 * @param PropertyChain points to the FEditPropertyChain for the Property, if this style is for a Property, else it is nullptr
+	 * @param OverriddenObject the UObject associated with this widget style
+	 * @return 
+	 */
+	TAttribute<EVisibility> GetVisibilityAttribute(const TSharedPtr<FEditPropertyChain>& PropertyChain, TWeakObjectPtr<UObject>& OverriddenObjectWeakPtr) const;
+	
+	/**
+	 * returns the const FComboButtonStyle& for this OverridesWidgetStyleKey
+	 *
+	 * @param bIsForOuterCategory true if this style is for an outer Category
+	 */
+	PROPERTYEDITOR_API const FComboButtonStyle& GetComboButtonStyle(const bool bIsForOuterCategory = false) const;
+
 	/**
 	 * the name of the key
 	 */
 	const FName Name;
 
+	/**
+	 * The EOverriddenPropertyOperation for which this style is visible for a property
+	 */
+	const TOptional<EOverriddenPropertyOperation> VisibleOverriddenPropertyOperation;
+
+	/**
+	 * The EOverriddenState for which this style is visible for a UObject
+	 */
+	const TOptional<EOverriddenState> VisibleOverriddenState;
+
+	const bool bCanBeVisible = false;
+
 private:
 
 	/**
 	 * The image brush specified by this style key
 	 */
 	FSlateBrush ImageBrush;
+
+	/**
+	 * Construct the style key
+	 */
+	void Construct();
 };
 
-/**
+/**s
  * The FOverridesWidgetStyleKeys class provides style keys which can
  * create the needed styles for overrides widgets
  */
@@ -77,6 +128,12 @@ public:
 
 	static const FOverridesWidgetStyleKey& HereInside();
 
+	static TArray< TSharedRef< const FOverridesWidgetStyleKey >> GetKeys();
+	
+	static void Initialize(); 
+
+private:
+	static TArray< TSharedRef< const FOverridesWidgetStyleKey >> OverridesWidgetStyleKeys;
 };
 
 /**
@@ -86,13 +143,6 @@ class FDetailsViewStyle : public FSlateWidgetStyle
 {
 public:
 
-	/**
-	 * returns the const FComboButtonStyle& for the key const FOverridesWidgetStyleKey* OverridesWidgetStyleKey
-	 *
-	 * @param OverridesWidgetStyleKey the FOverridesWidgetStyleKey for the overrides combobutton style
-	 */
-	PROPERTYEDITOR_API const FComboButtonStyle& GetOverridesComboButtonStyle(const FOverridesWidgetStyleKey* OverridesWidgetStyleKey, const bool bIsForOuterCategory = false) const;
-
 	/**
 	 * The default constructor of this @code FDetailsViewStyle @endcode.
 	 */
diff --git a/Engine/Source/Editor/PropertyEditor/Private/SDetailSingleItemRow.cpp b/Engine/Source/Editor/PropertyEditor/Private/SDetailSingleItemRow.cpp
--- a/Engine/Source/Editor/PropertyEditor/Private/SDetailSingleItemRow.cpp
+++ b/Engine/Source/Editor/PropertyEditor/Private/SDetailSingleItemRow.cpp
@@ -634,11 +634,12 @@ void SDetailSingleItemRow::Construct( const FArguments& InArgs, FDetailLayoutCus
 				DetailsView &&
 				DetailsView->GetDisplayManager().IsValid())
 			{
-					
-				Property = GetPropertyNode()->GetProperty();
+				TSharedPtr<FPropertyNode>  PropertyNode = GetPropertyNode(); 	
+				Property = PropertyNode->GetProperty();
 				DisplayManager = DetailsView->GetDisplayManager();
+				TSharedRef<FEditPropertyChain> EditPropertyChain = PropertyNode->BuildPropertyChain( Property ); 
 				PropertyUpdatedWidgetBuilder = DisplayManager->GetPropertyUpdatedWidget(
-					FExecuteAction::CreateSP(this, &SDetailSingleItemRow::OnResetToDefaultClicked));
+					FExecuteAction::CreateSP(this, &SDetailSingleItemRow::OnResetToDefaultClicked), EditPropertyChain, Category->GetObjectName());
 
 				if (PropertyUpdatedWidgetBuilder.IsValid())
 				{
@@ -660,20 +661,8 @@ void SDetailSingleItemRow::Construct( const FArguments& InArgs, FDetailLayoutCus
 				.Padding(0.0f)
 				[
 					PropertyUpdatedWidgetBuilder.IsValid() ?
-					PropertyUpdatedWidgetBuilder
-						->Bind_IsVisible(
-							TAttribute<EVisibility>::CreateLambda([this, Category, Property ]
-							{
-								const bool bIsResetVisible = IsResetToDefaultVisible();
-								const EVisibility Visibility = bIsResetVisible ? EVisibility::Visible : EVisibility::Collapsed;
-								if (DisplayManager.IsValid() && Category.IsValid() && Property)
-								{
-									DisplayManager->UpdatePropertyForCategory(Category->GetObjectName(), Property, bIsResetVisible);
-								}
-								return Visibility;
-							}))
-						.GenerateWidget().ToSharedRef()  :
-					ToolbarBuilder.MakeWidget()
+					         PropertyUpdatedWidgetBuilder->GenerateWidget().ToSharedRef()  :
+					         ToolbarBuilder.MakeWidget()
 				]
 			];
 		}
diff --git a/Engine/Source/Editor/PropertyEditor/Private/UserInterface/Widgets/OverridesComboButtonBuilder.cpp b/Engine/Source/Editor/PropertyEditor/Private/UserInterface/Widgets/OverridesComboButtonBuilder.cpp
--- a/Engine/Source/Editor/PropertyEditor/Private/UserInterface/Widgets/OverridesComboButtonBuilder.cpp
+++ b/Engine/Source/Editor/PropertyEditor/Private/UserInterface/Widgets/OverridesComboButtonBuilder.cpp
@@ -14,7 +14,8 @@ FOverridesComboButtonBuilder::FOverridesComboButtonBuilder(
 	FPropertyUpdatedWidgetBuilder(),
 	DisplayManager(InDetailsDisplayManager),
 	bIsCategoryOverridesComboButton(bInIsCategoryOverridesComboButton),
-    Object(InObject)
+    Object(InObject),
+    EditPropertyChain(nullptr)
 {
 }
 
@@ -27,34 +28,33 @@ FOverridesComboButtonBuilder& FOverridesComboButtonBuilder::Set_OnGetContent(FOn
 
 TSharedPtr<SWidget> FOverridesComboButtonBuilder::GenerateWidget()
 {
-    const FDetailsViewStyle* DetailsViewStyle = DisplayManager->GetDetailsViewStyle();
+	const TArray< TSharedRef< const FOverridesWidgetStyleKey >> OverridesWidgetKeys = FOverridesWidgetStyleKeys::GetKeys();
+	
+	if ( !Object.IsValid() || OverridesWidgetKeys.IsEmpty() )
+	{
+		return SNullWidget::NullWidget;
+	}
 
-	return	SNew(SHorizontalBox)
-				.Visibility(IsVisible)
-				+SHorizontalBox::Slot()
-					
-				.HAlign(HAlign_Right)
-				.VAlign(VAlign_Center)
-				.AutoWidth()
-				[
-					SNew( SComboButton )
-					.ComboButtonStyle( &DetailsViewStyle->GetOverridesComboButtonStyle(
-										&FOverridesWidgetStyleKeys::Here(), bIsCategoryOverridesComboButton ) )
-					.Visibility(this, &FOverridesComboButtonBuilder::GetFullyOverridenVisibility)
+	TSharedRef<SHorizontalBox> Box =	SNew(SHorizontalBox).Visibility(IsVisible);
+	
+	for (const TSharedRef< const FOverridesWidgetStyleKey >& Key : OverridesWidgetKeys )
+	{
+		if ( Key->bCanBeVisible )
+		{
+			Box->AddSlot()
+			   .HAlign(HAlign_Right)
+			   .VAlign(VAlign_Center)
+			   .AutoWidth()
+			[
+				SNew( SComboButton )
+					.ComboButtonStyle( &Key->GetComboButtonStyle(  bIsCategoryOverridesComboButton )  )
+					.Visibility( Key->GetVisibilityAttribute(EditPropertyChain, Object ))
 					.HasDownArrow(true)
-				]
-				+SHorizontalBox::Slot()
-				.HAlign(HAlign_Right)
-				.VAlign(VAlign_Center)
-				.AutoWidth()
-				[
-					SNew( SComboButton )
-					.ComboButtonStyle( &DetailsViewStyle->GetOverridesComboButtonStyle(
-					                                    &FOverridesWidgetStyleKeys::Inside()
-                										, bIsCategoryOverridesComboButton ) )
-				                    .Visibility(this, &FOverridesComboButtonBuilder::GetOverridenInsideVisibility)
-                					.HasDownArrow(true)
-				];
+			];
+		}
+	}
+
+	return Box.ToSharedPtr();
 }
 
 TSharedRef<SWidget> FOverridesComboButtonBuilder::operator*()
@@ -67,31 +67,8 @@ FOverridesComboButtonBuilder::~FOverridesComboButtonBuilder()
 	OnGetContent.Unbind();
 }
 
-EVisibility FOverridesComboButtonBuilder::GetOverridenInsideVisibility() const
+void FOverridesComboButtonBuilder:: SetEditPropertyChain(TSharedRef<FEditPropertyChain>& InEditPropertyChain)
 {
-	static FOverridableManager& Manager = FOverridableManager::Get();
-	
-	if ( UObject* OverridableComponent = Object.Get() )
-	{
-		const EOverriddenState State = Manager.GetOverriddenState(*OverridableComponent);
-		if (State != EOverriddenState::NoOverrides && State != EOverriddenState::AllOverridden)
-		{
-			return EVisibility::Visible;
-		}		
-	}
-	return EVisibility::Collapsed;
+	EditPropertyChain = InEditPropertyChain;	
+	bIsCategory = false;
 }
-
-EVisibility FOverridesComboButtonBuilder::GetFullyOverridenVisibility() const
-{
-	if ( UObject* OverridableObject = Object.Get() )
-	{
-		static FOverridableManager& Manager = FOverridableManager::Get();
-		const EOverriddenState State = Manager.GetOverriddenState(*OverridableObject);
-		if (State == EOverriddenState::AllOverridden)
-		{
-			return EVisibility::Visible;
-		}
-	}
-	return EVisibility::Collapsed;
-}
\ No newline at end of file
diff --git a/Engine/Source/Editor/PropertyEditor/Public/DetailsDisplayManager.h b/Engine/Source/Editor/PropertyEditor/Public/DetailsDisplayManager.h
--- a/Engine/Source/Editor/PropertyEditor/Public/DetailsDisplayManager.h
+++ b/Engine/Source/Editor/PropertyEditor/Public/DetailsDisplayManager.h
@@ -107,8 +107,24 @@ public:
 	 */	
 	PROPERTYEDITOR_API virtual bool AddEmptyCategoryToDetailLayoutIfNeeded(TSharedRef<FComplexPropertyNode> Node, TSharedRef<FDetailLayoutBuilderImpl> DetailLayoutBuilder);
 
+	/**
+	* Returns a widget which will show in place of the reset to default button, or a nullptr if the default reset button should be used
+	*
+	* @param ResetToDefault the delegate which should be called to reset the row to default
+	* @param bIsCategoryUpdateWidget if true this widget builder is for a Category instead of a property within
+	* @param InCategoryObjectName  the name of the UObject associated with the Category for the widget builder, if one exists, else it is NAME_NONE
+	*/
 	virtual TSharedPtr<FPropertyUpdatedWidgetBuilder> GetPropertyUpdatedWidget(FResetToDefault ResetToDefault, bool bIsCategoryUpdateWidget = false, FName InCategoryObjectName = NAME_None);
 
+	/**
+	* Returns a widget which will show in place of the reset to default button, or a nullptr if the default reset button should be used
+	*
+	* @param ResetToDefault the delegate which should be called to reset the row to default
+	* @param InEditorPropertyChain the FEditorPropertyChain for the FPropertyNode whose state is visualized by this  property updated widget 
+	* @param InCategoryObjectName the name of the UObject for which the Category is displayed, if one is associated with the Category
+	*/
+	virtual TSharedPtr<FPropertyUpdatedWidgetBuilder> GetPropertyUpdatedWidget(FResetToDefault ResetToDefault, TSharedRef<FEditPropertyChain> InEditorPropertyChain, FName InCategoryObjectName);
+
 	void UpdatePropertyForCategory(FName InCategoryObjectName, FProperty* Property, bool bAddProperty);
 
 	bool GetCategoryHasAnyUpdatedProperties(FName InCategoryObjectName) const;
diff --git a/Engine/Source/Editor/PropertyEditor/Public/UserInterface/Widgets/OverridesComboButtonBuilder.h b/Engine/Source/Editor/PropertyEditor/Public/UserInterface/Widgets/OverridesComboButtonBuilder.h
--- a/Engine/Source/Editor/PropertyEditor/Public/UserInterface/Widgets/OverridesComboButtonBuilder.h
+++ b/Engine/Source/Editor/PropertyEditor/Public/UserInterface/Widgets/OverridesComboButtonBuilder.h
@@ -7,6 +7,7 @@
 #include "Widgets/Input/SComboButton.h"
 #include "Templates/SharedPointer.h"
 
+enum class EOverriddenState : uint8;
 class FDetailsDisplayManager;
 
 
@@ -42,6 +43,11 @@ public:
 	 */
 	PROPERTYEDITOR_API FOverridesComboButtonBuilder& Set_OnGetContent(FOnGetContent InOnGetContent);
 
+	/**
+	 * Sets the FEditPropertyChain for this FOverridesComboButtonBuilder
+	 */
+	PROPERTYEDITOR_API void SetEditPropertyChain(TSharedRef<FEditPropertyChain>& InEditPropertyChain);
+
 	/**
 	 * Implements the generation of the Category Menu button SWidget
 	 */
@@ -58,15 +64,6 @@ public:
 	TSharedRef<SWidget> operator*();
 
 private:
-	/**
-	 * Returns the visibility of the image that denotes that a Component is fully overridden
-	 */
-	EVisibility GetFullyOverridenVisibility() const;
-
-	/**
-	 * Returns the visibility of the image that denotes a Component is overridden only inside
-	 */
-	EVisibility GetOverridenInsideVisibility() const;
 
 	/**
 	 * The @code DetailsDisplayManager @endcode which provides an API to manage some of the characteristics of the
@@ -83,4 +80,17 @@ private:
 	 * The UObject that will be queried for its override state
 	 */
 	TWeakObjectPtr<UObject> Object;
+	
+	/**
+    * If this is a ComboButton builder for a Property, this will point to the FEditPropertyChain for that Property, else it will be nullptr
+    * 
+    */
+	TSharedPtr<FEditPropertyChain> EditPropertyChain;
+
+	/**
+	 * if true this is a builder for an Overrides combo button for a Category, else it is a builder
+	 * for Properties within a Category 
+	 */
+	bool bIsCategory = true;
+
 };
