Fixed per-track compression to not allow selection of invalid scale formats

Also fixes various property metadata helpers when dealing with enums held in containers

diff --git a/Engine/Source/Editor/PropertyEditor/Private/PropertyEditorHelpers.cpp b/Engine/Source/Editor/PropertyEditor/Private/PropertyEditorHelpers.cpp
--- a/Engine/Source/Editor/PropertyEditor/Private/PropertyEditorHelpers.cpp
+++ b/Engine/Source/Editor/PropertyEditor/Private/PropertyEditorHelpers.cpp
@@ -1137,11 +1137,12 @@ namespace PropertyEditorHelpers
 		TArray<FName> ValidEnumValues;
 
 		static const FName ValidEnumValuesName("ValidEnumValues");
-		
-		if (Property->HasMetaData(ValidEnumValuesName))
+
+		const FProperty* OwnerProperty = Property->GetOwnerProperty();
+		if (OwnerProperty->HasMetaData(ValidEnumValuesName))
 		{
 			TArray<FString> ValidEnumValuesAsString;
-			Property->GetMetaData(ValidEnumValuesName).ParseIntoArray(ValidEnumValuesAsString, TEXT(","));
+			OwnerProperty->GetMetaData(ValidEnumValuesName).ParseIntoArray(ValidEnumValuesAsString, TEXT(","));
 			
 			for (FString& Value : ValidEnumValuesAsString)
 			{
@@ -1158,11 +1159,13 @@ namespace PropertyEditorHelpers
 		TArray<FName> InvalidEnumValues;
 
 		static const FName InvalidEnumValuesName("InvalidEnumValues");
-		if (Property->HasMetaData(InvalidEnumValuesName))
+
+		const FProperty* OwnerProperty = Property->GetOwnerProperty();
+		if (OwnerProperty->HasMetaData(InvalidEnumValuesName))
 		{
 			TArray<FString> InvalidEnumValuesAsString;
 
-			Property->GetMetaData(InvalidEnumValuesName).ParseIntoArray(InvalidEnumValuesAsString, TEXT(","));
+			OwnerProperty->GetMetaData(InvalidEnumValuesName).ParseIntoArray(InvalidEnumValuesAsString, TEXT(","));
 			for (FString& Value : InvalidEnumValuesAsString)
 			{
 				Value.TrimStartInline();
@@ -1179,10 +1182,11 @@ namespace PropertyEditorHelpers
 
 		static const FName GetRestrictedEnumValuesName("GetRestrictedEnumValues");
 		TArray<FString> ValidEnumValuesAsString;
-		
-		if (Property->HasMetaData(GetRestrictedEnumValuesName))
+
+		const FProperty* OwnerProperty = Property->GetOwnerProperty();
+		if (OwnerProperty->HasMetaData(GetRestrictedEnumValuesName))
 		{
-			const FString GetRestrictedEnumValuesNameFunctionName = Property->GetMetaData(GetRestrictedEnumValuesName);
+			const FString GetRestrictedEnumValuesNameFunctionName = OwnerProperty->GetMetaData(GetRestrictedEnumValuesName);
 			if (!GetRestrictedEnumValuesNameFunctionName.IsEmpty())
 			{
 				for (UObject* Object : ObjectList)
@@ -1213,7 +1217,8 @@ namespace PropertyEditorHelpers
 
 		static const FName NAME_EnumValueDisplayNameOverrides = "EnumValueDisplayNameOverrides";
 
-		const FString& DisplayNameOverridesStr = Property->GetMetaData(NAME_EnumValueDisplayNameOverrides);
+		const FProperty* OwnerProperty = Property->GetOwnerProperty();
+		const FString& DisplayNameOverridesStr = OwnerProperty->GetMetaData(NAME_EnumValueDisplayNameOverrides);
 		if (DisplayNameOverridesStr.Len() > 0)
 		{
 			TArray<FString> DisplayNameOverridePairs;
diff --git a/Engine/Source/Runtime/Engine/Classes/Animation/AnimCompress_PerTrackCompression.h b/Engine/Source/Runtime/Engine/Classes/Animation/AnimCompress_PerTrackCompression.h
--- a/Engine/Source/Runtime/Engine/Classes/Animation/AnimCompress_PerTrackCompression.h
+++ b/Engine/Source/Runtime/Engine/Classes/Animation/AnimCompress_PerTrackCompression.h
@@ -45,7 +45,7 @@ class UAnimCompress_PerTrackCompression : public UAnimCompress_RemoveLinearKeys
 	TArray<TEnumAsByte<enum AnimationCompressionFormat> > AllowedTranslationFormats;
 
 	/** Which encoding formats is the per-track compressor allowed to try on scale keys */
-	UPROPERTY(EditAnywhere, Category=PerTrack)
+	UPROPERTY(EditAnywhere, Category=PerTrack, meta = (InvalidEnumValues = "ACF_Fixed32NoW,ACF_Float32NoW"))
 	TArray<TEnumAsByte<enum AnimationCompressionFormat> > AllowedScaleFormats;
 
 	/** If true, resample the animation to ResampleFramerate frames per second */
