Fix build break

diff --git a/Engine/Source/Editor/PropertyEditor/Private/DetailsViewStyle.cpp b/Engine/Source/Editor/PropertyEditor/Private/DetailsViewStyle.cpp
--- a/Engine/Source/Editor/PropertyEditor/Private/DetailsViewStyle.cpp
+++ b/Engine/Source/Editor/PropertyEditor/Private/DetailsViewStyle.cpp
@@ -125,7 +125,7 @@ TAttribute<EVisibility> FOverridesWidgetStyleKey::GetVisibilityAttribute(const T
 	    	UObject& OverriddenObject = *OverriddenObjectWeakPtr.Get();
 
 	    	const bool bIsPropertyVisible = bIsProperty &&
-												VisibleOverriddenPropertyOperation == Manager.GetOverriddenPropertyOperation(OverriddenObject, *PropertyChain.Get());
+												VisibleOverriddenPropertyOperation == Manager.GetOverriddenPropertyOperation(OverriddenObject, FPropertyChangedEvent(nullptr), *PropertyChain.Get());
 
 			const bool bIsComponentVisible = !bIsProperty && VisibleOverriddenState == Manager.GetOverriddenState(OverriddenObject);
 
