Prevent property propagation for object using overridable serialization

diff --git a/Engine/Source/Editor/PropertyEditor/Private/PropertyHandleImpl.cpp b/Engine/Source/Editor/PropertyEditor/Private/PropertyHandleImpl.cpp
--- a/Engine/Source/Editor/PropertyEditor/Private/PropertyHandleImpl.cpp
+++ b/Engine/Source/Editor/PropertyEditor/Private/PropertyHandleImpl.cpp
@@ -32,6 +32,7 @@
 #include "PropertyPathHelpers.h"
 #include "PropertyTextUtilities.h"
 #include "HAL/PlatformApplicationMisc.h"
+#include "UObject/OverridableManager.h"
 #include "UObject/PropertyOptional.h"
 
 #define LOCTEXT_NAMESPACE "PropertyHandleImplementation"
@@ -510,7 +511,8 @@ FPropertyAccess::Result FPropertyValueImpl::ImportText( const TArray<FObjectBase
 					!bIsGameWorld)
 				{
 					// propagate the changes to instances unless we're modifying class shared data
-					if (!bIsSparseClassData)
+					// or the object is using overridable serialization, the propagation is done via reinstantiation
+					if (!bIsSparseClassData && !FOverridableManager::Get().IsEnabled(*CurObject))
 					{
 						InPropertyNode->PropagatePropertyChange(CurObject, *NewValue, PreviousContainerValue.IsEmpty() ? PreviousValue : PreviousContainerValue);
 					}
