Property Editor: Add getting the expansion state for a detail property row, which is needed in the following jira so we can show/hide widgets based upon expansion state. Also need to get the property row from the root nodes so added function to detail layout to do that.

diff --git a/Engine/Source/Editor/PropertyEditor/Private/DetailLayoutBuilderImpl.cpp b/Engine/Source/Editor/PropertyEditor/Private/DetailLayoutBuilderImpl.cpp
--- a/Engine/Source/Editor/PropertyEditor/Private/DetailLayoutBuilderImpl.cpp
+++ b/Engine/Source/Editor/PropertyEditor/Private/DetailLayoutBuilderImpl.cpp
@@ -208,9 +208,27 @@ IDetailPropertyRow* FDetailLayoutBuilderImpl::EditDefaultProperty(TSharedPtr<IPr
 				}
 			}
 		}
-
 	}
+	return nullptr;
+}
 
+IDetailPropertyRow* FDetailLayoutBuilderImpl::EditPropertyFromRoot(TSharedPtr<IPropertyHandle> InPropertyHandle)
+{
+	for (const TSharedRef<FDetailTreeNode>& RootTreeNode : AllRootTreeNodes)
+	{
+		FDetailNodeList ChildNodes;
+		RootTreeNode->GetChildren(ChildNodes, true/*bIgnoreVisibility*/);
+		for (const TSharedRef<FDetailTreeNode>& ChildNode : ChildNodes)
+		{
+			if (TSharedPtr<IDetailPropertyRow> PropertyRow = ChildNode->GetRow())
+			{
+				if (PropertyRow->GetPropertyHandle() == InPropertyHandle)
+				{
+					return PropertyRow.Get();
+				}
+			}
+		}
+	}
 	return nullptr;
 }
 
diff --git a/Engine/Source/Editor/PropertyEditor/Private/DetailLayoutBuilderImpl.h b/Engine/Source/Editor/PropertyEditor/Private/DetailLayoutBuilderImpl.h
--- a/Engine/Source/Editor/PropertyEditor/Private/DetailLayoutBuilderImpl.h
+++ b/Engine/Source/Editor/PropertyEditor/Private/DetailLayoutBuilderImpl.h
@@ -38,6 +38,7 @@ public:
 	virtual TSharedPtr<IPropertyHandle> AddObjectPropertyData(TConstArrayView<UObject*> Objects, FName PropertyName) override;
 	virtual TSharedPtr<IPropertyHandle> AddStructurePropertyData(const TSharedPtr<FStructOnScope>& StructData, FName PropertyName) override;
 	virtual IDetailPropertyRow* EditDefaultProperty(TSharedPtr<IPropertyHandle> InPropertyHandle) override;
+	virtual IDetailPropertyRow* EditPropertyFromRoot(TSharedPtr<IPropertyHandle> InPropertyHandle) override;
 	virtual bool DoesCategoryHaveGeneratedChildren(FName CategoryName) override;
 	virtual TSharedRef<IPropertyHandle> GetProperty(const FName PropertyPath, const UStruct* ClassOutermost, FName InInstanceName) const override;
 	virtual FName GetTopLevelProperty() override;
diff --git a/Engine/Source/Editor/PropertyEditor/Private/DetailPropertyRow.cpp b/Engine/Source/Editor/PropertyEditor/Private/DetailPropertyRow.cpp
--- a/Engine/Source/Editor/PropertyEditor/Private/DetailPropertyRow.cpp
+++ b/Engine/Source/Editor/PropertyEditor/Private/DetailPropertyRow.cpp
@@ -178,6 +178,15 @@ IDetailPropertyRow& FDetailPropertyRow::DragDropHandler(TSharedPtr<IDetailDragDr
 	return *this;
 }
 
+bool FDetailPropertyRow::IsExpanded() const
+{
+	if (GetPropertyNode())
+	{
+		return GetPropertyNode()->HasNodeFlags(EPropertyNodeFlags::Expanded);
+	}
+	return false;
+}
+
 void FDetailPropertyRow::GetDefaultWidgets( TSharedPtr<SWidget>& OutNameWidget, TSharedPtr<SWidget>& OutValueWidget, bool bAddWidgetDecoration )
 {
 	FDetailWidgetRow Row;
diff --git a/Engine/Source/Editor/PropertyEditor/Private/DetailPropertyRow.h b/Engine/Source/Editor/PropertyEditor/Private/DetailPropertyRow.h
--- a/Engine/Source/Editor/PropertyEditor/Private/DetailPropertyRow.h
+++ b/Engine/Source/Editor/PropertyEditor/Private/DetailPropertyRow.h
@@ -37,6 +37,7 @@ public:
 	virtual IDetailPropertyRow& Visibility( TAttribute<EVisibility> Visibility ) override;
 	virtual IDetailPropertyRow& OverrideResetToDefault(const FResetToDefaultOverride& ResetToDefault) override;
 	virtual IDetailPropertyRow& DragDropHandler(TSharedPtr<IDetailDragDropHandler> InDragDropHandler) override;
+	virtual bool IsExpanded() const override;
 	virtual FDetailWidgetRow& CustomWidget( bool bShowChildren = false ) override;
 	virtual FDetailWidgetDecl* CustomNameWidget() override;
 	virtual FDetailWidgetDecl* CustomValueWidget() override;
diff --git a/Engine/Source/Editor/PropertyEditor/Public/DetailLayoutBuilder.h b/Engine/Source/Editor/PropertyEditor/Public/DetailLayoutBuilder.h
--- a/Engine/Source/Editor/PropertyEditor/Public/DetailLayoutBuilder.h
+++ b/Engine/Source/Editor/PropertyEditor/Public/DetailLayoutBuilder.h
@@ -175,11 +175,20 @@ public:
 	 * The property will remain in the default location but the widget or other attributes for the property can be changed 
 	 * Note This cannot be used to customize other customizations
 
-	 * @param InPropertyHandle	The handle to the property that you want to add to its own category.
+	 * @param InPropertyHandle	The handle to the property that you want to edit
 	 * @return					The property row to edit or nullptr if the property row does not exist
 	 */
 	virtual IDetailPropertyRow* EditDefaultProperty(TSharedPtr<IPropertyHandle> InPropertyHandle) = 0;
 
+	/**
+	 * Get the property row from the root of the details panel after it's been constructed, so this will work with default or custom 
+	 * properties
+	 * @param InPropertyHandle	The handle to the property that you want to edit
+	 * @return					The property row to edit or nullptr if the property row does not exist, which may happen if not 
+	 * constructed yet
+	 */
+	virtual IDetailPropertyRow* EditPropertyFromRoot(TSharedPtr<IPropertyHandle> InPropertyHandle) = 0;
+
 	/**
 	 * @return true if the category contains child rows. 
 	 */
diff --git a/Engine/Source/Editor/PropertyEditor/Public/IDetailPropertyRow.h b/Engine/Source/Editor/PropertyEditor/Public/IDetailPropertyRow.h
--- a/Engine/Source/Editor/PropertyEditor/Public/IDetailPropertyRow.h
+++ b/Engine/Source/Editor/PropertyEditor/Public/IDetailPropertyRow.h
@@ -250,6 +250,13 @@ public:
 	 */
 	virtual IDetailPropertyRow& DragDropHandler(TSharedPtr<IDetailDragDropHandler> InDragDropHandler) = 0;
 
+	/**
+	 * Returns the property row expansion state
+	 *
+	 * @return Will return true if the row is expanded, false if not
+	 */
+	virtual bool IsExpanded() const = 0;
+
 	/**
 	 * Returns the name and value widget of this property row.  You can use this widget to apply further customization to existing widgets (by using this  with CustomWidget)
 	 *
