Fix array paste when the input string contains commas
#tests Copy paste arrays with commas
#rnx

diff --git a/Engine/Source/Editor/PropertyEditor/Private/PropertyHandleImpl.cpp b/Engine/Source/Editor/PropertyEditor/Private/PropertyHandleImpl.cpp
--- a/Engine/Source/Editor/PropertyEditor/Private/PropertyHandleImpl.cpp
+++ b/Engine/Source/Editor/PropertyEditor/Private/PropertyHandleImpl.cpp
@@ -5063,7 +5063,7 @@ FPropertyAccess::Result FPropertyHandleArray::SetValueFromFormattedString(const
 {
 	FPropertyAccess::Result Result = FPropertyAccess::Success;
 
-	if (InValue.StartsWith("(") && InValue.EndsWith(")"))
+	if (InValue.StartsWith("((") && InValue.EndsWith("))"))
 	{
 		if (EmptyArray() != FPropertyAccess::Success)
 		{
@@ -5072,8 +5072,8 @@ FPropertyAccess::Result FPropertyHandleArray::SetValueFromFormattedString(const
 
 		TArray<FString> Values;
 
-		// Remove first and last parenthesis  
-		InValue.LeftChop(1).RightChop(1).ParseIntoArrayWS(Values, TEXT(","), true);
+		// Remove firsts and lasts parenthesis, and split the string 
+		InValue.LeftChop(2).RightChop(2).ParseIntoArrayWS(Values, TEXT("),("), true);
 
 		for (const FString& Value : Values)
 		{
@@ -5086,19 +5086,22 @@ FPropertyAccess::Result FPropertyHandleArray::SetValueFromFormattedString(const
 			{
 				return FPropertyAccess::Fail;
 			}
+
+			// Readd trimmed parenthesis
+			FString PasteValue = "(" + Value + ")";
 			
 			const TSharedRef<IPropertyHandle> Property = GetElement( NumElements - 1 );
-			if (Property->IsValidHandle() && Property->SetValueFromFormattedString(Value, Flags) == FPropertyAccess::Fail)
+			if (Property->IsValidHandle() && Property->SetValueFromFormattedString(PasteValue, Flags) == FPropertyAccess::Fail)
 			{
 				Result = FPropertyAccess::Fail;
 			}
 		}
-		return Result;
 	}
 	else
 	{
-		return FPropertyHandleBase::SetValueFromFormattedString(InValue, Flags);
+		Result = FPropertyHandleBase::SetValueFromFormattedString(InValue, Flags);
 	}
+	return Result;
 }
 
 bool FPropertyHandleArray::IsEditable() const
@@ -5494,7 +5497,7 @@ FPropertyAccess::Result FPropertyHandleSet::SetValueFromFormattedString(const FS
 {
 	FPropertyAccess::Result Result = FPropertyAccess::Success;
 
-	if (InValue.StartsWith("(") && InValue.EndsWith(")"))
+	if (InValue.StartsWith("((") && InValue.EndsWith("))"))
 	{
 		if (Empty() != FPropertyAccess::Success)
 		{
@@ -5503,8 +5506,8 @@ FPropertyAccess::Result FPropertyHandleSet::SetValueFromFormattedString(const FS
 		
 		TArray<FString> Values;
 
-		// Remove first and last parenthesis  
-		InValue.LeftChop(1).RightChop(1).ParseIntoArrayWS(Values, TEXT(","), true);
+		// Remove firsts and lasts parenthesis  
+		InValue.LeftChop(2).RightChop(2).ParseIntoArrayWS(Values, TEXT("),("), true);
 		for (const FString& Value : Values)
 		{
 			AddItem();
@@ -5513,18 +5516,22 @@ FPropertyAccess::Result FPropertyHandleSet::SetValueFromFormattedString(const FS
 			{
 				return FPropertyAccess::Fail;
 			}
+
+			// Readd trimmed parenthesis
+			FString PasteValue = "(" + Value + ")";
+			
 			const TSharedRef<IPropertyHandle> Property = GetElement( NumElements - 1 );
-			if ( Property->IsValidHandle() && Property->SetValueFromFormattedString(Value, Flags) == FPropertyAccess::Fail)
+			if ( Property->IsValidHandle() && Property->SetValueFromFormattedString(PasteValue, Flags) == FPropertyAccess::Fail)
 			{
 				Result = FPropertyAccess::Fail;
 			}
 		}
-		return Result;
 	}
 	else
 	{
-		return FPropertyHandleBase::SetValueFromFormattedString(InValue, Flags);
+		Result = FPropertyHandleBase::SetValueFromFormattedString(InValue, Flags);
 	}
+	return Result;
 }
 
 // Maps
