Hooked up overrides menu items which was added by [REDACTED_HANDLE]

diff --git a/Engine/Source/Editor/PropertyEditor/Private/DetailsViewStyle.cpp b/Engine/Source/Editor/PropertyEditor/Private/DetailsViewStyle.cpp
--- a/Engine/Source/Editor/PropertyEditor/Private/DetailsViewStyle.cpp
+++ b/Engine/Source/Editor/PropertyEditor/Private/DetailsViewStyle.cpp
@@ -4,55 +4,23 @@
 #include "DetailsViewStyle.h"
 #include "Containers/Map.h"
 #include "Brushes/SlateImageBrush.h"
-#include "Styling/StyleColors.h"
-#include "Styling/StarshipCoreStyle.h"
 #include "Brushes/SlateRoundedBoxBrush.h"
 #include "UObject/OverridableManager.h"
 #include "Layout/Visibility.h"
+#include "Styling/AppStyle.h"
 
 TArray< TSharedRef< const FOverridesWidgetStyleKey >> FOverridesWidgetStyleKeys::OverridesWidgetStyleKeys;
 
 const FSlateBrush& FOverridesWidgetStyleKey::GetConstStyleBrush() const
 {
-	return ImageBrush;
-}
-
-const FComboButtonStyle& FOverridesWidgetStyleKey::GetComboButtonStyle( const bool bIsOverridesWidgetForOuterCategory ) const
-{
-	static TMap<const FOverridesWidgetStyleKey*, const FComboButtonStyle> OverridesKeyToComboButtonStyleMap;
-	const FComboButtonStyle* ComboButtonStylePtr = OverridesKeyToComboButtonStyleMap.Find( this );
-
-	if ( ComboButtonStylePtr )
-	{
-		return *ComboButtonStylePtr;
-	}
-	
-	const FSlateColor BackgroundColor = bIsOverridesWidgetForOuterCategory ? FStyleColors::Header : FStyleColors::Panel;
-	static const FSlateColor HoveredBackgroundColor = FStyleColors::Header;
-			
-	const FButtonStyle OverridesButton = FButtonStyle()
-	                                     .SetNormalForeground(FStyleColors::AccentBlue)
-	                                     .SetHoveredForeground(FStyleColors::AccentBlue)
-	                                     .SetPressedForeground(FStyleColors::AccentBlue)
-	                                     .SetHovered(FSlateRoundedBoxBrush(HoveredBackgroundColor, 0.f))
-	                                     .SetNormal(FSlateRoundedBoxBrush(BackgroundColor, 0.f))
-	                                     .SetPressed(FSlateRoundedBoxBrush(BackgroundColor, 0.f))
-	                                     .SetNormalPadding(FMargin(2.f, 0.f, 0.f, 0.f))
-	                                     .SetPressedPadding(FMargin(2.f, 0.f, 0.f, 0.f));
-			
-
-	OverridesKeyToComboButtonStyleMap.Add( this, FComboButtonStyle(FStarshipCoreStyle::GetCoreStyle().GetWidgetStyle<FComboButtonStyle>("ComboButton"))
-												   .SetButtonStyle(OverridesButton)
-												   .SetDownArrowImage(GetConstStyleBrush())
-												   .SetDownArrowPadding(FMargin(2.f, 5.f, 3.f, 5.f)));
-
-		
-	return *OverridesKeyToComboButtonStyleMap.Find( this );
+	const FSlateBrush* ImageBrush = FAppStyle::GetBrush( *(TEXT("Icons.Details") + Name.ToString()) );
+	checkf(ImageBrush, TEXT("Expecting a valid brush"));
+	return *ImageBrush;
 }
 
 const FOverridesWidgetStyleKey& FOverridesWidgetStyleKeys::Here()
 {
-	static const FOverridesWidgetStyleKey Here{"OverrideHere", EOverriddenPropertyOperation::Replace, EOverriddenState::AllOverridden };
+	static const FOverridesWidgetStyleKey Here{"OverrideHere", EOverriddenPropertyOperation::Replace, EOverriddenState::AllOverridden, /*bInherited*/false };
 	return Here;
 }
 
@@ -63,19 +31,20 @@ void FOverridesWidgetStyleKeys::Initialize()
 	OverridesWidgetStyleKeys.Add( MakeShared<FOverridesWidgetStyleKey>(HereInside()));
 	OverridesWidgetStyleKeys.Add( MakeShared<FOverridesWidgetStyleKey>(Added()));
 	OverridesWidgetStyleKeys.Add( MakeShared<FOverridesWidgetStyleKey>(Removed()));
-	OverridesWidgetStyleKeys.Add( MakeShared<FOverridesWidgetStyleKey>(Options()));	
+	OverridesWidgetStyleKeys.Add( MakeShared<FOverridesWidgetStyleKey>(None()));
+	OverridesWidgetStyleKeys.Add( MakeShared<FOverridesWidgetStyleKey>(Inherited()));
 }
 
 const FOverridesWidgetStyleKey& FOverridesWidgetStyleKeys::Added()
 {
-	static const FOverridesWidgetStyleKey Added{"OverrideAdded"};
+	static const FOverridesWidgetStyleKey Added{"OverrideAdded", EOverriddenPropertyOperation::Add, EOverriddenState::Added };
 	return Added;
 }
 
-const FOverridesWidgetStyleKey& FOverridesWidgetStyleKeys::Options()
+const FOverridesWidgetStyleKey& FOverridesWidgetStyleKeys::None()
 {
-	static const FOverridesWidgetStyleKey Options{"OverrideOptions"};
-	return Options;
+	static const FOverridesWidgetStyleKey None{"OverrideNone", EOverriddenPropertyOperation::None, EOverriddenState::NoOverrides };
+	return None;
 }
 
 const FOverridesWidgetStyleKey& FOverridesWidgetStyleKeys::Removed()
@@ -96,22 +65,17 @@ const FOverridesWidgetStyleKey& FOverridesWidgetStyleKeys::HereInside()
 	return HereInside;
 }
 
-TArray< TSharedRef< const FOverridesWidgetStyleKey >> FOverridesWidgetStyleKeys::GetKeys()
+const FOverridesWidgetStyleKey& FOverridesWidgetStyleKeys::Inherited()
 {
-	return OverridesWidgetStyleKeys;
+	static const FOverridesWidgetStyleKey Inherited{"OverrideInherited", EOverriddenPropertyOperation::Replace, TOptional<EOverriddenState>(), /*bInherited*/true };
+	return Inherited;
 }
 
-void FOverridesWidgetStyleKey::Construct()
+TArray< TSharedRef< const FOverridesWidgetStyleKey >> FOverridesWidgetStyleKeys::GetKeys()
 {
-	static const FVector2D Icon16x16{16.0f, 16.0f};
-	static const FString Path = FPaths::ConvertRelativePathToFull(FPaths::Combine(FPaths::EngineContentDir(), TEXT("Slate/Starship/Common/")));
-		
-	const FSlateVectorImageBrush Brush{Path + Name.ToString() + ".svg", Icon16x16};
-	const FSlateBrush* SlateBrushPtr = &Brush;
-	ImageBrush = *SlateBrushPtr;
+	return OverridesWidgetStyleKeys;
 }
 
-
 TAttribute<EVisibility> FOverridesWidgetStyleKey::GetVisibilityAttribute(const TSharedPtr<FEditPropertyChain>& PropertyChain, TWeakObjectPtr<UObject>& OverriddenObjectWeakPtr) const
 {
 	static FOverridableManager& Manager = FOverridableManager::Get();
@@ -122,14 +86,21 @@ TAttribute<EVisibility> FOverridesWidgetStyleKey::GetVisibilityAttribute(const T
 		
 	    if ( OverriddenObjectWeakPtr.IsValid() )
 	    {
-	    	UObject& OverriddenObject = *OverriddenObjectWeakPtr.Get();
+		    UObject& OverriddenObject = *OverriddenObjectWeakPtr.Get();
 
-	    	const bool bIsPropertyVisible = bIsProperty &&
-												VisibleOverriddenPropertyOperation == Manager.GetOverriddenPropertyOperation(OverriddenObject, FPropertyChangedEvent(nullptr), *PropertyChain.Get());
-
-			const bool bIsComponentVisible = !bIsProperty && VisibleOverriddenState == Manager.GetOverriddenState(OverriddenObject);
+			bool bIsVisible = false;
+			if (bIsProperty)
+			{
+				bool bInherited = false;
+				bIsVisible =  VisibleOverriddenPropertyOperation == Manager.GetOverriddenPropertyOperation(OverriddenObject, FPropertyChangedEvent(nullptr), *PropertyChain.Get(), &bInherited) &&
+							  (!bStateInherited.IsSet() || bStateInherited.GetValue() == bInherited);
+			}
+			else if (VisibleOverriddenState.IsSet())
+			{
+				bIsVisible = VisibleOverriddenState.GetValue() == Manager.GetOverriddenState(OverriddenObject);
+			}
 
-			if ( bIsComponentVisible || bIsPropertyVisible )
+			if ( bIsVisible )
 			{
 				return EVisibility::Visible;
 			}
@@ -141,18 +112,18 @@ TAttribute<EVisibility> FOverridesWidgetStyleKey::GetVisibilityAttribute(const T
 
 FOverridesWidgetStyleKey::FOverridesWidgetStyleKey(FName InName) : Name{InName}
 {
-     Construct();
 }
 
 FOverridesWidgetStyleKey::FOverridesWidgetStyleKey(FName InName,
-																				EOverriddenPropertyOperation InOverriddenPropertyOperation,
-																				EOverriddenState InOverriddenState) :
+													EOverriddenPropertyOperation InOverriddenPropertyOperation,
+													const TOptional<EOverriddenState>& InOverriddenState,
+													const TOptional<bool>& bInStateInherited /*= TOptional<bool>()*/) :
    Name(InName),
    VisibleOverriddenPropertyOperation(InOverriddenPropertyOperation),
    VisibleOverriddenState(InOverriddenState),
+   bStateInherited(bInStateInherited),
    bCanBeVisible(true)
 {
-	Construct();
 }
 
 FDetailsViewStyle::FDetailsViewStyle()
diff --git a/Engine/Source/Editor/PropertyEditor/Private/DetailsViewStyle.h b/Engine/Source/Editor/PropertyEditor/Private/DetailsViewStyle.h
--- a/Engine/Source/Editor/PropertyEditor/Private/DetailsViewStyle.h
+++ b/Engine/Source/Editor/PropertyEditor/Private/DetailsViewStyle.h
@@ -33,11 +33,12 @@ public:
  * @param InName the FName which is the name of the key
  * @param InOverriddenPropertyOperation the overridden property operation for which this style is visible
  * @param InOverriddenState the overridden state for components for which this style is visible
+ * @param bInStateInherited take into account if the state was inherited or not, not set means anything
  */
-	PROPERTYEDITOR_API FOverridesWidgetStyleKey(FName InName,
-																			EOverriddenPropertyOperation InOverriddenPropertyOperation,
-																			EOverriddenState InOverriddenState
-																			);
+	PROPERTYEDITOR_API FOverridesWidgetStyleKey(FName InName, 
+												EOverriddenPropertyOperation InOverriddenPropertyOperation,
+												const TOptional<EOverriddenState>& InOverriddenState,
+												const TOptional<bool>& bInStateInherited = TOptional<bool>());
 
 	/**
 	 * returns the const FSlateBrush& that creates the icon for this widget style
@@ -52,13 +53,6 @@ public:
 	 * @return 
 	 */
 	TAttribute<EVisibility> GetVisibilityAttribute(const TSharedPtr<FEditPropertyChain>& PropertyChain, TWeakObjectPtr<UObject>& OverriddenObjectWeakPtr) const;
-	
-	/**
-	 * returns the const FComboButtonStyle& for this OverridesWidgetStyleKey
-	 *
-	 * @param bIsForOuterCategory true if this style is for an outer Category
-	 */
-	PROPERTYEDITOR_API const FComboButtonStyle& GetComboButtonStyle(const bool bIsForOuterCategory = false) const;
 
 	/**
 	 * the name of the key
@@ -75,19 +69,12 @@ public:
 	 */
 	const TOptional<EOverriddenState> VisibleOverriddenState;
 
-	const bool bCanBeVisible = false;
-
-private:
-
 	/**
-	 * The image brush specified by this style key
+	 * Whetever the state was inherited from a parent node or not
 	 */
-	FSlateBrush ImageBrush;
+	TOptional<bool> bStateInherited;
 
-	/**
-	 * Construct the style key
-	 */
-	void Construct();
+	const bool bCanBeVisible = false;
 };
 
 /**s
@@ -111,9 +98,9 @@ public:
 
 	/**
 	 * The style for the override widget when the user hovers over it which shows that they have associated
-	 * action options to choose from
+	 * action no overrides to choose from
 	 */
-	PROPERTYEDITOR_API static const FOverridesWidgetStyleKey& Options();
+	PROPERTYEDITOR_API static const FOverridesWidgetStyleKey& None();
 
 	/**
 	 * The style for the override widget when an item has been removed 
@@ -128,6 +115,11 @@ public:
 
 	static const FOverridesWidgetStyleKey& HereInside();
 
+	/**
+	 * The override state was inherited from a parent node
+	 */
+	static const FOverridesWidgetStyleKey& Inherited();
+
 	static TArray< TSharedRef< const FOverridesWidgetStyleKey >> GetKeys();
 	
 	static void Initialize(); 
diff --git a/Engine/Source/Editor/PropertyEditor/Private/SDetailCategoryTableRow.cpp b/Engine/Source/Editor/PropertyEditor/Private/SDetailCategoryTableRow.cpp
--- a/Engine/Source/Editor/PropertyEditor/Private/SDetailCategoryTableRow.cpp
+++ b/Engine/Source/Editor/PropertyEditor/Private/SDetailCategoryTableRow.cpp
@@ -53,6 +53,9 @@ void SDetailCategoryTableRow::Construct(const FArguments& InArgs, TSharedRef<FDe
 	
 	PasteAction.ExecuteAction = FExecuteAction::CreateSP(this, &SDetailCategoryTableRow::OnPasteCategory);
 	PasteAction.CanExecuteAction = FCanExecuteAction::CreateSP(this, &SDetailCategoryTableRow::CanPasteCategory);
+
+	ResetToDefault.ExecuteAction = FExecuteAction::CreateSP(this, &SDetailCategoryTableRow::OnResetToDefaultCategory);
+	ResetToDefault.CanExecuteAction = FCanExecuteAction::CreateSP(this, &SDetailCategoryTableRow::CanResetToDefaultCategory);
 	
 	TSharedRef<SHorizontalBox> HeaderBox = SNew(SHorizontalBox)
 		+ SHorizontalBox::Slot()
@@ -107,7 +110,10 @@ void SDetailCategoryTableRow::Construct(const FArguments& InArgs, TSharedRef<FDe
 	OwnerTableViewWeak = InOwnerTableView;
 	PropertyUpdatedWidgetBuilder = DisplayManager->GetPropertyUpdatedWidget(FExecuteAction::CreateLambda( [this]
 		{
-			UE_LOG(LogTemp, Warning, TEXT("DisplayManager->GetPropertyUpdatedWidget reset"));
+			if( ResetToDefault.CanExecute())
+			{
+				ResetToDefault.Execute();
+			}
 		}), true, ObjectName );
 	if (PropertyUpdatedWidgetBuilder.IsValid())
 	{
@@ -512,6 +518,34 @@ bool SDetailCategoryTableRow::CanPasteCategory()
 	return PreviousClipboardData.bIsApplicable = Algo::Compare(PreviousClipboardData.PropertyNames, PropertyNames);
 }
 
+void SDetailCategoryTableRow::OnResetToDefaultCategory()
+{
+	if (!OwnerTreeNode.IsValid())
+	{
+		return;
+	}
+
+	if (TArray<TSharedPtr<IPropertyHandle>> CategoryProperties = GetPropertyHandles(true);
+		!CategoryProperties.IsEmpty())
+	{
+		for (TSharedPtr<IPropertyHandle> PropertyHandle : CategoryProperties)
+		{
+			PropertyHandle->ResetToDefault();
+		}
+	}
+}
+
+bool SDetailCategoryTableRow::CanResetToDefaultCategory() const
+{
+	if (!OwnerTreeNode.IsValid())
+	{
+		return false;
+	}
+
+	TArray<TSharedPtr<IPropertyHandle>> PropertyHandles = GetPropertyHandles(true);
+	return !PropertyHandles.IsEmpty();
+}
+
 FReply SDetailCategoryTableRow::OnMouseButtonDown(const FGeometry& MyGeometry, const FPointerEvent& MouseEvent)
 {
 	if (MouseEvent.GetModifierKeys().IsShiftDown())
diff --git a/Engine/Source/Editor/PropertyEditor/Private/SDetailCategoryTableRow.h b/Engine/Source/Editor/PropertyEditor/Private/SDetailCategoryTableRow.h
--- a/Engine/Source/Editor/PropertyEditor/Private/SDetailCategoryTableRow.h
+++ b/Engine/Source/Editor/PropertyEditor/Private/SDetailCategoryTableRow.h
@@ -65,6 +65,8 @@ private:
 	bool CanCopyCategory() const;
 	void OnPasteCategory();
 	bool CanPasteCategory();
+	void OnResetToDefaultCategory();
+	bool CanResetToDefaultCategory() const;
 
 	virtual FReply OnMouseButtonDown(const FGeometry& MyGeometry, const FPointerEvent& MouseEvent) override;
 	virtual FReply OnMouseButtonDoubleClick(const FGeometry& InMyGeometry, const FPointerEvent& InMouseEvent) override;
@@ -99,6 +101,7 @@ private:
 
 	FUIAction CopyAction;
 	FUIAction PasteAction;
+	FUIAction ResetToDefault;
 
 	/** Delegate handling pasting an optionally tagged text snippet */
 	TSharedPtr<FOnPasteFromText> OnPasteFromTextDelegate;
diff --git a/Engine/Source/Editor/PropertyEditor/Private/UserInterface/Widgets/OverridesComboButtonBuilder.cpp b/Engine/Source/Editor/PropertyEditor/Private/UserInterface/Widgets/OverridesComboButtonBuilder.cpp
--- a/Engine/Source/Editor/PropertyEditor/Private/UserInterface/Widgets/OverridesComboButtonBuilder.cpp
+++ b/Engine/Source/Editor/PropertyEditor/Private/UserInterface/Widgets/OverridesComboButtonBuilder.cpp
@@ -3,9 +3,9 @@
 #include "UserInterface/Widgets/OverridesComboButtonBuilder.h"
 
 #include "DetailsViewStyle.h"
+#include "SSimpleComboButton.h"
 #include "Styling/SlateBrush.h"
 #include "Brushes/SlateRoundedBoxBrush.h"
-#include "UObject/OverridableManager.h"
 
 FOverridesComboButtonBuilder::FOverridesComboButtonBuilder(
 	TSharedRef<FDetailsDisplayManager> InDetailsDisplayManager,
@@ -46,10 +46,10 @@ TSharedPtr<SWidget> FOverridesComboButtonBuilder::GenerateWidget()
 			   .VAlign(VAlign_Center)
 			   .AutoWidth()
 			[
-				SNew( SComboButton )
-					.ComboButtonStyle( &Key->GetComboButtonStyle(  bIsCategoryOverridesComboButton )  )
+				SNew( SSimpleComboButton )
 					.Visibility( Key->GetVisibilityAttribute(EditPropertyChain, Object ))
-					.HasDownArrow(true)
+					.OnGetMenuContent(OnGetContent)
+					.Icon(&Key->GetConstStyleBrush())
 			];
 		}
 	}
@@ -72,3 +72,8 @@ void FOverridesComboButtonBuilder:: SetEditPropertyChain(TSharedRef<FEditPropert
 	EditPropertyChain = InEditPropertyChain;	
 	bIsCategory = false;
 }
+
+TSharedPtr<FEditPropertyChain> FOverridesComboButtonBuilder::GetEditPropertyChain() const
+{
+	return EditPropertyChain;
+}
diff --git a/Engine/Source/Editor/PropertyEditor/Public/UserInterface/Widgets/OverridesComboButtonBuilder.h b/Engine/Source/Editor/PropertyEditor/Public/UserInterface/Widgets/OverridesComboButtonBuilder.h
--- a/Engine/Source/Editor/PropertyEditor/Public/UserInterface/Widgets/OverridesComboButtonBuilder.h
+++ b/Engine/Source/Editor/PropertyEditor/Public/UserInterface/Widgets/OverridesComboButtonBuilder.h
@@ -48,6 +48,11 @@ public:
 	 */
 	PROPERTYEDITOR_API void SetEditPropertyChain(TSharedRef<FEditPropertyChain>& InEditPropertyChain);
 
+	/**
+	 * Gets the FEditPropertyChain for this FOverridesComboButtonBuilder
+	 */
+	PROPERTYEDITOR_API TSharedPtr<FEditPropertyChain> GetEditPropertyChain() const;
+
 	/**
 	 * Implements the generation of the Category Menu button SWidget
 	 */
diff --git a/Engine/Source/Runtime/SlateCore/Private/Styling/StarshipCoreStyle.cpp b/Engine/Source/Runtime/SlateCore/Private/Styling/StarshipCoreStyle.cpp
--- a/Engine/Source/Runtime/SlateCore/Private/Styling/StarshipCoreStyle.cpp
+++ b/Engine/Source/Runtime/SlateCore/Private/Styling/StarshipCoreStyle.cpp
@@ -424,7 +424,15 @@ TSharedRef<ISlateStyle> FStarshipCoreStyle::Create()
 		Style->Set("Icons.BadgeOverrideSelfAndChild", new IMAGE_BRUSH_SVG("Starship/Common/CircleBadgeHalf", Icon16x16, FStyleColors::Recessed));
 		Style->Set("Icons.BadgeOverrideAdd", new IMAGE_BRUSH_SVG("Starship/Common/CircleBadgePlus", Icon16x16, FStyleColors::Recessed));
 		Style->Set("Icons.BadgeOverrideRemove", new IMAGE_BRUSH_SVG("Starship/Common/CircleBadgeMinus", Icon16x16, FStyleColors::Recessed));
-		
+
+		// Background images for all the details panels override states
+		Style->Set("Icons.DetailsOverrideHere", new IMAGE_BRUSH_SVG("Starship/Common/DetailsOverrideHere", Icon16x16, FStyleColors::AccentBlue));
+		Style->Set("Icons.DetailsOverrideAdded", new IMAGE_BRUSH_SVG("Starship/Common/DetailsOverrideAdded", Icon16x16, FStyleColors::AccentGreen));
+		Style->Set("Icons.DetailsOverrideNone", new IMAGE_BRUSH_SVG("Starship/Common/DetailsOverrideNone", Icon16x16, FStyleColors::Foreground));
+		Style->Set("Icons.DetailsOverrideRemoved", new IMAGE_BRUSH_SVG("Starship/Common/DetailsOverrideRemoved", Icon16x16, FStyleColors::Error));
+		Style->Set("Icons.DetailsOverrideInside", new IMAGE_BRUSH_SVG("Starship/Common/DetailsOverrideInside", Icon16x16, FStyleColors::AccentBlue));
+		Style->Set("Icons.DetailsOverrideHereInside", new IMAGE_BRUSH_SVG("Starship/Common/DetailsOverrideHereInside", Icon16x16, FStyleColors::AccentBlue));
+		Style->Set("Icons.DetailsOverrideInherited", new IMAGE_BRUSH_SVG("Starship/Common/DetailsOverrideInherited", Icon16x16, FStyleColors::AccentBlue));
 		
 		// Toolbar Size Icons
 		Style->Set("Icons.Toolbar.Settings", new IMAGE_BRUSH_SVG("Starship/Common/settings", Icon20x20));
