Property editor - Optimized performance when using a PropertyPermissionList to save several minutes of editor startup time

diff --git a/Engine/Source/Editor/PropertyEditor/Private/SDetailsViewBase.cpp b/Engine/Source/Editor/PropertyEditor/Private/SDetailsViewBase.cpp
--- a/Engine/Source/Editor/PropertyEditor/Private/SDetailsViewBase.cpp
+++ b/Engine/Source/Editor/PropertyEditor/Private/SDetailsViewBase.cpp
@@ -29,6 +29,7 @@ SDetailsViewBase::SDetailsViewBase() :
 	, bHasOpenColorPicker(false)
 	, bDisableCustomDetailLayouts(false)
 	, bPendingCleanupTimerSet(false)
+	, bPendingRefreshTimerSet(false)
 	, bRunningDeferredActions(false)
 {
 	UDetailsConfig::Initialize();
@@ -44,8 +45,9 @@ SDetailsViewBase::SDetailsViewBase() :
 		CurrentFilter.bShowOnlyModified = ViewConfig->bShowOnlyModified;
 	}
 
-	PropertyPermissionListChangedDelegate = FPropertyEditorPermissionList::Get().PermissionListUpdatedDelegate.AddLambda([this](TSoftObjectPtr<UStruct> Struct, FName Owner) { ForceRefresh(); });
-	PropertyPermissionListEnabledDelegate = FPropertyEditorPermissionList::Get().PermissionListEnabledDelegate.AddRaw(this, &SDetailsViewBase::ForceRefresh);
+	// Avoid calling ForceRefresh until next tick otherwise the editor can lock up for several minutes in one frame refreshing multiple times
+	PropertyPermissionListChangedDelegate = FPropertyEditorPermissionList::Get().PermissionListUpdatedDelegate.AddLambda([this](TSoftObjectPtr<UStruct> Struct, FName Owner) { SetPendingRefreshTimer(); });
+	PropertyPermissionListEnabledDelegate = FPropertyEditorPermissionList::Get().PermissionListEnabledDelegate.AddRaw(this, &SDetailsViewBase::SetPendingRefreshTimer);
 }
 
 SDetailsViewBase::~SDetailsViewBase()
@@ -1113,6 +1115,24 @@ void SDetailsViewBase::HandlePendingCleanup()
 	DetailLayoutsPendingDelete.Empty();
 }
 
+void SDetailsViewBase::SetPendingRefreshTimer()
+{
+	if (!bPendingRefreshTimerSet)
+	{
+		if (GEditor && GEditor->IsTimerManagerValid())
+		{
+			bPendingRefreshTimerSet = true;
+			GEditor->GetTimerManager()->SetTimerForNextTick(FTimerDelegate::CreateSP(this, &SDetailsViewBase::HandlePendingRefreshTimer));
+		}
+	}
+}
+
+void SDetailsViewBase::HandlePendingRefreshTimer()
+{
+	bPendingRefreshTimerSet = false;
+	ForceRefresh();
+}
+
 /** Ticks the property view.  This function performs a data consistency check */
 void SDetailsViewBase::Tick( const FGeometry& AllottedGeometry, const double InCurrentTime, const float InDeltaTime )
 {
diff --git a/Engine/Source/Editor/PropertyEditor/Private/SDetailsViewBase.h b/Engine/Source/Editor/PropertyEditor/Private/SDetailsViewBase.h
--- a/Engine/Source/Editor/PropertyEditor/Private/SDetailsViewBase.h
+++ b/Engine/Source/Editor/PropertyEditor/Private/SDetailsViewBase.h
@@ -394,7 +394,12 @@ protected:
 	void HandlePendingCleanupTimer();
 	/** Free memory that is pending delete */
 	void HandlePendingCleanup();
-	
+
+	/** Set timer to force refresh if one not already set */
+	void SetPendingRefreshTimer();
+	/** Force refresh during editor tick */
+	void HandlePendingRefreshTimer();
+
 	void SavePreSearchExpandedItems();
 	void RestorePreSearchExpandedItems();
 
@@ -517,6 +522,9 @@ protected:
 	/** Timer has already been set to be run next tick */
 	bool bPendingCleanupTimerSet : 1;
 
+	/** Refresh timer has already been set to be run next tick */
+	bool bPendingRefreshTimerSet : 1;
+
 	/** Are we currently running deferred actions? */
 	bool bRunningDeferredActions : 1;
 
