Check widget row in addition to details view enablement for paste action enablement on SDetailSingleItemRow

diff --git a/Engine/Source/Editor/PropertyEditor/Private/SDetailSingleItemRow.cpp b/Engine/Source/Editor/PropertyEditor/Private/SDetailSingleItemRow.cpp
--- a/Engine/Source/Editor/PropertyEditor/Private/SDetailSingleItemRow.cpp
+++ b/Engine/Source/Editor/PropertyEditor/Private/SDetailSingleItemRow.cpp
@@ -1058,7 +1058,7 @@ void SDetailSingleItemRow::PopulateContextMenu(UToolMenu* ToolMenu)
 			// Paste
 			{
 				// Paste is only enabled if property editing is enabled
-				if (OwningDetailsView && OwningDetailsView->IsPropertyEditingEnabled())
+				if (OwningDetailsView && OwningDetailsView->IsPropertyEditingEnabled() && WidgetRow.IsEnabledAttr.Get())
 				{
 					TAttribute<FText> Label;
 					TAttribute<FText> ToolTip;
