Engine (map and set add item);

Implementing a way to have the map and set fill in with the next available enum value instead of always doing the default value.

Slack discussion with video: https://epicgames.slack.com/archives/C0447BVGFFZ/p1705510234620249

-- Code --
PropertyHandleImpl:
FPropertyHandleMap::AddItem() - Added a HasKey() lambda. We're then checking if the map property's key is an enum property, if so, we'll loop the enum values to see if the map doesn't already have that enum value as a key. When we find one that isn't in the map, we're adding a child node and then setting the child's key node to be the enum value that is not in the map yet. If the map property's key is not enum then we do the old check for if the default key isn't there.
FPropertyHandleSet::AddItem() - Added a HasElement() lambda. We're then checking if the set property's element is an enum property, if so, we'll loop the enum values to see if the set doesn't already have the enum value as an element. When we find one that isn't in the set, we're adding a child node and then setting the child node to be the enum value that is not in the set yet. If the set property's element is not enum then we do the old check for if the default key isn't there.

UnrealTypes:
FSetProperty::GetElementProperty() - Added this getter to return the ElementProp property.

#tests PIE, opened a uasset that has a TMap property with enum as the key, added each enum key as I clicked the add item button, changed an entry to something out of order to make sure the add item button would add the next available entry, all the way until it wouldn't add anymore entries because all enum keys were taken. Also tried a TMap property with a gameplaytag as the key to make sure it still worked as expected. Repeated with TSet.

diff --git a/Engine/Source/Editor/PropertyEditor/Private/PropertyHandleImpl.cpp b/Engine/Source/Editor/PropertyEditor/Private/PropertyHandleImpl.cpp
--- a/Engine/Source/Editor/PropertyEditor/Private/PropertyHandleImpl.cpp
+++ b/Engine/Source/Editor/PropertyEditor/Private/PropertyHandleImpl.cpp
@@ -5419,18 +5419,118 @@ FPropertyAccess::Result FPropertyHandleSet::AddItem()
 	FPropertyAccess::Result Result = FPropertyAccess::Fail;
 	if (IsEditable())
 	{
-		if (!HasDefaultElement())
-		{
-			Implementation->AddChild();
-			Implementation->GetPropertyNode()->RebuildChildren();
-			Result = FPropertyAccess::Success;
-		}
-		else
+		/**
+		 * Checks if an element has already been added to the set
+		 *
+		 * @param	Helper			The set helper used to query the property.
+		 * @param	InBaseAddress	The base address of the set
+		 * @param	InElementValue	The element value to check for
+		 *
+		 * @return	True if the element is found in the set, false otherwise
+		 */
+		static auto HasElement = [](const FScriptSetHelper& Helper, void* InBaseAddress, const FString& InElementValue)
+			{
+				FProperty* ElementProp = Helper.GetElementProperty();
+
+				void* TempElementStorage = ElementProp->AllocateAndInitializeValue();
+				ON_SCOPE_EXIT
+				{
+					ElementProp->DestroyAndFreeValue(TempElementStorage);
+				};
+
+				for (int32 Index = 0, ItemsLeft = Helper.Num(); ItemsLeft > 0; ++Index)
+				{
+					if (Helper.IsValidIndex(Index))
+					{
+						--ItemsLeft;
+
+						const uint8* Element = Helper.GetElementPtr(Index);
+
+						if (Element != InBaseAddress && ElementProp->ImportText_Direct(*InElementValue, TempElementStorage, nullptr, 0) && ElementProp->Identical(Element, TempElementStorage))
+						{
+							return true;
+						}
+					}
+				}
+
+				return false;
+			};
+
+		TSharedPtr<FPropertyNode> PropNode = Implementation->GetPropertyNode();
+		if (PropNode.IsValid())
 		{
-			Implementation->ShowInvalidOperationError(LOCTEXT("DuplicateSetElement_Add", "Cannot add a new element to the set while an element with the default value exists"));
+			TArray<FObjectBaseAddress> Addresses;
+			Implementation->GetObjectsToModify(Addresses, PropNode.Get());
+
+			if (Addresses.Num() > 0)
+			{
+				const bool bIsSparseClassData = PropNode->HasNodeFlags(EPropertyNodeFlags::IsSparseClassData) != 0;
+				const bool bIsStruct = !Addresses[0].Object;
+				uint8* ValueBaseAddress = PropNode->GetValueBaseAddress(Addresses[0].StructAddress, bIsSparseClassData, bIsStruct);
+
+				FSetProperty* SetProperty = CastFieldChecked<FSetProperty>(PropNode->GetProperty());
+				FScriptSetHelper SetHelper(SetProperty, ValueBaseAddress);
+
+				if (const FEnumProperty* EnumProperty = CastField<const FEnumProperty>(SetProperty->GetElementProperty()))
+				{
+					if (UEnum* Enum = EnumProperty->GetEnum())
+					{
+						static FString HiddenName(TEXT("Hidden"));
+						static FString SpacerName(TEXT("Spacer"));
+						// NumEnums() - 1, because the last item in an enum is the _MAX item
+						for (int32 EnumIndex = 0; EnumIndex < Enum->NumEnums() - 1; ++EnumIndex)
+						{
+							// Skip hidden and spacers so we don't add them.
+							const bool bShouldBeHidden = Enum->HasMetaData(*HiddenName, EnumIndex) || Enum->HasMetaData(*SpacerName, EnumIndex);
+							if (bShouldBeHidden)
+							{
+								continue;
+							}
+
+							// Skip duplicate elements.
+							if (HasElement(SetHelper, ValueBaseAddress, Enum->GetNameStringByIndex(EnumIndex)))
+							{
+								continue;
+							}
+
+							// If we don't have this element then add an entry and set it to this element value.
+							Implementation->AddChild();
+							Implementation->GetPropertyNode()->RebuildChildren();
+
+							// Grab the last entry since we just added it.
+							const int32 ChildNodeIndex = PropNode->GetNumChildNodes() - 1;
+							if (ChildNodeIndex >= 0)
+							{
+								TSharedPtr<FPropertyNode> ChildNode = Implementation->GetChildNode(ChildNodeIndex);
+								if (ChildNode.IsValid())
+								{
+									// Set the new entry to have the element value.
+									Implementation->ImportText(Enum->GetNameStringByIndex(EnumIndex), ChildNode.Get(), EPropertyValueSetFlags::DefaultFlags);
+									Result = FPropertyAccess::Success;
+									break;
+								}
+							}
+						}
+					}
+				}
+				else
+				{
+					if (!HasDefaultElement())
+					{
+						Implementation->AddChild();
+						Implementation->GetPropertyNode()->RebuildChildren();
+						Result = FPropertyAccess::Success;
+					}
+				}
+			}
 		}
 	}
 
+	if (Result != FPropertyAccess::Success)
+	{
+		Implementation->ShowInvalidOperationError(LOCTEXT("DuplicateSetElement_Add", "Cannot add a new element to the set while an element with the same value exists"));
+	}
+
 	return Result;
 }
 
@@ -5571,17 +5671,122 @@ FPropertyAccess::Result FPropertyHandleMap::AddItem()
 	FPropertyAccess::Result Result = FPropertyAccess::Fail;
 	if (IsEditable())
 	{
-		if ( !HasDefaultKey() )
-		{
-			Implementation->AddChild();
-			Result = FPropertyAccess::Success;
-		}
-		else
+		/**
+		 * Checks if a key in the map matches the specified key
+		 *
+		 * @param	Helper			The map helper used to query the property.
+		 * @param	InBaseAddress	The base address of the map
+		 * @param	InKeyValue		The key to find within the map
+		 *
+		 * @return	True if the key is found, false otherwise
+		 */
+		static auto HasKey = [](const FScriptMapHelper& Helper, void* InBaseAddress, const FString& InKeyValue)
+			{
+				FProperty* KeyProp = Helper.GetKeyProperty();
+
+				void* TempKeyStorage = KeyProp->AllocateAndInitializeValue();
+				ON_SCOPE_EXIT
+				{
+					KeyProp->DestroyAndFreeValue(TempKeyStorage);
+				};
+
+				for (int32 Index = 0, ItemsLeft = Helper.Num(); ItemsLeft > 0; ++Index)
+				{
+					if (Helper.IsValidIndex(Index))
+					{
+						--ItemsLeft;
+
+						const uint8* PairPtr = Helper.GetPairPtr(Index);
+						const uint8* KeyPtr = KeyProp->ContainerPtrToValuePtr<const uint8>(PairPtr);
+
+						if (KeyPtr != InBaseAddress && KeyProp->ImportText_Direct(*InKeyValue, TempKeyStorage, nullptr, 0) && KeyProp->Identical(KeyPtr, TempKeyStorage))
+						{
+							return true;
+						}
+					}
+				}
+
+				return false;
+			};
+
+		TSharedPtr<FPropertyNode> PropNode = Implementation->GetPropertyNode();
+		if (PropNode.IsValid())
 		{
-			Implementation->ShowInvalidOperationError(LOCTEXT("DuplicateMapKey_Add", "Cannot add a new key to the map while a key with the default value exists"));
+			TArray<FObjectBaseAddress> Addresses;
+			Implementation->GetObjectsToModify(Addresses, PropNode.Get());
+
+			if (Addresses.Num() > 0)
+			{
+				const bool bIsSparseClassData = PropNode->HasNodeFlags(EPropertyNodeFlags::IsSparseClassData) != 0;
+				const bool bIsStruct = !Addresses[0].Object;
+				uint8* ValueBaseAddress = PropNode->GetValueBaseAddress(Addresses[0].StructAddress, bIsSparseClassData, bIsStruct);
+			
+				FMapProperty* MapProperty = CastFieldChecked<FMapProperty>(PropNode->GetProperty());
+				FScriptMapHelper MapHelper(MapProperty, ValueBaseAddress);
+
+				if (const FEnumProperty* EnumProperty = CastField<const FEnumProperty>(MapProperty->GetKeyProperty()))
+				{
+					if (UEnum* Enum = EnumProperty->GetEnum())
+					{
+						static FString HiddenName(TEXT("Hidden"));
+						static FString SpacerName(TEXT("Spacer"));
+						// NumEnums() - 1, because the last item in an enum is the _MAX item
+						for (int32 EnumIndex = 0; EnumIndex < Enum->NumEnums() - 1; ++EnumIndex)
+						{
+							// Skip hidden and spacers so we don't add them.
+							const bool bShouldBeHidden = Enum->HasMetaData(*HiddenName, EnumIndex) || Enum->HasMetaData(*SpacerName, EnumIndex);
+							if (bShouldBeHidden)
+							{
+								continue;
+							}
+
+							// Skip duplicate keys.
+							if (HasKey(MapHelper, ValueBaseAddress, Enum->GetNameStringByIndex(EnumIndex)))
+							{
+								continue;
+							}
+
+							// If we don't have this key then add an entry and set it to this key value.
+							Implementation->AddChild();
+
+							// Grab the last entry since we just added it.
+							const int32 ChildNodeIndex = PropNode->GetNumChildNodes() - 1;
+							if (ChildNodeIndex >= 0)
+							{
+								TSharedPtr<FPropertyNode> ChildNode = Implementation->GetChildNode(ChildNodeIndex);
+								if (ChildNode.IsValid())
+								{
+									// We want to change the key, so get the key property.
+									TSharedPtr<FPropertyNode> ChildKeyNode = ChildNode->GetPropertyKeyNode();
+									if (ChildKeyNode.IsValid())
+									{
+										// Set the new entry to have the key value.
+										Implementation->ImportText(Enum->GetNameStringByIndex(EnumIndex), ChildKeyNode.Get(), EPropertyValueSetFlags::DefaultFlags);
+										Result = FPropertyAccess::Success;
+										break;
+									}
+								}
+							}
+						}
+					}
+				}
+				else
+				{
+					if (!HasDefaultKey())
+					{
+						Implementation->AddChild();
+						Result = FPropertyAccess::Success;
+					}
+				}
+			}
 		}
 	}
 
+	if (Result != FPropertyAccess::Success)
+	{
+		Implementation->ShowInvalidOperationError(LOCTEXT("DuplicateMapKey_Add", "Cannot add a new key to the map while a key with the same value exists"));
+	}
+
 	return Result;
 }
 
diff --git a/Engine/Source/Runtime/CoreUObject/Public/UObject/UnrealType.h b/Engine/Source/Runtime/CoreUObject/Public/UObject/UnrealType.h
--- a/Engine/Source/Runtime/CoreUObject/Public/UObject/UnrealType.h
+++ b/Engine/Source/Runtime/CoreUObject/Public/UObject/UnrealType.h
@@ -3989,6 +3989,12 @@ public:
 		FScriptSet* Set = (FScriptSet*)InSet;
 		return (uint8*)Set->GetData(InternalIndex, SetLayout);
 	}
+
+	const FProperty* GetElementProperty() const
+	{
+		return ElementProp;
+	}
+
 #if WITH_EDITORONLY_DATA
 	virtual void AppendSchemaHash(FBlake3& Builder, bool bSkipEditorOnly) const override;
 #endif
