Details Panel Overrides UX updates:

Move icons into an editor folder so they don't get cooked
Change the override combo button to match the desired behavior

diff --git a/Engine/Source/Editor/EditorStyle/Private/StarshipStyle.cpp b/Engine/Source/Editor/EditorStyle/Private/StarshipStyle.cpp
--- a/Engine/Source/Editor/EditorStyle/Private/StarshipStyle.cpp
+++ b/Engine/Source/Editor/EditorStyle/Private/StarshipStyle.cpp
@@ -3023,7 +3023,26 @@ void FStarshipEditorStyle::FStyle::SetupPropertyEditorStyles()
 		Set( "DetailsView.CardHeaderRounded", new FSlateRoundedBoxBrush(FStyleColors::Header, FVector4(4.0f, 4.0f, 4.0f, 4.0f )));
 		
 		/*******************************************************************************/
+		
+		// Background images for all the details panels override states
+		Set("DetailsView.OverrideHere", new IMAGE_BRUSH_SVG("Starship/DetailsView/DetailsOverrideHere", Icon16x16, FStyleColors::AccentBlue));
+		Set("DetailsView.OverrideAdded", new IMAGE_BRUSH_SVG("Starship/DetailsView/DetailsOverrideAdded", Icon16x16, FStyleColors::AccentGreen));
+		Set("DetailsView.OverrideNone", new IMAGE_BRUSH_SVG("Starship/DetailsView/DetailsOverrideNone", Icon16x16, FStyleColors::Foreground));
+		Set("DetailsView.OverrideRemoved", new IMAGE_BRUSH_SVG("Starship/DetailsView/DetailsOverrideRemoved", Icon16x16, FStyleColors::Error));
+		Set("DetailsView.OverrideInside", new IMAGE_BRUSH_SVG("Starship/DetailsView/DetailsOverrideInside", Icon16x16, FStyleColors::AccentBlue));
+		Set("DetailsView.OverrideHereInside", new IMAGE_BRUSH_SVG("Starship/DetailsView/DetailsOverrideHereInside", Icon16x16, FStyleColors::AccentBlue));
+		Set("DetailsView.OverrideInherited", new IMAGE_BRUSH_SVG("Starship/DetailsView/DetailsOverrideInherited", Icon16x16, FStyleColors::AccentBlue));
+
+		// Hovered images for all the details panel override states
+		Set("DetailsView.OverrideHere.Hovered", new IMAGE_BRUSH_SVG("Starship/DetailsView/DetailsOverrideHere", Icon16x16, FStyleColors::White));
+		Set("DetailsView.OverrideAdded.Hovered", new IMAGE_BRUSH_SVG("Starship/DetailsView/DetailsOverrideAdded", Icon16x16, FStyleColors::White));
+		Set("DetailsView.OverrideNone.Hovered", new IMAGE_BRUSH_SVG("Starship/DetailsView/DetailsOverrideNone", Icon16x16, FStyleColors::White));
+		Set("DetailsView.OverrideRemoved.Hovered", new IMAGE_BRUSH_SVG("Starship/DetailsView/DetailsOverrideRemoved", Icon16x16, FStyleColors::White));
+		Set("DetailsView.OverrideInside.Hovered", new IMAGE_BRUSH_SVG("Starship/DetailsView/DetailsOverrideInside", Icon16x16, FStyleColors::White));
+		Set("DetailsView.OverrideHereInside.Hovered", new IMAGE_BRUSH_SVG("Starship/DetailsView/DetailsOverrideHereInside", Icon16x16, FStyleColors::White));
+		Set("DetailsView.OverrideInherited.Hovered", new IMAGE_BRUSH_SVG("Starship/DetailsView/DetailsOverrideInherited", Icon16x16, FStyleColors::White));
 
+		
 		const FButtonStyle CategoryRowButton = FButtonStyle()
 			.SetNormal(FSlateRoundedBoxBrush(FStyleColors::Header, 0.f))
 			.SetHovered(FSlateRoundedBoxBrush(FStyleColors::Input, 0.f))
diff --git a/Engine/Source/Editor/PropertyEditor/Private/DetailsViewStyle.cpp b/Engine/Source/Editor/PropertyEditor/Private/DetailsViewStyle.cpp
--- a/Engine/Source/Editor/PropertyEditor/Private/DetailsViewStyle.cpp
+++ b/Engine/Source/Editor/PropertyEditor/Private/DetailsViewStyle.cpp
@@ -9,19 +9,27 @@
 #include "Layout/Visibility.h"
 #include "Styling/AppStyle.h"
 
+#define LOCTEXT_NAMESPACE "DetailsViewStyle"
+
 TArray< TSharedRef< const FOverridesWidgetStyleKey >> FOverridesWidgetStyleKeys::OverridesWidgetStyleKeys;
 
 const FSlateBrush& FOverridesWidgetStyleKey::GetConstStyleBrush() const
 {
-	const FSlateBrush* ImageBrush = FAppStyle::GetBrush( *(TEXT("Icons.Details") + Name.ToString()) );
+	const FSlateBrush* ImageBrush = FAppStyle::GetBrush( *(TEXT("DetailsView.") + Name.ToString()) );
 	checkf(ImageBrush, TEXT("Expecting a valid brush"));
 	return *ImageBrush;
 }
 
-const FOverridesWidgetStyleKey& FOverridesWidgetStyleKeys::Here()
+const FSlateBrush& FOverridesWidgetStyleKey::GetConstStyleBrushHovered() const
 {
-	static const FOverridesWidgetStyleKey Here{"OverrideHere", EOverriddenPropertyOperation::Replace, EOverriddenState::AllOverridden, /*bInherited*/false };
-	return Here;
+	const FSlateBrush* ImageBrush = FAppStyle::GetBrush( *(TEXT("DetailsView.") + Name.ToString() + TEXT(".Hovered")) );
+	checkf(ImageBrush, TEXT("Expecting a valid brush"));
+	return *ImageBrush;
+}
+
+FText FOverridesWidgetStyleKey::GetToolTipText(bool bIsCategory) const
+{
+	return bIsCategory ? CategoryTooltip : PropertyTooltip;
 }
 
 void FOverridesWidgetStyleKeys::Initialize()
@@ -35,15 +43,39 @@ void FOverridesWidgetStyleKeys::Initialize()
 	OverridesWidgetStyleKeys.Add( MakeShared<FOverridesWidgetStyleKey>(Inherited()));
 }
 
+const FOverridesWidgetStyleKey& FOverridesWidgetStyleKeys::Here()
+{
+	static const FOverridesWidgetStyleKey Here{"OverrideHere",
+		EOverriddenPropertyOperation::Replace,
+		EOverriddenState::AllOverridden,
+		false,
+		LOCTEXT("HereCategoryTooltip", "This component has been overridden."),
+		LOCTEXT("HerePropertyTooltip", "This property has been overridden.")
+	};
+	return Here;
+}
+
 const FOverridesWidgetStyleKey& FOverridesWidgetStyleKeys::Added()
 {
-	static const FOverridesWidgetStyleKey Added{"OverrideAdded", EOverriddenPropertyOperation::Add, EOverriddenState::Added };
+	static const FOverridesWidgetStyleKey Added{"OverrideAdded",
+		EOverriddenPropertyOperation::Add,
+		EOverriddenState::Added,
+		TOptional<bool>(),
+		LOCTEXT("AddedCategoryTooltip", "This component has been added."),
+		LOCTEXT("AddedPropertyTooltip", "This property has been added.")
+	};
 	return Added;
 }
 
 const FOverridesWidgetStyleKey& FOverridesWidgetStyleKeys::None()
 {
-	static const FOverridesWidgetStyleKey None{"OverrideNone", EOverriddenPropertyOperation::None, EOverriddenState::NoOverrides };
+	static const FOverridesWidgetStyleKey None{"OverrideNone",
+		EOverriddenPropertyOperation::None,
+		EOverriddenState::NoOverrides,
+		TOptional<bool>(),
+		LOCTEXT("NoneCategoryTooltip", "No override."),
+		LOCTEXT("NonePropertyTooltip", "No override.")
+	};
 	return None;
 }
 
@@ -55,7 +87,13 @@ const FOverridesWidgetStyleKey& FOverridesWidgetStyleKeys::Removed()
 
 const FOverridesWidgetStyleKey& FOverridesWidgetStyleKeys::Inside()
 {
-	static const FOverridesWidgetStyleKey Inside{"OverrideInside", EOverriddenPropertyOperation::Modified, EOverriddenState::HasOverrides };
+	static const FOverridesWidgetStyleKey Inside{"OverrideInside",
+		EOverriddenPropertyOperation::Modified,
+		EOverriddenState::HasOverrides,
+		TOptional<bool>(),
+		LOCTEXT("InsideCategoryTooltip", "At least one of this component's properties has been overridden."),
+		LOCTEXT("InsidePropertyTooltip", "At least one of this property's values has been overridden.")
+	};
 	return Inside;
 }
 
@@ -67,7 +105,13 @@ const FOverridesWidgetStyleKey& FOverridesWidgetStyleKeys::HereInside()
 
 const FOverridesWidgetStyleKey& FOverridesWidgetStyleKeys::Inherited()
 {
-	static const FOverridesWidgetStyleKey Inherited{"OverrideInherited", EOverriddenPropertyOperation::Replace, TOptional<EOverriddenState>(), /*bInherited*/true };
+	static const FOverridesWidgetStyleKey Inherited{"OverrideInherited",
+		EOverriddenPropertyOperation::Replace,
+		TOptional<EOverriddenState>(),
+		true,
+		LOCTEXT("InheritedCategoryTooltip", "This component's parent entity has been overridden."),
+		LOCTEXT("InheritedPropertyTooltip", "This property's parent component has been overridden.")
+	};
 	return Inherited;
 }
 
@@ -117,12 +161,16 @@ FOverridesWidgetStyleKey::FOverridesWidgetStyleKey(FName InName) : Name{InName}
 FOverridesWidgetStyleKey::FOverridesWidgetStyleKey(FName InName,
 													EOverriddenPropertyOperation InOverriddenPropertyOperation,
 													const TOptional<EOverriddenState>& InOverriddenState,
-													const TOptional<bool>& bInStateInherited /*= TOptional<bool>()*/) :
+													const TOptional<bool>& bInStateInherited /*= TOptional<bool>()*/,
+													const FText& InCategoryTooltip,
+													const FText& InPropertyTooltip) :
    Name(InName),
    VisibleOverriddenPropertyOperation(InOverriddenPropertyOperation),
    VisibleOverriddenState(InOverriddenState),
    bStateInherited(bInStateInherited),
-   bCanBeVisible(true)
+   bCanBeVisible(true),
+   CategoryTooltip(InCategoryTooltip),
+   PropertyTooltip(InPropertyTooltip)
 {
 }
 
@@ -313,4 +361,6 @@ const FDetailsViewStyle* FDetailsViewStyle::GetStyle(const FDetailsViewStyleKey&
 					StyleKeyToStyleTemplateMap.Find(FDetailsViewStyleKeys::Default().GetName());
 		
 	return StylePtr ? *StylePtr : nullptr;
-}
\ No newline at end of file
+}
+
+#undef LOCTEXT_NAMESPACE
\ No newline at end of file
diff --git a/Engine/Source/Editor/PropertyEditor/Private/DetailsViewStyle.h b/Engine/Source/Editor/PropertyEditor/Private/DetailsViewStyle.h
--- a/Engine/Source/Editor/PropertyEditor/Private/DetailsViewStyle.h
+++ b/Engine/Source/Editor/PropertyEditor/Private/DetailsViewStyle.h
@@ -34,17 +34,32 @@ public:
  * @param InOverriddenPropertyOperation the overridden property operation for which this style is visible
  * @param InOverriddenState the overridden state for components for which this style is visible
  * @param bInStateInherited take into account if the state was inherited or not, not set means anything
+ * @param InCategoryTooltip tooltip for this widget if it is a category
+ * @param InPropertyTooltip tooltip for this widget if it is a property
+
  */
 	PROPERTYEDITOR_API FOverridesWidgetStyleKey(FName InName, 
 												EOverriddenPropertyOperation InOverriddenPropertyOperation,
 												const TOptional<EOverriddenState>& InOverriddenState,
-												const TOptional<bool>& bInStateInherited = TOptional<bool>());
+												const TOptional<bool>& bInStateInherited = TOptional<bool>(),
+												const FText& InCategoryTooltip = FText::GetEmpty(),
+												const FText& InPropertyTooltip = FText::GetEmpty());
 
 	/**
 	 * returns the const FSlateBrush& that creates the icon for this widget style
 	 */
 	const FSlateBrush& GetConstStyleBrush() const;
 
+	/**
+	* returns the const FSlateBrush& that creates the hovered icon for this widget style
+	*/
+	const FSlateBrush& GetConstStyleBrushHovered() const;
+	
+	/**
+	* returns the tooltip for this widget style
+	*/
+	FText GetToolTipText(bool bIsCategory = false) const;
+
 	/**
 	 * Creates a TAttribute<EVisibility> which indicates the EVisibility for this widget style
 	 * 
@@ -70,11 +85,21 @@ public:
 	const TOptional<EOverriddenState> VisibleOverriddenState;
 
 	/**
-	 * Whetever the state was inherited from a parent node or not
+	 * Whether the state was inherited from a parent node or not
 	 */
 	TOptional<bool> bStateInherited;
 
 	const bool bCanBeVisible = false;
+
+	/**
+	 * Tooltip for this widget if it is a category
+	 */
+	FText CategoryTooltip;
+	
+	/**
+	 * tTooltip for this widget if it is a property
+	 */
+	FText PropertyTooltip;
 };
 
 /**s
diff --git a/Engine/Source/Editor/PropertyEditor/Private/UserInterface/Widgets/OverridesComboButtonBuilder.cpp b/Engine/Source/Editor/PropertyEditor/Private/UserInterface/Widgets/OverridesComboButtonBuilder.cpp
--- a/Engine/Source/Editor/PropertyEditor/Private/UserInterface/Widgets/OverridesComboButtonBuilder.cpp
+++ b/Engine/Source/Editor/PropertyEditor/Private/UserInterface/Widgets/OverridesComboButtonBuilder.cpp
@@ -2,10 +2,89 @@
 
 #include "UserInterface/Widgets/OverridesComboButtonBuilder.h"
 
+#include "Brushes/SlateRoundedBoxBrush.h"
 #include "DetailsViewStyle.h"
-#include "SSimpleComboButton.h"
 #include "Styling/SlateBrush.h"
-#include "Brushes/SlateRoundedBoxBrush.h"
+#include "Widgets/Images/SImage.h"
+#include "Framework/Application/SlateApplication.h"
+#include "Styling/StyleColors.h"
+#include "Widgets/Layout/SBox.h"
+
+/*
+ * We use a custom widget that acts like a combo button, an image that pops up a menu on clicked so that we can have customized hover behavior
+ * to change the foreground color and not show the button's background highlight
+ */
+class SOverrideWidget : public SImage
+{
+public:
+	
+	SLATE_BEGIN_ARGS(SOverrideWidget) {}
+	
+	SLATE_ARGUMENT( const FSlateBrush*, Image )
+	SLATE_ARGUMENT( const FSlateBrush*, HoveredImage )
+	SLATE_EVENT( FOnGetContent, OnGetMenuContent )
+
+	SLATE_END_ARGS()
+
+	const FSlateBrush* GetBrush() const
+	{
+		return IsHovered() ? HoveredImage : Image;
+	}
+
+	FSlateColor GetColor() const
+	{
+		return IsHovered() ? FStyleColors::White : FStyleColors::Foreground;
+	}
+	
+	FReply HandleClick()
+	{
+		// Create the overrides menu at the cursor location
+		if (OnGetMenuContent.IsBound())
+		{
+			FSlateApplication::Get().PushMenu(
+				SharedThis(this),
+				FWidgetPath(),
+				OnGetMenuContent.Execute(),
+				FSlateApplication::Get().GetCursorPos(),
+				FPopupTransitionEffect(FPopupTransitionEffect::TypeInPopup));
+
+			return FReply::Handled();
+		}
+
+		return FReply::Unhandled();
+	}
+
+	/** Called when the mouse button is pressed down on this widget */
+	virtual FReply OnMouseButtonDown(const FGeometry& MyGeometry, const FPointerEvent& MouseEvent) override
+	{
+		if (MouseEvent.GetEffectingButton() != EKeys::LeftMouseButton)
+		{
+			return FReply::Unhandled();
+		}
+
+		return HandleClick();
+	}
+	
+	/** Construct this widget */
+	void Construct(const FArguments& InArgs)
+	{
+		Image = InArgs._Image;
+		HoveredImage = InArgs._HoveredImage;
+		OnGetMenuContent = InArgs._OnGetMenuContent;
+		
+		SImage::Construct(
+		SImage::FArguments()
+			.ColorAndOpacity(this, &SOverrideWidget::GetColor)
+			.Image(this, &SOverrideWidget::GetBrush)
+		);
+	}
+private:
+	const FSlateBrush* Image = nullptr;
+	const FSlateBrush* HoveredImage = nullptr;
+	FOnGetContent OnGetMenuContent;
+};
+
+
 
 FOverridesComboButtonBuilder::FOverridesComboButtonBuilder(
 	TSharedRef<FDetailsDisplayManager> InDetailsDisplayManager,
@@ -30,31 +109,39 @@ TSharedPtr<SWidget> FOverridesComboButtonBuilder::GenerateWidget()
 {
 	const TArray< TSharedRef< const FOverridesWidgetStyleKey >> OverridesWidgetKeys = FOverridesWidgetStyleKeys::GetKeys();
 	
-	if ( !Object.IsValid() || OverridesWidgetKeys.IsEmpty() )
+	if (!Object.IsValid() || OverridesWidgetKeys.IsEmpty())
 	{
 		return SNullWidget::NullWidget;
 	}
 
-	TSharedRef<SHorizontalBox> Box =	SNew(SHorizontalBox).Visibility(IsVisible);
+	TSharedRef<SHorizontalBox> Box = SNew(SHorizontalBox).Visibility(IsVisible);
 	
 	for (const TSharedRef< const FOverridesWidgetStyleKey >& Key : OverridesWidgetKeys )
 	{
-		if ( Key->bCanBeVisible )
+		if (Key->bCanBeVisible)
 		{
 			Box->AddSlot()
-			   .HAlign(HAlign_Right)
-			   .VAlign(VAlign_Center)
-			   .AutoWidth()
-			[
-				SNew( SSimpleComboButton )
-					.Visibility( Key->GetVisibilityAttribute(EditPropertyChain, Object ))
+				.HAlign(HAlign_Center)
+				.VAlign(VAlign_Center)
+				.AutoWidth()
+				[
+					SNew(SOverrideWidget)
+					.Image(&Key->GetConstStyleBrush())
+					.HoveredImage(&Key->GetConstStyleBrushHovered())
+					.Visibility(Key->GetVisibilityAttribute(EditPropertyChain, Object))
+					.ToolTipText(Key->GetToolTipText(bIsCategory))
 					.OnGetMenuContent(OnGetContent)
-					.Icon(&Key->GetConstStyleBrush())
-			];
+				];
 		}
 	}
 
-	return Box.ToSharedPtr();
+	// We use an SBox to force our size to fill the entire slot, otherwise the button is not centered on the column
+	return SNew(SBox)
+			.WidthOverride(20.0f)
+			.HeightOverride(20.0f)
+			[
+				Box
+			];
 }
 
 TSharedRef<SWidget> FOverridesComboButtonBuilder::operator*()
diff --git a/Engine/Source/Runtime/SlateCore/Private/Styling/StarshipCoreStyle.cpp b/Engine/Source/Runtime/SlateCore/Private/Styling/StarshipCoreStyle.cpp
--- a/Engine/Source/Runtime/SlateCore/Private/Styling/StarshipCoreStyle.cpp
+++ b/Engine/Source/Runtime/SlateCore/Private/Styling/StarshipCoreStyle.cpp
@@ -413,15 +413,6 @@ TSharedRef<ISlateStyle> FStarshipCoreStyle::Create()
 
 		Style->Set("Icons.BadgeModified", new IMAGE_BRUSH_SVG("Starship/Common/badge-modified", Icon16x16));
 		
-		// Background images for all the details panels override states
-		Style->Set("Icons.DetailsOverrideHere", new IMAGE_BRUSH_SVG("Starship/Common/DetailsOverrideHere", Icon16x16, FStyleColors::AccentBlue));
-		Style->Set("Icons.DetailsOverrideAdded", new IMAGE_BRUSH_SVG("Starship/Common/DetailsOverrideAdded", Icon16x16, FStyleColors::AccentGreen));
-		Style->Set("Icons.DetailsOverrideNone", new IMAGE_BRUSH_SVG("Starship/Common/DetailsOverrideNone", Icon16x16, FStyleColors::Foreground));
-		Style->Set("Icons.DetailsOverrideRemoved", new IMAGE_BRUSH_SVG("Starship/Common/DetailsOverrideRemoved", Icon16x16, FStyleColors::Error));
-		Style->Set("Icons.DetailsOverrideInside", new IMAGE_BRUSH_SVG("Starship/Common/DetailsOverrideInside", Icon16x16, FStyleColors::AccentBlue));
-		Style->Set("Icons.DetailsOverrideHereInside", new IMAGE_BRUSH_SVG("Starship/Common/DetailsOverrideHereInside", Icon16x16, FStyleColors::AccentBlue));
-		Style->Set("Icons.DetailsOverrideInherited", new IMAGE_BRUSH_SVG("Starship/Common/DetailsOverrideInherited", Icon16x16, FStyleColors::AccentBlue));
-		
 		// Toolbar Size Icons
 		Style->Set("Icons.Toolbar.Settings", new IMAGE_BRUSH_SVG("Starship/Common/settings", Icon20x20));
 
