Details Panel - Add API to mutate just the property name customization for a specific property path for non-nested properties

- created FDetailsNameWidgetOverrideCustomization and function to get it on FDetailsViewPrivate, SDetailsView and FDetailsViewArgs to provide a customization that can override any Name widget
- FDetailsNameWidgetOverrideCustomization can override any non-nested property, using the CustomizeName function which takes a TSharedRef<SWidget> that is the name widget that would show without the new customization, and the FPropertyPath in order to be able to tell what properties to override.
- In SDetalsSingleItemRow added GetNameWidget to use the FDetailsNameWidgetOverrideCustomization to provide the name widget, if it has been provided

diff --git a/Engine/Source/Editor/PropertyEditor/Private/DetailsNameWidgetOverrideCustomization.cpp b/Engine/Source/Editor/PropertyEditor/Private/DetailsNameWidgetOverrideCustomization.cpp
new file mode 100644
--- /dev/null
+++ b/Engine/Source/Editor/PropertyEditor/Private/DetailsNameWidgetOverrideCustomization.cpp
@@ -0,0 +1,14 @@
+// Copyright Epic Games, Inc. All Rights Reserved.
+
+#include "DetailsNameWidgetOverrideCustomization.h"
+#include "PropertyPath.h"
+#include "Widgets/SWidget.h"
+
+FDetailsNameWidgetOverrideCustomization::FDetailsNameWidgetOverrideCustomization()
+{
+}
+
+TSharedRef<SWidget> FDetailsNameWidgetOverrideCustomization::CustomizeName(TSharedRef<SWidget> InnerNameContent, FPropertyPath& Path)
+{
+	return InnerNameContent;
+}
diff --git a/Engine/Source/Editor/PropertyEditor/Private/IDetailsViewPrivate.h b/Engine/Source/Editor/PropertyEditor/Private/IDetailsViewPrivate.h
--- a/Engine/Source/Editor/PropertyEditor/Private/IDetailsViewPrivate.h
+++ b/Engine/Source/Editor/PropertyEditor/Private/IDetailsViewPrivate.h
@@ -8,6 +8,7 @@
 #include "IDetailsView.h"
 #include "PropertyNode.h"
 
+class FDetailsNameWidgetOverrideCustomization;
 class FEditConditionParser;
 class FNotifyHook;
 class IDetailPropertyExtensionHandler;
@@ -207,4 +208,8 @@ public:
 	 */
 	virtual TSharedPtr<FDetailsDisplayManager> GetDisplayManager() = 0;
 
+	virtual TSharedPtr<FDetailsNameWidgetOverrideCustomization> GetDetailsNameWidgetOverrideCustomization()
+	{
+		return nullptr;
+	}
 };
diff --git a/Engine/Source/Editor/PropertyEditor/Private/SDetailSingleItemRow.cpp b/Engine/Source/Editor/PropertyEditor/Private/SDetailSingleItemRow.cpp
--- a/Engine/Source/Editor/PropertyEditor/Private/SDetailSingleItemRow.cpp
+++ b/Engine/Source/Editor/PropertyEditor/Private/SDetailSingleItemRow.cpp
@@ -5,6 +5,7 @@
 #include "Algo/Compare.h"
 #include "DetailGroup.h"
 #include "DetailPropertyRow.h"
+#include "DetailsNameWidgetOverrideCustomization.h"
 #include "DetailWidgetRow.h"
 #include "Editor.h"
 #include "IDetailDragDropHandler.h"
@@ -536,7 +537,7 @@ void SDetailSingleItemRow::Construct( const FArguments& InArgs, FDetailLayoutCus
 					.Value(ColumnSizeData.GetNameColumnWidth())
 					.OnSlotResized(ColumnSizeData.GetOnNameColumnResized())
 					[
-						NameColumnBox
+						GetNameWidget(NameColumnBox, GetPropertyNode())
 					];
 
 				// create Value column:
@@ -1378,6 +1379,27 @@ bool SDetailSingleItemRow::CanPasteProperty() const
 	return CanPasteFromText(TEXT(""), ClipboardContent);
 }
 
+TSharedRef<SWidget> SDetailSingleItemRow::GetNameWidget(TSharedRef<SWidget> NameWidget, const TSharedPtr<FPropertyNode>& Node) const
+{
+	 if (Node.IsValid() && OwnerTreeNode.IsValid() )
+     {
+	 	const TSharedPtr<FDetailTreeNode> DetailTreeNodeSP = OwnerTreeNode.Pin();
+	 	
+	     if (DetailTreeNodeSP.IsValid() && DetailTreeNodeSP->GetDetailsView())
+	     {
+	     	const TSharedPtr<FDetailsNameWidgetOverrideCustomization> DetailsNameWidgetOverrideCustomization =
+				DetailTreeNodeSP->GetDetailsView()->GetDetailsNameWidgetOverrideCustomization();
+
+	     	if ( DetailsNameWidgetOverrideCustomization.IsValid() )
+	     	{
+	     		const TSharedRef< FPropertyPath > Path = FPropertyNode::CreatePropertyPath(Node.ToSharedRef());
+	     		return DetailsNameWidgetOverrideCustomization->CustomizeName(NameWidget, Path.Get());
+	     	}
+	     }
+     }
+	return NameWidget;
+}
+
 bool SDetailSingleItemRow::CanPasteFromText(const FString& InTag, const FString& InText) const
 {
 	if (InText.IsEmpty())
diff --git a/Engine/Source/Editor/PropertyEditor/Private/SDetailSingleItemRow.h b/Engine/Source/Editor/PropertyEditor/Private/SDetailSingleItemRow.h
--- a/Engine/Source/Editor/PropertyEditor/Private/SDetailSingleItemRow.h
+++ b/Engine/Source/Editor/PropertyEditor/Private/SDetailSingleItemRow.h
@@ -86,6 +86,15 @@ private:
 	void OnPasteGroup();
 	bool CanPasteGroup();
 
+	/**
+	 * Returns the TSharedRef<SWidget> for the Name column widget
+	 * 
+	 * @param NameWidget 
+	 * @param Node 
+	 * @return 
+	 */
+	TSharedRef<SWidget> GetNameWidget(TSharedRef<SWidget> NameWidget, const TSharedPtr<FPropertyNode>& Node) const;
+
 	/**
 	 * @return True if the (optionally tagged) input contents can be pasted
 	 */
diff --git a/Engine/Source/Editor/PropertyEditor/Private/SDetailsView.cpp b/Engine/Source/Editor/PropertyEditor/Private/SDetailsView.cpp
--- a/Engine/Source/Editor/PropertyEditor/Private/SDetailsView.cpp
+++ b/Engine/Source/Editor/PropertyEditor/Private/SDetailsView.cpp
@@ -6,6 +6,7 @@
 #include "Settings/EditorStyleSettings.h"
 #include "DetailCategoryBuilderImpl.h"
 #include "DetailLayoutBuilderImpl.h"
+#include "DetailsNameWidgetOverrideCustomization.h"
 #include "DetailsViewGenericObjectFilter.h"
 #include "DetailsViewPropertyGenerationUtilities.h"
 #include "Editor.h"
@@ -56,6 +57,7 @@ SDetailsView::~SDetailsView()
 void SDetailsView::Construct(const FArguments& InArgs, const FDetailsViewArgs& InDetailsViewArgs)
 {
 	DetailsViewArgs = InDetailsViewArgs;
+	DetailsNameWidgetOverrideCustomization = DetailsViewArgs.DetailsNameWidgetOverrideCustomization;
 
 	const FDetailsViewConfig* ViewConfig = GetConstViewConfig();
 	if (ViewConfig != nullptr)
@@ -1509,6 +1511,11 @@ bool SDetailsView::IsDefaultStyle() const
 	return StyleKeySP.IsValid() && *StyleKeySP.Get() == GetPrimaryDetailsViewStyleKey();
 }
 
+TSharedPtr<FDetailsNameWidgetOverrideCustomization> SDetailsView::GetDetailsNameWidgetOverrideCustomization()
+{
+	return DetailsNameWidgetOverrideCustomization;
+}
+
 void SDetailsView::UpdateStyleKey()
 {
 	const bool IsDefault = FDetailsViewStyleKeys::IsDefault(DisplayManager->GetDetailsViewStyleKey());
diff --git a/Engine/Source/Editor/PropertyEditor/Private/SDetailsView.h b/Engine/Source/Editor/PropertyEditor/Private/SDetailsView.h
--- a/Engine/Source/Editor/PropertyEditor/Private/SDetailsView.h
+++ b/Engine/Source/Editor/PropertyEditor/Private/SDetailsView.h
@@ -133,6 +133,10 @@ public:
 	virtual bool IsDefaultStyle() const override;
 
 private:
+	/**
+	 * returns a pointer to a Customization that can be applied to all name widgets, if one exists
+	 */
+	virtual TSharedPtr<FDetailsNameWidgetOverrideCustomization> GetDetailsNameWidgetOverrideCustomization() override;
 
 	void SetObjectArrayPrivate(const TArray<UObject*>& InObjects);
 
@@ -240,6 +244,9 @@ private:
 	*/
 	TSharedPtr<FDetailsViewStyleKey> StyleKeySP;
 	
+	/** A Customization that can be applied to all name widgets */
+	TSharedPtr<FDetailsNameWidgetOverrideCustomization> DetailsNameWidgetOverrideCustomization;;
+	
 	/** The filter for objects viewed by this details panel */
 	TSharedPtr<FDetailsViewObjectFilter> ObjectFilter;
 
diff --git a/Engine/Source/Editor/PropertyEditor/Public/DetailsNameWidgetOverrideCustomization.h b/Engine/Source/Editor/PropertyEditor/Public/DetailsNameWidgetOverrideCustomization.h
new file mode 100644
--- /dev/null
+++ b/Engine/Source/Editor/PropertyEditor/Public/DetailsNameWidgetOverrideCustomization.h
@@ -0,0 +1,34 @@
+// Copyright Epic Games, Inc. All Rights Reserved.
+
+#pragma once
+
+#include "Templates/SharedPointer.h"
+
+class FPropertyPath;
+class SWidget;
+
+/**
+ * A customization applied to the whole details panel which can update any Name widget
+ */
+class FDetailsNameWidgetOverrideCustomization : public TSharedFromThis<FDetailsNameWidgetOverrideCustomization>
+{
+public:
+	/**
+	 * default constructor
+	 */
+	PROPERTYEDITOR_API explicit FDetailsNameWidgetOverrideCustomization();
+
+	/**
+	 * destructor
+	 */
+	PROPERTYEDITOR_API virtual ~FDetailsNameWidgetOverrideCustomization() = default;
+
+	/**
+	 * @param InnerNameContent the unaltered content of the Name widget. It shows what would
+	 * be there normally with no modification by this customization
+	 * @param Path the FPropertyPath for the current property
+	 * @return the new TSharedRef<SWidget> for the Name widget if one is provided, else it returns
+	 * the default Name widget for the current property path
+	 */
+	 PROPERTYEDITOR_API virtual TSharedRef<SWidget> CustomizeName( TSharedRef<SWidget> InnerNameContent,  FPropertyPath& Path);
+};
\ No newline at end of file
diff --git a/Engine/Source/Editor/PropertyEditor/Public/DetailsViewArgs.h b/Engine/Source/Editor/PropertyEditor/Public/DetailsViewArgs.h
--- a/Engine/Source/Editor/PropertyEditor/Public/DetailsViewArgs.h
+++ b/Engine/Source/Editor/PropertyEditor/Public/DetailsViewArgs.h
@@ -6,6 +6,7 @@
 #include "DetailsViewStyleKey.h"
 #include "Widgets/Layout/SScrollBar.h"
 
+class FDetailsNameWidgetOverrideCustomization;
 class FDetailsViewObjectFilter;
 class FNotifyHook;
 class FTabManager;
@@ -57,6 +58,8 @@ struct FDetailsViewArgs
 	EHorizontalAlignment ScrollbarAlignment;
 	/** A @code TSharedPtr<FDetailsViewStyleKey> @endcode which provides the key to the Details View's style */
 	TSharedPtr<FDetailsViewStyleKey> StyleKey;
+	/** A Customization that can be applied to all name widgets */
+	TSharedPtr<FDetailsNameWidgetOverrideCustomization> DetailsNameWidgetOverrideCustomization;
 	/** Identifier for this details view; NAME_None if this view is anonymous */
 	FName ViewIdentifier;
 	/** Notify hook to call when properties are changed */
@@ -117,6 +120,7 @@ public:
 	FDetailsViewArgs()
 		: DefaultsOnlyVisibility(EEditDefaultsOnlyNodeVisibility::Show)
 		, ScrollbarAlignment(HAlign_Right)
+	    , DetailsNameWidgetOverrideCustomization(nullptr)
 		, ViewIdentifier(NAME_None)
 		, NotifyHook(nullptr)
 		, NameAreaSettings(ActorsUseNameArea)
