Crash on paste of subcategory in the details panel.

Currently pasting a subcategory will result in a crash due to not checking the OnPasteFromTextDelegate for validity. We should check it for validity in SDetailCategoryTableRow prior to broadcasting.

diff --git a/Engine/Source/Editor/PropertyEditor/Private/SDetailCategoryTableRow.cpp b/Engine/Source/Editor/PropertyEditor/Private/SDetailCategoryTableRow.cpp
--- a/Engine/Source/Editor/PropertyEditor/Private/SDetailCategoryTableRow.cpp
+++ b/Engine/Source/Editor/PropertyEditor/Private/SDetailCategoryTableRow.cpp
@@ -434,7 +434,10 @@ void SDetailCategoryTableRow::OnPasteCategory()
 			const FGuid OperationGuid = FGuid::NewGuid();			
 			for (const TPair<FName, FString>& KVP : PreviousClipboardData.PropertyValues)
 			{
-				OnPasteFromTextDelegate->Broadcast(KVP.Key.ToString(), KVP.Value, OperationGuid);
+				if ( OnPasteFromTextDelegate.IsValid() )
+				{
+					OnPasteFromTextDelegate->Broadcast(KVP.Key.ToString(), KVP.Value, OperationGuid);
+				}
 			}
 		}
 		
