Fixed paste on properties where the custom EditConditionValue isn't set. EditConditionValue would always return false and pasting wouldn't be possible. The fix is to return true if EditConditionValue has never been set.

#changelist validated
#virtualized

diff --git a/Engine/Source/Editor/PropertyEditor/Private/SDetailSingleItemRow.cpp b/Engine/Source/Editor/PropertyEditor/Private/SDetailSingleItemRow.cpp
--- a/Engine/Source/Editor/PropertyEditor/Private/SDetailSingleItemRow.cpp
+++ b/Engine/Source/Editor/PropertyEditor/Private/SDetailSingleItemRow.cpp
@@ -1084,7 +1084,7 @@ void SDetailSingleItemRow::PopulateContextMenu(UToolMenu* ToolMenu)
 			// Paste
 			{
 				// Paste is only enabled if property editing is enabled
-				if (OwningDetailsView && OwningDetailsView->IsPropertyEditingEnabled() && WidgetRow.EditConditionValue.Get())
+				if (OwningDetailsView && OwningDetailsView->IsPropertyEditingEnabled() && WidgetRow.EditConditionValue.Get(true /*DefaultValue*/))
 				{
 					TAttribute<FText> Label;
 					TAttribute<FText> ToolTip;
