Edit Condition Parser: add support for INDEX_NONE parsing.

diff --git a/Engine/Source/Editor/PropertyEditor/Private/EditConditionParser.cpp b/Engine/Source/Editor/PropertyEditor/Private/EditConditionParser.cpp
--- a/Engine/Source/Editor/PropertyEditor/Private/EditConditionParser.cpp
+++ b/Engine/Source/Editor/PropertyEditor/Private/EditConditionParser.cpp
@@ -70,6 +70,17 @@ static TOptional<FExpressionError> ConsumeNullPtr(FExpressionTokenConsumer& Cons
 	return TOptional<FExpressionError>();
 }
 
+static TOptional<FExpressionError> ConsumeIndexNone(FExpressionTokenConsumer& Consumer)
+{
+	TOptional<FStringToken> IndexNoneToken = Consumer.GetStream().ParseToken(TEXT("INDEX_NONE"));
+	if (IndexNoneToken.IsSet())
+	{
+		Consumer.Add(IndexNoneToken.GetValue(), EditConditionParserTokens::FIndexNoneToken());
+	}
+
+	return TOptional<FExpressionError>();
+}
+
 static TOptional<FExpressionError> ConsumePropertyName(FExpressionTokenConsumer& Consumer)
 {
 	enum class EParsedStringType : uint8
@@ -321,6 +332,31 @@ static FExpressionResult ApplyPropertyIsNull(const EditConditionParserTokens::FP
 	return MakeValue(bIsNull);
 }
 
+static FExpressionResult ApplyPropertyIsIndexNone(const EditConditionParserTokens::FPropertyToken& Property, const IEditConditionContext& Context, bool bNegate)
+{
+	TWeakObjectPtr<UFunction> CachedFunction = Context.GetFunction(Property.PropertyName);
+
+	TOptional<FString> TypeName = Context.GetTypeName(Property.PropertyName, CachedFunction);
+	if (!TypeName.IsSet())
+	{
+		return MakeError(FText::Format(LOCTEXT("InvalidOperand", "EditCondition attempted to use an invalid operand \"{0}\"."), FText::FromString(Property.PropertyName)));
+	}
+
+	TOptional<int64> Value = Context.GetIntegerValue(Property.PropertyName, CachedFunction);
+	if (!Value.IsSet())
+	{
+		return MakeError(FText::Format(LOCTEXT("InvalidOperand", "EditCondition attempted to use an invalid operand \"{0}\"."), FText::FromString(Property.PropertyName)));
+	}
+
+	bool bIsIndexNone = Value.GetValue() == (int64)INDEX_NONE;
+	if (bNegate)
+	{
+		bIsIndexNone = !bIsIndexNone;
+	}
+
+	return MakeValue(bIsIndexNone);
+}
+
 static FExpressionResult ApplyPropertiesEqual(const EditConditionParserTokens::FPropertyToken& A, const EditConditionParserTokens::FPropertyToken& B, const IEditConditionContext& Context, bool bNegate)
 {
 	TWeakObjectPtr<UFunction> CachedFunctionA = Context.GetFunction(A.PropertyName);
@@ -730,6 +766,7 @@ FEditConditionParser::FEditConditionParser()
 	TokenDefinitions.DefineToken(&ExpressionParser::ConsumeSymbol<FSubExpressionEnd>);
 	TokenDefinitions.DefineToken(&ExpressionParser::ConsumeNumber);
 	TokenDefinitions.DefineToken(&ConsumeNullPtr);
+	TokenDefinitions.DefineToken(&ConsumeIndexNone);
 	TokenDefinitions.DefineToken(&ConsumeBool);
 	TokenDefinitions.DefineToken(&ConsumePropertyName);
 				
@@ -771,6 +808,17 @@ FEditConditionParser::FEditConditionParser()
 		return ApplyPropertyIsNull(A, *Context, true);
 	});
 
+	// INDEX_NONE
+	OperatorJumpTable.MapBinary<FEqual>([](const FPropertyToken& A, const FIndexNoneToken& B, const IEditConditionContext* Context) -> FExpressionResult
+	{
+		return ApplyPropertyIsIndexNone(A, *Context, false);
+	});
+
+	OperatorJumpTable.MapBinary<FNotEqual>([](const FPropertyToken& A, const FIndexNoneToken& B, const IEditConditionContext* Context) -> FExpressionResult
+	{
+		return ApplyPropertyIsIndexNone(A, *Context, true);
+	});
+
 	// BITWISE AND
 	OperatorJumpTable.MapBinary<FBitwiseAnd>([](const FPropertyToken& A, const FEnumToken& B, const IEditConditionContext* Context) -> FExpressionResult
 	{
diff --git a/Engine/Source/Editor/PropertyEditor/Private/EditConditionParser.h b/Engine/Source/Editor/PropertyEditor/Private/EditConditionParser.h
--- a/Engine/Source/Editor/PropertyEditor/Private/EditConditionParser.h
+++ b/Engine/Source/Editor/PropertyEditor/Private/EditConditionParser.h
@@ -56,6 +56,13 @@ namespace EditConditionParserTokens
 		FNullPtrToken(const FNullPtrToken&) = default;
 		FNullPtrToken& operator=(const FNullPtrToken&) = default;
 	};
+
+	struct FIndexNoneToken
+	{
+		FIndexNoneToken() = default;
+		FIndexNoneToken(const FIndexNoneToken&) = default;
+		FIndexNoneToken& operator=(const FIndexNoneToken&) = default;
+	};
 }
 
 #define DEFINE_EDIT_CONDITION_NODE(TYPE, ...) \
@@ -82,6 +89,7 @@ namespace EditConditionParserTokens
 	DEFINE_EXPRESSION_NODE_TYPE(EditConditionParserTokens::FPropertyToken, 0x9A3FAF6F, 0xB2E45E4D, 0xA80A70C6, 0x47A89BD7)
 	DEFINE_EXPRESSION_NODE_TYPE(EditConditionParserTokens::FEnumToken, 0xC9A35C24, 0x21FC904B, 0x9F1B2B6A, 0xDF6F4BC4)
 	DEFINE_EXPRESSION_NODE_TYPE(EditConditionParserTokens::FNullPtrToken, 0xC3E91735, 0xB3B90C4D, 0x99BA3DCA, 0xC6CC6231)
+	DEFINE_EXPRESSION_NODE_TYPE(EditConditionParserTokens::FIndexNoneToken, 0x6A9CA107, 0x8ED649E4, 0xAF4525CA, 0xEBF84F97)
 
 #undef DEFINE_EDIT_CONDITION_NODE
 
