Export PropertyUpdatedWidgetBuilder public functions

diff --git a/Engine/Source/Editor/PropertyEditor/Public/UserInterface/Widgets/PropertyUpdatedWidgetBuilder.h b/Engine/Source/Editor/PropertyEditor/Public/UserInterface/Widgets/PropertyUpdatedWidgetBuilder.h
--- a/Engine/Source/Editor/PropertyEditor/Public/UserInterface/Widgets/PropertyUpdatedWidgetBuilder.h
+++ b/Engine/Source/Editor/PropertyEditor/Public/UserInterface/Widgets/PropertyUpdatedWidgetBuilder.h
@@ -15,31 +15,31 @@ public:
 	/**
 	 * default constructor
 	 */
-	FPropertyUpdatedWidgetBuilder();
+	PROPERTYEDITOR_API FPropertyUpdatedWidgetBuilder();
 
 	/**
 	 * Set the IsVisible delegate for the menu that this button is responsible for
 	 */
-	FPropertyUpdatedWidgetBuilder& Bind_IsVisible(TAttribute<EVisibility> InIsVisible);
+	PROPERTYEDITOR_API FPropertyUpdatedWidgetBuilder& Bind_IsVisible(TAttribute<EVisibility> InIsVisible);
 
 	/**
 	* Sets the ResetToDefault delegate
 	*
 	* @param InResetToDefault the FResetToDefault delegate to set for this widget builder
 	*/
-	FPropertyUpdatedWidgetBuilder& Set_ResetToDefault(FResetToDefault InResetToDefault);
+	PROPERTYEDITOR_API FPropertyUpdatedWidgetBuilder& Set_ResetToDefault(FResetToDefault InResetToDefault);
 
 	/**
 	 * returns the ResetToDefault delegate
 	 */
-	FResetToDefault Get_ResetToDefault();
+	PROPERTYEDITOR_API FResetToDefault Get_ResetToDefault();
 
 	/**
 	 *  binds the IsRowHovered Attribute
 	 *
 	 * @param bInIsRowHoveredAttr the attribute which will determine if the row for this property is hovered
 	 */
-	void Bind_IsRowHovered(TAttribute<bool> InIsRowHoveredAttr);
+	PROPERTYEDITOR_API void Bind_IsRowHovered(TAttribute<bool> InIsRowHoveredAttr);
 
 	/**
 	 * If true, the row that this is associated with is currently hovered over
