Exposed RequestForceRefresh for detail customizations to use

This was available within SDetailsViewBase, but the only publicly available function was ForceRefresh which isn't deferred

Using RequestForceRefresh is preferable when the refresh can be deferred, as it can prevent performance problems from refreshing the view multiple times within a single frame

diff --git a/Engine/Source/Editor/PropertyEditor/Private/IDetailsViewPrivate.h b/Engine/Source/Editor/PropertyEditor/Private/IDetailsViewPrivate.h
--- a/Engine/Source/Editor/PropertyEditor/Private/IDetailsViewPrivate.h
+++ b/Engine/Source/Editor/PropertyEditor/Private/IDetailsViewPrivate.h
@@ -47,6 +47,9 @@ public:
 	 */
 	virtual TSharedPtr<IPropertyUtilities> GetPropertyUtilities() = 0;
 
+	/** Request the details view to be refreshed (new widgets generated) with the current set of objects on the next Tick */
+	virtual void RequestForceRefresh() = 0;
+
 	/** Causes the details view to be refreshed (new widgets generated) with the current set of objects */
 	virtual void ForceRefresh() = 0;
 
diff --git a/Engine/Source/Editor/PropertyEditor/Private/Presentation/PropertyTable/PropertyTable.cpp b/Engine/Source/Editor/PropertyEditor/Private/Presentation/PropertyTable/PropertyTable.cpp
--- a/Engine/Source/Editor/PropertyEditor/Private/Presentation/PropertyTable/PropertyTable.cpp
+++ b/Engine/Source/Editor/PropertyEditor/Private/Presentation/PropertyTable/PropertyTable.cpp
@@ -84,6 +84,12 @@ void FPropertyTable::RequestRefresh()
 	bRefreshRequested = true;
 }
 
+void FPropertyTable::RequestForceRefresh()
+{
+	// RequestRefresh is already a deferred ForceRefresh
+	RequestRefresh();
+}
+
 class FNotifyHook* FPropertyTable::GetNotifyHook() const
 {
 	return NULL;
diff --git a/Engine/Source/Editor/PropertyEditor/Private/Presentation/PropertyTable/PropertyTable.h b/Engine/Source/Editor/PropertyEditor/Private/Presentation/PropertyTable/PropertyTable.h
--- a/Engine/Source/Editor/PropertyEditor/Private/Presentation/PropertyTable/PropertyTable.h
+++ b/Engine/Source/Editor/PropertyEditor/Private/Presentation/PropertyTable/PropertyTable.h
@@ -21,6 +21,7 @@ public:
 
 	virtual void ForceRefresh() override;
 	virtual void RequestRefresh() override;
+	virtual void RequestForceRefresh() override;
 
 	virtual class FNotifyHook* GetNotifyHook() const override;
 	virtual bool AreFavoritesEnabled() const override;
diff --git a/Engine/Source/Editor/PropertyEditor/Private/PropertyRowGenerator.cpp b/Engine/Source/Editor/PropertyEditor/Private/PropertyRowGenerator.cpp
--- a/Engine/Source/Editor/PropertyEditor/Private/PropertyRowGenerator.cpp
+++ b/Engine/Source/Editor/PropertyEditor/Private/PropertyRowGenerator.cpp
@@ -65,6 +65,11 @@ public:
 			Generator->RequestRefresh();
 		}
 	}
+	virtual void RequestForceRefresh() override
+	{
+		// RequestRefresh is already a deferred ForceRefresh
+		RequestRefresh();
+	}
 
 	virtual TSharedPtr<class FAssetThumbnailPool> GetThumbnailPool() const override
 	{
diff --git a/Engine/Source/Editor/PropertyEditor/Private/SDetailsView.cpp b/Engine/Source/Editor/PropertyEditor/Private/SDetailsView.cpp
--- a/Engine/Source/Editor/PropertyEditor/Private/SDetailsView.cpp
+++ b/Engine/Source/Editor/PropertyEditor/Private/SDetailsView.cpp
@@ -515,6 +515,9 @@ EVisibility SDetailsView::GetActorNameAreaVisibility() const
 void SDetailsView::ForceRefresh()
 {
 	TRACE_CPUPROFILER_EVENT_SCOPE_STR("SDetailsView::ForceRefresh");
+
+	ClearPendingRefreshTimer();
+
 	TArray<UObject*> NewObjectList;
 	NewObjectList.Reserve(UnfilteredSelectedObjects.Num());
 	TArray<TWeakObjectPtr<UObject>> ValidSelectedObjects;
diff --git a/Engine/Source/Editor/PropertyEditor/Private/SDetailsViewBase.cpp b/Engine/Source/Editor/PropertyEditor/Private/SDetailsViewBase.cpp
--- a/Engine/Source/Editor/PropertyEditor/Private/SDetailsViewBase.cpp
+++ b/Engine/Source/Editor/PropertyEditor/Private/SDetailsViewBase.cpp
@@ -29,7 +29,6 @@ SDetailsViewBase::SDetailsViewBase() :
 	, bHasOpenColorPicker(false)
 	, bDisableCustomDetailLayouts(false)
 	, bPendingCleanupTimerSet(false)
-	, bPendingRefreshTimerSet(false)
 	, bRunningDeferredActions(false)
 {
 	UDetailsConfig::Initialize();
@@ -45,9 +44,9 @@ SDetailsViewBase::SDetailsViewBase() :
 		CurrentFilter.bShowOnlyModified = ViewConfig->bShowOnlyModified;
 	}
 
-	// Avoid calling ForceRefresh until next tick otherwise the editor can lock up for several minutes in one frame refreshing multiple times
-	PropertyPermissionListChangedDelegate = FPropertyEditorPermissionList::Get().PermissionListUpdatedDelegate.AddLambda([this](TSoftObjectPtr<UStruct> Struct, FName Owner) { SetPendingRefreshTimer(); });
-	PropertyPermissionListEnabledDelegate = FPropertyEditorPermissionList::Get().PermissionListEnabledDelegate.AddRaw(this, &SDetailsViewBase::SetPendingRefreshTimer);
+	// RequestForceRefresh is deferred until next tick to avoid the editor locking up for several minutes in one frame when refreshing multiple times
+	PropertyPermissionListChangedDelegate = FPropertyEditorPermissionList::Get().PermissionListUpdatedDelegate.AddLambda([this](TSoftObjectPtr<UStruct> Struct, FName Owner) { RequestForceRefresh(); });
+	PropertyPermissionListEnabledDelegate = FPropertyEditorPermissionList::Get().PermissionListEnabledDelegate.AddRaw(this, &SDetailsViewBase::RequestForceRefresh);
 }
 
 SDetailsViewBase::~SDetailsViewBase()
@@ -802,6 +801,11 @@ void SDetailsViewBase::RefreshTree()
 	DetailTree->RequestTreeRefresh();
 }
 
+void SDetailsViewBase::RequestForceRefresh()
+{
+	SetPendingRefreshTimer();
+}
+
 void SDetailsViewBase::SaveCustomExpansionState(const FString& NodePath, bool bIsExpanded)
 {
 	if (bIsExpanded)
@@ -1117,19 +1121,30 @@ void SDetailsViewBase::HandlePendingCleanup()
 
 void SDetailsViewBase::SetPendingRefreshTimer()
 {
-	if (!bPendingRefreshTimerSet)
+	if (!PendingRefreshTimerHandle.IsValid())
+	{
+		if (GEditor && GEditor->IsTimerManagerValid())
+		{
+			PendingRefreshTimerHandle = GEditor->GetTimerManager()->SetTimerForNextTick(FTimerDelegate::CreateSP(this, &SDetailsViewBase::HandlePendingRefreshTimer));
+		}
+	}
+}
+
+void SDetailsViewBase::ClearPendingRefreshTimer()
+{
+	if (PendingRefreshTimerHandle.IsValid())
 	{
 		if (GEditor && GEditor->IsTimerManagerValid())
 		{
-			bPendingRefreshTimerSet = true;
-			GEditor->GetTimerManager()->SetTimerForNextTick(FTimerDelegate::CreateSP(this, &SDetailsViewBase::HandlePendingRefreshTimer));
+			GEditor->GetTimerManager()->ClearTimer(PendingRefreshTimerHandle);
 		}
+		PendingRefreshTimerHandle.Invalidate();
 	}
 }
 
 void SDetailsViewBase::HandlePendingRefreshTimer()
 {
-	bPendingRefreshTimerSet = false;
+	PendingRefreshTimerHandle.Invalidate();
 	ForceRefresh();
 }
 
diff --git a/Engine/Source/Editor/PropertyEditor/Private/SDetailsViewBase.h b/Engine/Source/Editor/PropertyEditor/Private/SDetailsViewBase.h
--- a/Engine/Source/Editor/PropertyEditor/Private/SDetailsViewBase.h
+++ b/Engine/Source/Editor/PropertyEditor/Private/SDetailsViewBase.h
@@ -135,6 +135,7 @@ public:
 	bool GetCustomSavedExpansionState(const FString& NodePath) const override;
 	virtual void NotifyFinishedChangingProperties(const FPropertyChangedEvent& PropertyChangedEvent) override;
 	void RefreshTree() override;
+	virtual void RequestForceRefresh() override;
 	TSharedPtr<FAssetThumbnailPool> GetThumbnailPool() const override;
 	virtual const TArray<TSharedRef<IClassViewerFilter>>& GetClassViewerFilters() const override;
 	TSharedPtr<IPropertyUtilities> GetPropertyUtilities() override;
@@ -397,6 +398,8 @@ protected:
 
 	/** Set timer to force refresh if one not already set */
 	void SetPendingRefreshTimer();
+	/** Clear timer to force refresh if set */
+	void ClearPendingRefreshTimer();
 	/** Force refresh during editor tick */
 	void HandlePendingRefreshTimer();
 
@@ -522,12 +525,12 @@ protected:
 	/** Timer has already been set to be run next tick */
 	bool bPendingCleanupTimerSet : 1;
 
-	/** Refresh timer has already been set to be run next tick */
-	bool bPendingRefreshTimerSet : 1;
-
 	/** Are we currently running deferred actions? */
 	bool bRunningDeferredActions : 1;
 
+	/** Handle to the pending refresh request, if any */
+	FTimerHandle PendingRefreshTimerHandle;
+
 	mutable TSharedPtr<FEditConditionParser> EditConditionParser;
 	
 	/** Optional custom filter(s) to be applied when selecting values for class properties */
diff --git a/Engine/Source/Editor/PropertyEditor/Private/SPropertyTreeViewImpl.cpp b/Engine/Source/Editor/PropertyEditor/Private/SPropertyTreeViewImpl.cpp
--- a/Engine/Source/Editor/PropertyEditor/Private/SPropertyTreeViewImpl.cpp
+++ b/Engine/Source/Editor/PropertyEditor/Private/SPropertyTreeViewImpl.cpp
@@ -71,6 +71,12 @@ public:
 		View.RequestRefresh();
 	}
 
+	virtual void RequestForceRefresh() override
+	{
+		// RequestRefresh is already a deferred ForceRefresh
+		RequestRefresh();
+	}
+
 	virtual bool IsPropertyEditingEnabled() const override
 	{
 		return true;
diff --git a/Engine/Source/Editor/PropertyEditor/Private/SSingleProperty.cpp b/Engine/Source/Editor/PropertyEditor/Private/SSingleProperty.cpp
--- a/Engine/Source/Editor/PropertyEditor/Private/SSingleProperty.cpp
+++ b/Engine/Source/Editor/PropertyEditor/Private/SSingleProperty.cpp
@@ -72,6 +72,8 @@ public:
 
 	virtual void RequestRefresh() override {}
 
+	virtual void RequestForceRefresh() override {}
+
 	virtual TSharedPtr<class FAssetThumbnailPool> GetThumbnailPool() const override
 	{
 		return bShouldHideAssetThumbnail ? nullptr : UThumbnailManager::Get().GetSharedThumbnailPool();
diff --git a/Engine/Source/Editor/PropertyEditor/Private/SStructureDetailsView.cpp b/Engine/Source/Editor/PropertyEditor/Private/SStructureDetailsView.cpp
--- a/Engine/Source/Editor/PropertyEditor/Private/SStructureDetailsView.cpp
+++ b/Engine/Source/Editor/PropertyEditor/Private/SStructureDetailsView.cpp
@@ -271,6 +271,8 @@ void SStructureDetailsView::SetCustomName(const FText& Text)
 
 void SStructureDetailsView::ForceRefresh()
 {
+	ClearPendingRefreshTimer();
+
 	SetStructureProvider(StructProvider);
 }
 
diff --git a/Engine/Source/Editor/PropertyEditor/Private/UserInterface/PropertyDetails/PropertyDetailsUtilities.cpp b/Engine/Source/Editor/PropertyEditor/Private/UserInterface/PropertyDetails/PropertyDetailsUtilities.cpp
--- a/Engine/Source/Editor/PropertyEditor/Private/UserInterface/PropertyDetails/PropertyDetailsUtilities.cpp
+++ b/Engine/Source/Editor/PropertyEditor/Private/UserInterface/PropertyDetails/PropertyDetailsUtilities.cpp
@@ -50,6 +50,11 @@ void FPropertyDetailsUtilities::RequestRefresh()
 	DetailsView.RefreshTree();
 }
 
+void FPropertyDetailsUtilities::RequestForceRefresh()
+{
+	DetailsView.RequestForceRefresh();
+}
+
 TSharedPtr<class FAssetThumbnailPool> FPropertyDetailsUtilities::GetThumbnailPool() const
 {
 	return DetailsView.GetThumbnailPool();
diff --git a/Engine/Source/Editor/PropertyEditor/Private/UserInterface/PropertyDetails/PropertyDetailsUtilities.h b/Engine/Source/Editor/PropertyEditor/Private/UserInterface/PropertyDetails/PropertyDetailsUtilities.h
--- a/Engine/Source/Editor/PropertyEditor/Private/UserInterface/PropertyDetails/PropertyDetailsUtilities.h
+++ b/Engine/Source/Editor/PropertyEditor/Private/UserInterface/PropertyDetails/PropertyDetailsUtilities.h
@@ -21,6 +21,7 @@ public:
 	virtual bool IsPropertyEditingEnabled() const override;
 	virtual void ForceRefresh() override;
 	virtual void RequestRefresh() override;
+	virtual void RequestForceRefresh() override;
 	virtual TSharedPtr<class FAssetThumbnailPool> GetThumbnailPool() const override;
 	virtual void NotifyFinishedChangingProperties(const FPropertyChangedEvent& PropertyChangedEvent) override;
 	virtual bool DontUpdateValueWhileEditing() const override;
diff --git a/Engine/Source/Editor/PropertyEditor/Public/IPropertyUtilities.h b/Engine/Source/Editor/PropertyEditor/Public/IPropertyUtilities.h
--- a/Engine/Source/Editor/PropertyEditor/Public/IPropertyUtilities.h
+++ b/Engine/Source/Editor/PropertyEditor/Public/IPropertyUtilities.h
@@ -21,8 +21,25 @@ public:
 	virtual void CreateColorPickerWindow( const TSharedRef< class FPropertyEditor >& PropertyEditor, bool bUseAlpha ) const = 0;
 	virtual void EnqueueDeferredAction( FSimpleDelegate DeferredAction ) = 0;
 	virtual bool IsPropertyEditingEnabled() const = 0;
+
+	/**
+	 * Force a rebuild of the view, recreating and updating all widgets.
+	 * @note This may run immediately; consider RequestForceRefresh unless you need to ensure invalid references (eg, to deleted objects) are removed.
+	 */
 	virtual void ForceRefresh() = 0;
+
+	/**
+	 * Request a refresh of the view on the next tick, ideally without triggering a full rebuild of all widgets.
+	 * @note This may run immediately or perform a rebuild depending on the implementation.
+	 */
 	virtual void RequestRefresh() = 0;
+
+	/**
+	 * Request a rebuild of the view on the next tick, recreating and updating all widgets.
+	 * @note This may run immediately depending on the implementation.
+	 */
+	virtual void RequestForceRefresh() = 0;
+	
 	virtual TSharedPtr<class FAssetThumbnailPool> GetThumbnailPool() const = 0;
 	virtual void NotifyFinishedChangingProperties(const FPropertyChangedEvent& PropertyChangedEvent) = 0;
 	virtual bool HasClassDefaultObject() const = 0;
