SStandaloneCustomizedValueWidget: Remove an include to a private header in the public header.

diff --git a/Engine/Source/Editor/PropertyEditor/Private/SStandaloneCustomizedValueWidget.cpp b/Engine/Source/Editor/PropertyEditor/Private/SStandaloneCustomizedValueWidget.cpp
new file mode 100644
--- /dev/null
+++ b/Engine/Source/Editor/PropertyEditor/Private/SStandaloneCustomizedValueWidget.cpp
@@ -0,0 +1,29 @@
+// Copyright Epic Games, Inc. All Rights Reserved.
+
+#include "SStandaloneCustomizedValueWidget.h"
+
+#include "DetailCategoryBuilderImpl.h"
+#include "IPropertyTypeCustomization.h"
+	
+void SStandaloneCustomizedValueWidget::Construct( const FArguments& InArgs,
+	TSharedPtr<IPropertyTypeCustomization> InCustomizationInterface, TSharedRef<IPropertyHandle> InPropertyHandle)
+{
+	ParentCategory = InArgs._ParentCategory;
+	
+	CustomizationInterface = InCustomizationInterface;
+	PropertyHandle = InPropertyHandle;
+	CustomPropertyWidget = MakeShared<FDetailWidgetRow>();
+
+	CustomizationInterface->CustomizeHeader(InPropertyHandle, *CustomPropertyWidget, *this);
+
+	ChildSlot
+	[
+		CustomPropertyWidget->ValueWidget.Widget
+	];
+}
+
+TSharedPtr<FAssetThumbnailPool> SStandaloneCustomizedValueWidget::GetThumbnailPool() const
+{
+	TSharedPtr<FDetailCategoryImpl> ParentCategoryPinned = ParentCategory.Pin();
+	return ParentCategoryPinned.IsValid() ? ParentCategoryPinned->GetParentLayout().GetThumbnailPool() : nullptr;
+}
diff --git a/Engine/Source/Editor/PropertyEditor/Public/SStandaloneCustomizedValueWidget.h b/Engine/Source/Editor/PropertyEditor/Public/SStandaloneCustomizedValueWidget.h
--- a/Engine/Source/Editor/PropertyEditor/Public/SStandaloneCustomizedValueWidget.h
+++ b/Engine/Source/Editor/PropertyEditor/Public/SStandaloneCustomizedValueWidget.h
@@ -2,10 +2,13 @@
 
 #pragma once
 
-#include "DetailGroup.h"
 #include "DetailWidgetRow.h"
+#include "IPropertyTypeCustomization.h"
 #include "Widgets/SCompoundWidget.h"
 
+class FAssetThumbnailPool;
+class FDetailCategoryImpl;
+
 /**
  * Generates the header widget for a customized struct or other type.
  * This widget is generally used in the property editor to display a struct as a single row, like with an FColor.
@@ -20,28 +23,10 @@ public:
 		SLATE_ARGUMENT( TSharedPtr<FDetailCategoryImpl>, ParentCategory)
 	SLATE_END_ARGS()
 	
-	void Construct( const FArguments& InArgs,
-		TSharedPtr<IPropertyTypeCustomization> InCustomizationInterface, TSharedRef<IPropertyHandle> InPropertyHandle)
-	{
-		ParentCategory = InArgs._ParentCategory;
-		
-		CustomizationInterface = InCustomizationInterface;
-		PropertyHandle = InPropertyHandle;
-		CustomPropertyWidget = MakeShareable(new FDetailWidgetRow);
-
-		CustomizationInterface->CustomizeHeader(InPropertyHandle, *CustomPropertyWidget, *this);
-
-		ChildSlot
-		[
-			CustomPropertyWidget->ValueWidget.Widget
-		];
-	}
+	void Construct(const FArguments& InArgs,
+		TSharedPtr<IPropertyTypeCustomization> InCustomizationInterface, TSharedRef<IPropertyHandle> InPropertyHandle);
 
-	virtual TSharedPtr<FAssetThumbnailPool> GetThumbnailPool() const override
-	{
-		TSharedPtr<FDetailCategoryImpl> ParentCategoryPinned = ParentCategory.Pin();
-		return ParentCategoryPinned.IsValid() ? ParentCategoryPinned->GetParentLayout().GetThumbnailPool() : NULL;
-	}
+	virtual TSharedPtr<FAssetThumbnailPool> GetThumbnailPool() const override;
 
 private:
 	TWeakPtr<FDetailCategoryImpl> ParentCategory;
