Property Editor: Allow using a UFUNCTION to interactively validate string input.

diff --git a/Engine/Source/Editor/PropertyEditor/Private/UserInterface/PropertyEditor/SPropertyEditorText.cpp b/Engine/Source/Editor/PropertyEditor/Private/UserInterface/PropertyEditor/SPropertyEditorText.cpp
--- a/Engine/Source/Editor/PropertyEditor/Private/UserInterface/PropertyEditor/SPropertyEditorText.cpp
+++ b/Engine/Source/Editor/PropertyEditor/Private/UserInterface/PropertyEditor/SPropertyEditorText.cpp
@@ -36,17 +36,56 @@ void SPropertyEditorText::Construct( const FArguments& InArgs, const TSharedRef<
 	static const FName NAME_MultiLine = "MultiLine";
 	static const FName NAME_PasswordField = "PasswordField";
 	static const FName NAME_AllowedCharacters = "AllowedCharacters";
+	static const FName NAME_PropertyValidator = "PropertyValidator";
 
-	bIsMultiLine = InPropertyEditor->GetPropertyHandle()->GetBoolMetaData(NAME_MultiLine);
+	TSharedRef<IPropertyHandle> PropertyHandle(InPropertyEditor->GetPropertyHandle());
+	bIsMultiLine = PropertyHandle->GetBoolMetaData(NAME_MultiLine);
 
-	MaxLength = InPropertyEditor->GetPropertyHandle()->GetIntMetaData(NAME_MaxLength);
+	MaxLength = PropertyHandle->GetIntMetaData(NAME_MaxLength);
 	if (InPropertyEditor->PropertyIsA(FNameProperty::StaticClass()))
 	{
 		MaxLength = MaxLength <= 0 ? NAME_SIZE - 1 : FMath::Min(MaxLength, NAME_SIZE - 1);
 	}
 
-	const bool bIsPassword = InPropertyEditor->GetPropertyHandle()->GetBoolMetaData(NAME_PasswordField);
-	AllowedCharacters.InitializeFromString(InPropertyEditor->GetPropertyHandle()->GetMetaData(NAME_AllowedCharacters));
+	const bool bIsPassword = PropertyHandle->GetBoolMetaData(NAME_PasswordField);
+	AllowedCharacters.InitializeFromString(PropertyHandle->GetMetaData(NAME_AllowedCharacters));
+
+	if (PropertyHandle->HasMetaData(NAME_PropertyValidator))
+	{
+		const FString PropertyValidatorFunctionName = PropertyHandle->GetMetaData(NAME_PropertyValidator);
+		const UClass* OuterBaseClass = PropertyHandle->GetOuterBaseClass();
+		if (!PropertyValidatorFunctionName.IsEmpty() && OuterBaseClass)
+		{
+			static TSet<FString> LoggedWarnings;
+			
+			UObject* ValidatorObject = OuterBaseClass->GetDefaultObject<UObject>();
+			const UFunction* PropertyValidatorFunction = ValidatorObject->FindFunction(*PropertyValidatorFunctionName);
+			if (PropertyValidatorFunction)
+			{
+				if (PropertyValidatorFunction->FunctionFlags & FUNC_Static)
+				{
+					PropertyValidatorFunc = FPropertyValidatorFunc::CreateUFunction(ValidatorObject, PropertyValidatorFunction->GetFName());
+				}
+				else
+				{
+					const FString WarningId = OuterBaseClass->GetName() + TEXT(":") + *PropertyHandle->GetProperty()->GetName() + TEXT(":Static");
+				
+					UE_CLOG(!LoggedWarnings.Contains(WarningId), LogPropertyNode, Warning, TEXT("PropertyValidator ufunction '%s' on %s%s::%s must be a static function."),
+						*PropertyValidatorFunctionName, OuterBaseClass->GetPrefixCPP(), *OuterBaseClass->GetName(), *PropertyHandle->GetProperty()->GetName());
+					LoggedWarnings.Add(WarningId);
+				}
+			}
+			else
+			{
+				// Let the developer know that the function is missing.
+				const FString WarningId = OuterBaseClass->GetName() + TEXT(":") + *PropertyHandle->GetProperty()->GetName() + TEXT(":Missing");
+				
+				UE_CLOG(!LoggedWarnings.Contains(WarningId), LogPropertyNode, Warning, TEXT("PropertyValidator ufunction '%s' on %s%s::%s not found."),
+					*PropertyValidatorFunctionName, OuterBaseClass->GetPrefixCPP(), *OuterBaseClass->GetName(), *PropertyHandle->GetProperty()->GetName());
+				LoggedWarnings.Add(WarningId);
+			}
+		}
+	}
 
 	TSharedPtr<SHorizontalBox> HorizontalBox;
 	if(bIsMultiLine)
@@ -184,6 +223,16 @@ bool SPropertyEditorText::OnVerifyTextChanged(const FText& Text, FText& OutError
 		}
 	}
 
+	if (PropertyValidatorFunc.IsBound())
+	{
+		FText Result = PropertyValidatorFunc.Execute(TextString); 
+		if (!Result.IsEmpty())
+		{
+			OutError = Result;
+			return false;
+		}
+	}
+
 	return true;
 }
 
diff --git a/Engine/Source/Editor/PropertyEditor/Private/UserInterface/PropertyEditor/SPropertyEditorText.h b/Engine/Source/Editor/PropertyEditor/Private/UserInterface/PropertyEditor/SPropertyEditorText.h
--- a/Engine/Source/Editor/PropertyEditor/Private/UserInterface/PropertyEditor/SPropertyEditorText.h
+++ b/Engine/Source/Editor/PropertyEditor/Private/UserInterface/PropertyEditor/SPropertyEditorText.h
@@ -52,7 +52,8 @@ private:
 	bool IsReadOnly() const;
 
 private:
-
+	DECLARE_DELEGATE_RetVal_OneParam(FText, FPropertyValidatorFunc, const FString&);
+	
 	TSharedPtr< class FPropertyEditor > PropertyEditor;
 
 	TSharedPtr< class SWidget > PrimaryWidget;
@@ -71,4 +72,7 @@ private:
 
 	/** The characters that are allowed in the value. If empty, any character is allowed. */
 	FCharRangeList AllowedCharacters;
+
+	/** An optional validator function that will be called each time the text is changed */
+	FPropertyValidatorFunc PropertyValidatorFunc;
 };
