PropertyHandle: Add a CreatePropertyValueWidget that accounts for customization and thus supports struct properties.

diff --git a/Engine/Source/Editor/PropertyEditor/Private/PropertyHandleImpl.cpp b/Engine/Source/Editor/PropertyEditor/Private/PropertyHandleImpl.cpp
--- a/Engine/Source/Editor/PropertyEditor/Private/PropertyHandleImpl.cpp
+++ b/Engine/Source/Editor/PropertyEditor/Private/PropertyHandleImpl.cpp
@@ -27,6 +27,7 @@
 #include "UObject/EnumProperty.h"
 #include "UObject/FieldPathProperty.h"
 #include "IDetailPropertyRow.h"
+#include "IDetailsViewPrivate.h"
 #include "ObjectEditorUtils.h"
 #include "SResetToDefaultPropertyEditor.h"
 #include "PropertyPathHelpers.h"
@@ -34,6 +35,9 @@
 #include "HAL/PlatformApplicationMisc.h"
 #include "UObject/OverridableManager.h"
 #include "UObject/PropertyOptional.h"
+#include "SStandaloneCustomizedValueWidget.h"
+#include "Modules/ModuleManager.h"
+#include "PropertyEditorModule.h"
 
 #define LOCTEXT_NAMESPACE "PropertyHandleImplementation"
 
@@ -2415,6 +2419,37 @@ TSharedRef<SWidget> FPropertyHandleBase::CreatePropertyValueWidget( bool bDispla
 	return SNullWidget::NullWidget;
 }
 
+TSharedRef<SWidget> FPropertyHandleBase::CreatePropertyValueWidget( const IDetailsView* DetailsView )
+{
+	FCustomPropertyTypeLayoutMap CustomPropertyTypeLayoutMap = [DetailsView]()
+	{
+		if ( DetailsView )
+		{
+			return static_cast<const IDetailsViewPrivate*>( DetailsView )->GetCustomPropertyTypeLayoutMap();
+		}
+		else
+		{
+			return FCustomPropertyTypeLayoutMap();
+		}
+	}();
+
+	FPropertyEditorModule& PropertyEditorModule = FModuleManager::GetModuleChecked<FPropertyEditorModule>("PropertyEditor");
+
+	FPropertyTypeLayoutCallback LayoutCallback = PropertyEditorModule.GetPropertyTypeCustomization( GetProperty(), *this, CustomPropertyTypeLayoutMap );
+	if (LayoutCallback.IsValid())
+	{
+		TSharedRef<IPropertyTypeCustomization> PropertyTypeCustomization = LayoutCallback.GetCustomizationInstance();
+
+		return SNew( SStandaloneCustomizedValueWidget, PropertyTypeCustomization, AsShared() );
+	}
+	else
+	{
+		// Opting to not generate the default buttons since the code path that goes through a customization would likely not have them
+		const bool bDisplayDefaultPropertyButtons = false;
+		return CreatePropertyValueWidget( bDisplayDefaultPropertyButtons );
+	}
+}
+
 class SDefaultPropertyButtonWidgets : public SCompoundWidget
 {
 	SLATE_BEGIN_ARGS(SDefaultPropertyButtonWidgets)	{}
diff --git a/Engine/Source/Editor/PropertyEditor/Private/PropertyHandleImpl.h b/Engine/Source/Editor/PropertyEditor/Private/PropertyHandleImpl.h
--- a/Engine/Source/Editor/PropertyEditor/Private/PropertyHandleImpl.h
+++ b/Engine/Source/Editor/PropertyEditor/Private/PropertyHandleImpl.h
@@ -430,6 +430,7 @@ public:
 	virtual TSharedRef<SWidget> CreatePropertyNameWidget( const FText& NameOverride, const FText& ToolTipOverride, bool bDisplayResetToDefault, bool bDisplayText, bool bDisplayThumbnail ) const override;
 	virtual TSharedRef<SWidget> CreatePropertyNameWidget(const FText& NameOverride, const FText& ToolTipOverride) const override;
 	virtual TSharedRef<SWidget> CreatePropertyValueWidget( bool bDisplayDefaultPropertyButtons = true ) const override;
+	virtual TSharedRef<SWidget> CreatePropertyValueWidget( const IDetailsView* DetailsView ) override;
 	virtual TSharedRef<SWidget> CreateDefaultPropertyButtonWidgets() const override;
 	virtual void CreateDefaultPropertyCopyPasteActions(FUIAction& OutCopyAction, FUIAction& OutPasteAction) const override;
 	virtual bool IsEditConst() const override;
diff --git a/Engine/Source/Editor/PropertyEditor/Public/PropertyHandle.h b/Engine/Source/Editor/PropertyEditor/Public/PropertyHandle.h
--- a/Engine/Source/Editor/PropertyEditor/Public/PropertyHandle.h
+++ b/Engine/Source/Editor/PropertyEditor/Public/PropertyHandle.h
@@ -629,13 +629,26 @@ public:
 	 * @param ToolTipOverride			The tooltip override to use instead of the property name
 	 */
 	virtual TSharedRef<SWidget> CreatePropertyNameWidget(const FText& NameOverride = FText::GetEmpty(), const FText& ToolTipOverride = FText::GetEmpty()) const = 0;
+	
 	/**
 	 * Creates a value widget for this property
-
+	 * 
+	 * @param bDisplayDefaultPropertyButtons	If the value widget should include the property buttons.
+	 *
 	 * @return the value widget for this property
 	 */
 	virtual TSharedRef<SWidget> CreatePropertyValueWidget( bool bDisplayDefaultPropertyButtons = true ) const = 0;
 
+	/**
+	 * Creates a value widget for this property using customization, if available.
+	 * Note that this is only the value widget for the header/main row for properties with child rows (ie: structs)
+	 * 
+	 * @param DetailsView						The details view to create the value widget for. Used to retrieve per details view customizations.
+	 * 
+	 * @return the value widget for this property
+	 */
+	virtual TSharedRef<SWidget> CreatePropertyValueWidget( const IDetailsView* DetailsView ) = 0;
+
 	/**
 	 * Creates the default buttons which appear next to value widgets.  This is useful when creating customizations
 	 * which don't use CreatePropertyValueWidget but you still want array item behaviors and reset to default capabilities.
