PropertyHandle: Renamed the new version of CreatePropertyValueWidget to CreatePropertyValueWidgetWithCustomization to resolve ambiguity between bool and pointer parameters.

diff --git a/Engine/Source/Editor/PropertyEditor/Private/PropertyHandleImpl.cpp b/Engine/Source/Editor/PropertyEditor/Private/PropertyHandleImpl.cpp
--- a/Engine/Source/Editor/PropertyEditor/Private/PropertyHandleImpl.cpp
+++ b/Engine/Source/Editor/PropertyEditor/Private/PropertyHandleImpl.cpp
@@ -2419,7 +2419,7 @@ TSharedRef<SWidget> FPropertyHandleBase::CreatePropertyValueWidget( bool bDispla
 	return SNullWidget::NullWidget;
 }
 
-TSharedRef<SWidget> FPropertyHandleBase::CreatePropertyValueWidget( const IDetailsView* DetailsView )
+TSharedRef<SWidget> FPropertyHandleBase::CreatePropertyValueWidgetWithCustomization( const IDetailsView* DetailsView )
 {
 	FCustomPropertyTypeLayoutMap CustomPropertyTypeLayoutMap = [DetailsView]()
 	{
diff --git a/Engine/Source/Editor/PropertyEditor/Private/PropertyHandleImpl.h b/Engine/Source/Editor/PropertyEditor/Private/PropertyHandleImpl.h
--- a/Engine/Source/Editor/PropertyEditor/Private/PropertyHandleImpl.h
+++ b/Engine/Source/Editor/PropertyEditor/Private/PropertyHandleImpl.h
@@ -430,7 +430,7 @@ public:
 	virtual TSharedRef<SWidget> CreatePropertyNameWidget( const FText& NameOverride, const FText& ToolTipOverride, bool bDisplayResetToDefault, bool bDisplayText, bool bDisplayThumbnail ) const override;
 	virtual TSharedRef<SWidget> CreatePropertyNameWidget(const FText& NameOverride, const FText& ToolTipOverride) const override;
 	virtual TSharedRef<SWidget> CreatePropertyValueWidget( bool bDisplayDefaultPropertyButtons = true ) const override;
-	virtual TSharedRef<SWidget> CreatePropertyValueWidget( const IDetailsView* DetailsView ) override;
+	virtual TSharedRef<SWidget> CreatePropertyValueWidgetWithCustomization( const IDetailsView* DetailsView ) override;
 	virtual TSharedRef<SWidget> CreateDefaultPropertyButtonWidgets() const override;
 	virtual void CreateDefaultPropertyCopyPasteActions(FUIAction& OutCopyAction, FUIAction& OutPasteAction) const override;
 	virtual bool IsEditConst() const override;
diff --git a/Engine/Source/Editor/PropertyEditor/Public/PropertyHandle.h b/Engine/Source/Editor/PropertyEditor/Public/PropertyHandle.h
--- a/Engine/Source/Editor/PropertyEditor/Public/PropertyHandle.h
+++ b/Engine/Source/Editor/PropertyEditor/Public/PropertyHandle.h
@@ -647,7 +647,7 @@ public:
 	 * 
 	 * @return the value widget for this property
 	 */
-	virtual TSharedRef<SWidget> CreatePropertyValueWidget( const IDetailsView* DetailsView ) = 0;
+	virtual TSharedRef<SWidget> CreatePropertyValueWidgetWithCustomization( const IDetailsView* DetailsView ) = 0;
 
 	/**
 	 * Creates the default buttons which appear next to value widgets.  This is useful when creating customizations
