FPropertyNode Editor Archetype/EditConst Policy
- Allow registration/unregistration of IArchetypePolicy/IEditConstPolicy which allow outside systems to provide policies
- Update some details customizations to make usage of new PropertyEditorPolicy::GetArchetype/IsPropertyEditConst which respects registered policies

#rnx

diff --git a/Engine/Source/Editor/DetailCustomizations/Private/ComponentTransformDetails.cpp b/Engine/Source/Editor/DetailCustomizations/Private/ComponentTransformDetails.cpp
--- a/Engine/Source/Editor/DetailCustomizations/Private/ComponentTransformDetails.cpp
+++ b/Engine/Source/Editor/DetailCustomizations/Private/ComponentTransformDetails.cpp
@@ -25,6 +25,7 @@
 #include "Input/Events.h"
 #include "Internationalization/Internationalization.h"
 #include "IPropertyUtilities.h"
+#include "PropertyEditorArchetypePolicy.h"
 #include "Kismet2/ComponentEditorUtils.h"
 #include "Layout/Margin.h"
 #include "Math/Quat.h"
@@ -934,7 +935,7 @@ struct FGetRootComponentArchetype
 	static USceneComponent* Get(UObject* Object)
 	{
 		auto RootComponent = Object ? GetSceneComponentFromDetailsObject(Object) : nullptr;
-		return RootComponent ? Cast<USceneComponent>(RootComponent->GetArchetype()) : nullptr;
+		return RootComponent ? Cast<USceneComponent>(PropertyEditorPolicy::GetArchetype(RootComponent)) : nullptr;
 	}
 };
 
diff --git a/Engine/Source/Editor/DetailCustomizations/Private/LightComponentDetails.cpp b/Engine/Source/Editor/DetailCustomizations/Private/LightComponentDetails.cpp
--- a/Engine/Source/Editor/DetailCustomizations/Private/LightComponentDetails.cpp
+++ b/Engine/Source/Editor/DetailCustomizations/Private/LightComponentDetails.cpp
@@ -24,6 +24,7 @@
 #include "UObject/UnrealType.h"
 #include "UObject/WeakObjectPtr.h"
 #include "DetailWidgetRow.h"
+#include "PropertyEditorArchetypePolicy.h"
 
 #define LOCTEXT_NAMESPACE "LightComponentDetails"
 
@@ -242,7 +243,7 @@ void FLightComponentDetails::SetComponentIntensity(ULightComponent* Component, f
 
 void FLightComponentDetails::ResetIntensityToDefault(TSharedPtr<IPropertyHandle> PropertyHandle, TWeakObjectPtr<ULightComponent> Component)
 {
-	ULightComponent* ArchetypeComponent = Component.IsValid() ? Cast<ULocalLightComponent>(Component->GetArchetype()) : nullptr;
+	ULightComponent* ArchetypeComponent = Component.IsValid() ? Cast<ULocalLightComponent>(PropertyEditorPolicy::GetArchetype(Component.Get())) : nullptr;
 	if (ArchetypeComponent)
 	{
 		SetComponentIntensity(Component.Get(), ArchetypeComponent->ComputeLightBrightness());
@@ -256,7 +257,7 @@ void FLightComponentDetails::ResetIntensityToDefault(TSharedPtr<IPropertyHandle>
 
 bool FLightComponentDetails::IsIntensityResetToDefaultVisible(TSharedPtr<IPropertyHandle> PropertyHandle, TWeakObjectPtr<ULightComponent> Component) const
 {
-	ULightComponent* ArchetypeComponent = Component.IsValid() ? Cast<ULocalLightComponent>(Component->GetArchetype()) : nullptr;
+	ULightComponent* ArchetypeComponent = Component.IsValid() ? Cast<ULocalLightComponent>(PropertyEditorPolicy::GetArchetype(Component.Get())) : nullptr;
 	if (ArchetypeComponent)
 	{
 		return !FMath::IsNearlyEqual(Component->ComputeLightBrightness(), ArchetypeComponent->ComputeLightBrightness());
@@ -331,7 +332,7 @@ namespace
 void FLightComponentDetails::ResetIntensityUnitsToDefault(TSharedPtr<IPropertyHandle> PropertyHandle, ULocalLightComponent* Component)
 {
 	// Actors (and blueprints) spawned from the actor factory inherit the intensity units from the project settings.
-	if (Component && Component->GetArchetype() && !Component->GetArchetype()->IsInBlueprint())
+	if (Component && PropertyEditorPolicy::GetArchetype(Component) && !PropertyEditorPolicy::GetArchetype(Component)->IsInBlueprint())
 	{
 		static const auto CVarDefaultLightUnits = IConsoleManager::Get().FindTConsoleVariableDataInt(TEXT("r.DefaultFeature.LightUnits"));
 		const ELightUnits DefaultUnits = (ELightUnits)CVarDefaultLightUnits->GetValueOnGameThread();
@@ -351,7 +352,7 @@ void FLightComponentDetails::ResetIntensityUnitsToDefault(TSharedPtr<IPropertyHa
 bool FLightComponentDetails::IsIntensityUnitsResetToDefaultVisible(TSharedPtr<IPropertyHandle> PropertyHandle, ULocalLightComponent* Component) const
 {
 	// Actors (and blueprints) spawned from the actor factory inherit the project settings.
-	if (Component && Component->GetArchetype() && !Component->GetArchetype()->IsInBlueprint())
+	if (Component && PropertyEditorPolicy::GetArchetype(Component) && !PropertyEditorPolicy::GetArchetype(Component)->IsInBlueprint())
 	{
 		static const auto CVarDefaultLightUnits = IConsoleManager::Get().FindTConsoleVariableDataInt(TEXT("r.DefaultFeature.LightUnits"));
 		const ELightUnits DefaultUnits = (ELightUnits)CVarDefaultLightUnits->GetValueOnGameThread();
diff --git a/Engine/Source/Editor/PropertyEditor/Private/MaterialList.cpp b/Engine/Source/Editor/PropertyEditor/Private/MaterialList.cpp
--- a/Engine/Source/Editor/PropertyEditor/Private/MaterialList.cpp
+++ b/Engine/Source/Editor/PropertyEditor/Private/MaterialList.cpp
@@ -8,6 +8,7 @@
 #include "PropertyCustomizationHelpers.h"
 #include "PropertyEditorModule.h"
 #include "PropertyHandle.h"
+#include "PropertyEditorEditConstPolicy.h"
 #include "Framework/MultiBox/MultiBoxBuilder.h"
 #include "Modules/ModuleManager.h"
 #include "Styling/SlateIconFinder.h"
@@ -797,7 +798,7 @@ void FMaterialList::AddMaterialItem( FDetailWidgetRow& Row, int32 CurrentSlot, c
 		FProperty* MaterialProperty = nullptr;
 		if (SceneComponent->GetMaterialPropertyPath(Item.SlotIndex, OwnerObject, PropertyPath, MaterialProperty))
 		{
-			Row.IsEnabled(SceneComponent->CanEditChange(MaterialProperty));
+			Row.IsEnabled(!PropertyEditorPolicy::IsPropertyEditConst(MaterialProperty, SceneComponent) && SceneComponent->CanEditChange(MaterialProperty));
 		}		
 	}
 }
diff --git a/Engine/Source/Editor/PropertyEditor/Private/PropertyEditorArchetypePolicy.cpp b/Engine/Source/Editor/PropertyEditor/Private/PropertyEditorArchetypePolicy.cpp
new file mode 100644
--- /dev/null
+++ b/Engine/Source/Editor/PropertyEditor/Private/PropertyEditorArchetypePolicy.cpp
@@ -0,0 +1,22 @@
+// Copyright Epic Games, Inc. All Rights Reserved.
+
+#include "PropertyEditorArchetypePolicy.h"
+#include "PropertyNode.h"
+
+namespace PropertyEditorPolicy
+{
+	UObject* GetArchetype(const UObject* Object)
+	{
+		return FPropertyNode::GetArchetype(Object);
+	}
+
+	void RegisterArchetypePolicy(IArchetypePolicy* Policy)
+	{
+		FPropertyNode::RegisterArchetypePolicy(Policy);
+	}
+
+	void UnregisterArchetypePolicy(IArchetypePolicy* Policy)
+	{
+		FPropertyNode::UnregisterArchetypePolicy(Policy);
+	}
+}
\ No newline at end of file
diff --git a/Engine/Source/Editor/PropertyEditor/Private/PropertyEditorEditConstPolicy.cpp b/Engine/Source/Editor/PropertyEditor/Private/PropertyEditorEditConstPolicy.cpp
new file mode 100644
--- /dev/null
+++ b/Engine/Source/Editor/PropertyEditor/Private/PropertyEditorEditConstPolicy.cpp
@@ -0,0 +1,27 @@
+// Copyright Epic Games, Inc. All Rights Reserved.
+
+#include "PropertyEditorEditConstPolicy.h"
+#include "PropertyNode.h"
+
+namespace PropertyEditorPolicy
+{
+	void RegisterEditConstPolicy(IEditConstPolicy* Policy)
+	{
+		FPropertyNode::RegisterEditConstPolicy(Policy);
+	}
+
+	void UnregisterEditConstPolicy(IEditConstPolicy* Policy)
+	{
+		FPropertyNode::UnregisterEditConstPolicy(Policy);
+	}
+
+	bool IsPropertyEditConst(const FEditPropertyChain& PropertyChain, UObject* Object)
+	{
+		return FPropertyNode::IsPropertyEditConst(PropertyChain, Object);
+	}
+
+	bool IsPropertyEditConst(const FProperty* Property, UObject* Object)
+	{
+		return FPropertyNode::IsPropertyEditConst(Property, Object);
+	}
+}
\ No newline at end of file
diff --git a/Engine/Source/Editor/PropertyEditor/Private/PropertyNode.cpp b/Engine/Source/Editor/PropertyEditor/Private/PropertyNode.cpp
--- a/Engine/Source/Editor/PropertyEditor/Private/PropertyNode.cpp
+++ b/Engine/Source/Editor/PropertyEditor/Private/PropertyNode.cpp
@@ -16,6 +16,8 @@
 #include "PropertyHandleImpl.h"
 #include "PropertyRestriction.h"
 #include "PropertyTextUtilities.h"
+#include "PropertyEditorEditConstPolicy.h"
+#include "PropertyEditorArchetypePolicy.h"
 #include "StringPrefixTree.h"
 #include "StructurePropertyNode.h"
 
@@ -26,9 +28,135 @@
 #include "UObject/UnrealType.h"
 
 #include "UObject/PropertyOptional.h"
+#include "UObject/UObjectArchetypeHelper.h"
 
 #define LOCTEXT_NAMESPACE "PropertyNode"
 
+namespace PropertyEditorPolicy
+{
+	struct FPropertyNodePolicyImpl : public FObjectArchetypeHelper::IObjectArchetypePolicy
+	{
+		FPropertyNodePolicyImpl() {}
+		virtual ~FPropertyNodePolicyImpl() {}
+		
+		virtual UObject* GetArchetype(const UObject* Object) const override
+		{
+			for (const IArchetypePolicy* ArchetypePolicy : ArchetypePolicies)
+			{
+				if (UObject* Archetype = ArchetypePolicy->GetArchetypeForObject(Object))
+				{
+					return Archetype;
+				}
+			}
+
+			return nullptr;
+		}
+
+		bool CanEditProperty(const FEditPropertyChain& PropertyChain, const UObject* Object) const
+		{
+			for (const IEditConstPolicy* EditConstPolicy : EditConstPolicies)
+			{
+				if (!EditConstPolicy->CanEditProperty(PropertyChain, Object))
+				{
+					return false;
+				}
+			}
+			
+			return true;
+		}
+
+		bool CanEditProperty(const FProperty* Property, const UObject* Object) const
+		{
+			for (const IEditConstPolicy* EditConstPolicy : EditConstPolicies)
+			{
+				if (!EditConstPolicy->CanEditProperty(Property, Object))
+				{
+					return false;
+				}
+			}
+
+			return true;
+		}
+
+		void RegisterArchetypePolicy(IArchetypePolicy* ArchetypePolicy)
+		{
+			check(!ArchetypePolicies.Contains(ArchetypePolicy));
+			ArchetypePolicies.Add(ArchetypePolicy);
+			PolicyEpoch++;
+		}
+
+		void UnregisterArchetypePolicy(IArchetypePolicy* ArchetypePolicy)
+		{
+			check(ArchetypePolicies.Contains(ArchetypePolicy));
+			ArchetypePolicies.Remove(ArchetypePolicy);
+			PolicyEpoch++;
+		}
+
+		void RegisterEditConstPolicy(IEditConstPolicy* EditConstPolicy)
+		{
+			check(!EditConstPolicies.Contains(EditConstPolicy));
+			EditConstPolicies.Add(EditConstPolicy);
+			PolicyEpoch++;
+		}
+
+		void UnregisterEditConstPolicy(IEditConstPolicy* EditConstPolicy)
+		{
+			check(EditConstPolicies.Contains(EditConstPolicy));
+			EditConstPolicies.Remove(EditConstPolicy);
+			PolicyEpoch++;
+		}
+
+		int32 GetPolicyEpoch() const
+		{
+			return PolicyEpoch;
+		}
+				
+		TArray<IArchetypePolicy*> ArchetypePolicies;
+		TArray<IEditConstPolicy*> EditConstPolicies;
+		int32 PolicyEpoch = 0;
+	};
+
+	FPropertyNodePolicyImpl& Get()
+	{
+		static FPropertyNodePolicyImpl Policy;
+		return Policy;
+	}
+}
+void FPropertyNode::RegisterArchetypePolicy(PropertyEditorPolicy::IArchetypePolicy* ArchetypePolicy)
+{
+	PropertyEditorPolicy::Get().RegisterArchetypePolicy(ArchetypePolicy);
+}
+
+void FPropertyNode::UnregisterArchetypePolicy(PropertyEditorPolicy::IArchetypePolicy* ArchetypePolicy)
+{
+	PropertyEditorPolicy::Get().UnregisterArchetypePolicy(ArchetypePolicy);
+}
+
+void FPropertyNode::RegisterEditConstPolicy(PropertyEditorPolicy::IEditConstPolicy* EditConstPolicy)
+{
+	PropertyEditorPolicy::Get().RegisterEditConstPolicy(EditConstPolicy);
+}
+
+void FPropertyNode::UnregisterEditConstPolicy(PropertyEditorPolicy::IEditConstPolicy* EditConstPolicy)
+{
+	PropertyEditorPolicy::Get().UnregisterEditConstPolicy(EditConstPolicy);
+}
+
+UObject* FPropertyNode::GetArchetype(const UObject* Object)
+{
+	return Object ? FObjectArchetypeHelper::GetArchetype(Object, &PropertyEditorPolicy::Get()) : nullptr;
+}
+
+bool FPropertyNode::IsPropertyEditConst(const FEditPropertyChain& PropertyChain, UObject* Object)
+{
+	return !PropertyEditorPolicy::Get().CanEditProperty(PropertyChain, Object);
+}
+
+bool FPropertyNode::IsPropertyEditConst(const FProperty* Property, UObject* Object)
+{
+	return !PropertyEditorPolicy::Get().CanEditProperty(Property, Object);
+}
+
 FEditConditionParser FPropertyNode::EditConditionParser;
 
 FPropertySettings& FPropertySettings::Get()
@@ -74,8 +202,10 @@ FPropertyNode::FPropertyNode()
 	, PropertyPath(TEXT(""))
 	, bIsEditConst(false)
 	, bUpdateEditConstState(true)
+	, UpdateEditConstStateEpoch(0)
 	, bDiffersFromDefault(false)
 	, bUpdateDiffersFromDefault(true)
+	, UpdateDiffersFromDefaultEpoch(0)
 {
 }
 
@@ -1197,8 +1327,10 @@ bool FPropertyNode::IsPropertyConst() const
 /** @return whether this window's property is constant (can't be edited by the user) */
 bool FPropertyNode::IsEditConst() const
 {
-	if (bUpdateEditConstState)
+	if (bUpdateEditConstState || UpdateEditConstStateEpoch != PropertyEditorPolicy::Get().GetPolicyEpoch())
 	{
+		UpdateEditConstStateEpoch = PropertyEditorPolicy::Get().GetPolicyEpoch();
+
 		// Ask the objects whether this property can be changed
 		const FObjectPropertyNode* ObjectPropertyNode = FindObjectItemParent();
 
@@ -1267,6 +1399,12 @@ bool FPropertyNode::IsEditConst() const
 					const TWeakObjectPtr<UObject> CurObject = *CurObjectIt;
 					if (CurObject.IsValid())
 					{
+						if (!PropertyEditorPolicy::Get().CanEditProperty(*PropertyChain, CurObject.Get()))
+						{
+							bIsEditConst = true;
+							break;
+						}
+						
 						if (!CurObject->CanEditChange(*PropertyChain))
 						{
 							// At least one of the objects didn't like the idea of this property being changed.
@@ -1711,7 +1849,7 @@ public:
 			// calculate the addresses for the default object if it exists
 			if (bHasDefaultValue)
 			{
-				PropertyDefaultValueRoot.OwnerObject = PropertyValueRoot.OwnerObject ? PropertyValueRoot.OwnerObject->GetArchetype() : nullptr;
+				PropertyDefaultValueRoot.OwnerObject = PropertyValueRoot.OwnerObject ? FPropertyNode::GetArchetype(PropertyValueRoot.OwnerObject) : nullptr;
 
 				PropertyDefaultBaseAddress = Node->GetValueBaseAddressFromObject(PropertyDefaultValueRoot.OwnerObject);
 				PropertyDefaultAddress = PropertyNode->GetValueAddressFromObject(PropertyDefaultValueRoot.OwnerObject);
@@ -1876,7 +2014,7 @@ private:
 			}
 			check(PropertyValueBaseAddress);
 			check(PropertyValueRoot.OwnerObject);
-			UObject* ParentDefault = PropertyValueRoot.OwnerObject->GetArchetype();
+			UObject* ParentDefault = FPropertyNode::GetArchetype(PropertyValueRoot.OwnerObject);
 			check(ParentDefault);
 			if (OwnerClass == ParentDefault->GetClass())
 			{
@@ -2348,8 +2486,9 @@ bool FPropertyNode::GetDiffersFromDefaultForObject( FPropertyItemValueDataTracke
  */
 bool FPropertyNode::GetDiffersFromDefault()
 {
-	if( bUpdateDiffersFromDefault )
+	if( bUpdateDiffersFromDefault || UpdateDiffersFromDefaultEpoch != PropertyEditorPolicy::Get().GetPolicyEpoch())
 	{
+		UpdateDiffersFromDefaultEpoch = PropertyEditorPolicy::Get().GetPolicyEpoch();
 		bUpdateDiffersFromDefault = false;
 		bDiffersFromDefault = false;
 
@@ -3617,7 +3756,7 @@ void FPropertyNode::GatherInstancesAffectedByContainerPropertyChange(UObject* Mo
 		{
 			UObject* Obj = ArchetypeInstances[i];
 
-			if (Obj->GetArchetype() == ObjToChange)
+			if (GetArchetype(Obj) == ObjToChange)
 			{
 				ObjectsToChange.Push(Obj);
 				ArchetypeInstances.RemoveAt(i--);
@@ -3916,7 +4055,7 @@ void FPropertyNode::PropagatePropertyChange( UObject* ModifiedObject, const TCHA
 		{
 			UObject* Obj = ArchetypeInstances[InstanceIndex];
 
-			if (Obj->GetArchetype() == ObjToChange)
+			if (GetArchetype(Obj) == ObjToChange)
 			{
 				ObjectsToChange.Push(Obj);
 				ArchetypeInstances.RemoveAt(InstanceIndex--);
diff --git a/Engine/Source/Editor/PropertyEditor/Private/PropertyNode.h b/Engine/Source/Editor/PropertyEditor/Private/PropertyNode.h
--- a/Engine/Source/Editor/PropertyEditor/Private/PropertyNode.h
+++ b/Engine/Source/Editor/PropertyEditor/Private/PropertyNode.h
@@ -22,6 +22,12 @@ class FPropertyNode;
 class FPropertyRestriction;
 class FStructurePropertyNode;
 
+namespace PropertyEditorPolicy
+{
+	class IEditConstPolicy;
+	class IArchetypePolicy;
+}
+
 DECLARE_LOG_CATEGORY_EXTERN(LogPropertyNode, Log, All);
 
 namespace EPropertyNodeFlags
@@ -921,6 +927,18 @@ public:
 		return PossibleExtensions;
 	}
 
+	static UObject* GetArchetype(const UObject* Object);
+
+	static void RegisterArchetypePolicy(PropertyEditorPolicy::IArchetypePolicy* ArchetypePolicy);
+	static void UnregisterArchetypePolicy(PropertyEditorPolicy::IArchetypePolicy* ArchetypePolicy);
+	
+	static void RegisterEditConstPolicy(PropertyEditorPolicy::IEditConstPolicy* EditConstPolicy);
+	static void UnregisterEditConstPolicy(PropertyEditorPolicy::IEditConstPolicy* EditConstPolicy);
+
+	static bool IsPropertyEditConst(const FEditPropertyChain& PropertyChain, UObject* Object);
+
+	static bool IsPropertyEditConst(const FProperty* Property, UObject* Object);
+
 	/**
 	 * Adds a restriction to the possible values for this property.
 	 * @param Restriction	The restriction being added to this property.
@@ -1305,8 +1323,10 @@ protected:
 	*/
 	mutable bool bIsEditConst;
 	mutable bool bUpdateEditConstState;
+	mutable int32 UpdateEditConstStateEpoch;
 	mutable bool bDiffersFromDefault;
 	mutable bool bUpdateDiffersFromDefault;
+	mutable int32 UpdateDiffersFromDefaultEpoch;
 };
 
 class FComplexPropertyNode : public FPropertyNode
diff --git a/Engine/Source/Editor/PropertyEditor/Public/PropertyEditorArchetypePolicy.h b/Engine/Source/Editor/PropertyEditor/Public/PropertyEditorArchetypePolicy.h
new file mode 100644
--- /dev/null
+++ b/Engine/Source/Editor/PropertyEditor/Public/PropertyEditorArchetypePolicy.h
@@ -0,0 +1,18 @@
+// Copyright Epic Games, Inc. All Rights Reserved.
+
+#pragma once
+
+#include "CoreFwd.h"
+
+namespace PropertyEditorPolicy
+{
+	class IArchetypePolicy
+	{
+	public:
+		virtual UObject* GetArchetypeForObject(const UObject* Object) const = 0;
+	};
+
+	PROPERTYEDITOR_API UObject* GetArchetype(const UObject* Object);
+	PROPERTYEDITOR_API void RegisterArchetypePolicy(IArchetypePolicy* Policy);
+	PROPERTYEDITOR_API void UnregisterArchetypePolicy(IArchetypePolicy* Policy);
+}
\ No newline at end of file
diff --git a/Engine/Source/Editor/PropertyEditor/Public/PropertyEditorEditConstPolicy.h b/Engine/Source/Editor/PropertyEditor/Public/PropertyEditorEditConstPolicy.h
new file mode 100644
--- /dev/null
+++ b/Engine/Source/Editor/PropertyEditor/Public/PropertyEditorEditConstPolicy.h
@@ -0,0 +1,23 @@
+// Copyright Epic Games, Inc. All Rights Reserved.
+
+#pragma once
+
+#include "CoreFwd.h"
+
+class FProperty;
+class FEditPropertyChain;
+
+namespace PropertyEditorPolicy
+{
+	class IEditConstPolicy
+	{
+	public:
+		virtual bool CanEditProperty(const FEditPropertyChain& PropertyChain, const UObject* Object) const = 0;
+		virtual bool CanEditProperty(const FProperty* Property, const UObject* Object) const = 0;
+	};
+
+	PROPERTYEDITOR_API bool IsPropertyEditConst(const FEditPropertyChain& PropertyChain, UObject* Object);
+	PROPERTYEDITOR_API bool IsPropertyEditConst(const FProperty* Property, UObject* Object);
+	PROPERTYEDITOR_API void RegisterEditConstPolicy(IEditConstPolicy* Policy);
+	PROPERTYEDITOR_API void UnregisterEditConstPolicy(IEditConstPolicy* Policy);
+}
\ No newline at end of file
diff --git a/Engine/Source/Runtime/CoreUObject/Internal/UObject/UObjectArchetypeHelper.h b/Engine/Source/Runtime/CoreUObject/Internal/UObject/UObjectArchetypeHelper.h
new file mode 100644
--- /dev/null
+++ b/Engine/Source/Runtime/CoreUObject/Internal/UObject/UObjectArchetypeHelper.h
@@ -0,0 +1,24 @@
+// Copyright Epic Games, Inc. All Rights Reserved.
+
+#pragma once
+
+#include "CoreFwd.h"
+
+class FObjectArchetypeHelper
+{
+public:
+	class IObjectArchetypePolicy
+	{
+#if WITH_EDITOR
+	public:
+		virtual UObject* GetArchetype(const UObject* InObject) const = 0;
+#endif
+	};
+
+#if WITH_EDITOR
+private:
+	// Restrict usage to FPropertyNode
+	friend class FPropertyNode;
+#endif
+	COREUOBJECT_API static UObject* GetArchetype(const UObject* InObject, const IObjectArchetypePolicy* InPolicy);
+};
diff --git a/Engine/Source/Runtime/CoreUObject/Private/UObject/UObjectArchetype.cpp b/Engine/Source/Runtime/CoreUObject/Private/UObject/UObjectArchetype.cpp
--- a/Engine/Source/Runtime/CoreUObject/Private/UObject/UObjectArchetype.cpp
+++ b/Engine/Source/Runtime/CoreUObject/Private/UObject/UObjectArchetype.cpp
@@ -13,6 +13,7 @@
 #include "UObject/UObjectAnnotation.h"
 #include "Stats/StatsMisc.h"
 #include "HAL/IConsoleManager.h"
+#include "UObject/UObjectArchetypeHelper.h"
 
 #define UE_CACHE_ARCHETYPE (1 && !WITH_EDITORONLY_DATA)
 #define UE_VERIFY_CACHED_ARCHETYPE 0
@@ -69,7 +70,9 @@ FAutoConsoleVariableRef CVarEnableArchetypeCache(
 }
 #endif // UE_CACHE_ARCHETYPE
 
-UObject* GetArchetypeFromRequiredInfoImpl(const UClass* Class, const UObject* Outer, FName Name, EObjectFlags ObjectFlags, bool bUseUpToDateClass)
+UObject* GetArchetypeImpl(const UObject* InObject, const FObjectArchetypeHelper::IObjectArchetypePolicy* Policy);
+
+UObject* GetArchetypeFromRequiredInfoImpl(const UClass* Class, const UObject* Outer, FName Name, EObjectFlags ObjectFlags, bool bUseUpToDateClass, const FObjectArchetypeHelper::IObjectArchetypePolicy* Policy)
 {
 	UObject* Result = NULL;
 	const bool bIsCDO = !!(ObjectFlags & RF_ClassDefaultObject);
@@ -87,15 +90,23 @@ UObject* GetArchetypeFromRequiredInfoImpl(const UClass* Class, const UObject* Ou
 
 			UObject* ArchetypeToSearch = nullptr;
 #if UE_CACHE_ARCHETYPE
-			ArchetypeToSearch = Outer->GetArchetype();
+			ArchetypeToSearch = GetArchetypeImpl(Outer, Policy);
 #if UE_VERIFY_CACHED_ARCHETYPE
 			{
-				UObject* VerifyArchetype = GetArchetypeFromRequiredInfoImpl(Outer->GetClass(), Outer->GetOuter(), Outer->GetFName(), Outer->GetFlags(), bUseUpToDateClass);
+				UObject* VerifyArchetype = GetArchetypeFromRequiredInfoImpl(Outer->GetClass(), Outer->GetOuter(), Outer->GetFName(), Outer->GetFlags(), bUseUpToDateClass, Policy);
 				checkf(ArchetypeToSearch == VerifyArchetype, TEXT("Cached archetype mismatch, expected: %s, cached: %s"), *GetFullNameSafe(VerifyArchetype), *GetFullNameSafe(ArchetypeToSearch));
 			}
 #endif // UE_VERIFY_CACHED_ARCHETYPE
 #else
-			ArchetypeToSearch = GetArchetypeFromRequiredInfoImpl(Outer->GetClass(), Outer->GetOuter(), Outer->GetFName(), Outer->GetFlags(), bUseUpToDateClass);
+			
+#if WITH_EDITOR
+			ArchetypeToSearch = Policy ? Policy->GetArchetype(Outer) : nullptr;
+#endif
+
+			if (!ArchetypeToSearch)
+			{
+				ArchetypeToSearch = GetArchetypeFromRequiredInfoImpl(Outer->GetClass(), Outer->GetOuter(), Outer->GetFName(), Outer->GetFlags(), bUseUpToDateClass, Policy);
+			}
 #endif // UE_CACHE_ARCHETYPE
 			UObject* MyArchetype = static_cast<UObject*>(FindObjectWithOuter(ArchetypeToSearch, Class, Name));
 			if (MyArchetype)
@@ -172,33 +183,44 @@ void CacheArchetypeForObject(UObject* Object, UObject* Archetype)
 UObject* UObject::GetArchetypeFromRequiredInfo(const UClass* Class, const UObject* Outer, FName Name, EObjectFlags ObjectFlags)
 {
 	bool bUseUpToDateClass = false;
-	return GetArchetypeFromRequiredInfoImpl(Class, Outer, Name, ObjectFlags, bUseUpToDateClass);
+	return GetArchetypeFromRequiredInfoImpl(Class, Outer, Name, ObjectFlags, bUseUpToDateClass, nullptr);
 }
 
 //DECLARE_FLOAT_ACCUMULATOR_STAT(TEXT("UObject::GetArchetype"), STAT_FArchiveRealtimeGC_GetArchetype, STATGROUP_GC);
 
-UObject* UObject::GetArchetype() const
+UObject* GetArchetypeImpl(const UObject* InObject, const FObjectArchetypeHelper::IObjectArchetypePolicy* Policy)
 {
 	//SCOPE_SECONDS_ACCUMULATOR(STAT_FArchiveRealtimeGC_GetArchetype);
 
+#if WITH_EDITOR
+	if (Policy)
+	{
+		if (UObject* Archetype = Policy->GetArchetype(InObject))
+		{
+			return Archetype;
+		}
+	}
+#endif
+
+	bool bUseUpToDateClass = false;
 #if UE_CACHE_ARCHETYPE
 	if (!bEnableArchetypeCache)
 	{
-		return GetArchetypeFromRequiredInfo(GetClass(), GetOuter(), GetFName(), GetFlags());
+		return GetArchetypeFromRequiredInfoImpl(InObject->GetClass(), InObject->GetOuter(), InObject->GetFName(), InObject->GetFlags(), bUseUpToDateClass, Policy);
 	}
 
 	UObject* Archetype = nullptr;
-	FArchetypeInfo Annoatation = ArchetypeAnnotation.GetAnnotation(this);
+	FArchetypeInfo Annoatation = ArchetypeAnnotation.GetAnnotation(InObject);
 	int32 ArchetypeIndex = Annoatation.ArchetypeIndex;
 	int32 SerialNumber = ArchetypeIndex == INDEX_NONE ? INDEX_NONE : GUObjectArray.GetSerialNumber(ArchetypeIndex);
 	if ((ArchetypeIndex == INDEX_NONE) || (SerialNumber != Annoatation.SerialNumber))
 	{
-		Archetype = GetArchetypeFromRequiredInfo(GetClass(), GetOuter(), GetFName(), GetFlags());
+		Archetype = GetArchetypeFromRequiredInfoImpl(InObject->GetClass(), InObject->GetOuter(), InObject->GetFName(), InObject->GetFlags(), bUseUpToDateClass, Policy);
 		// If the Outer is pending load we can't cache the archetype as it may be inacurate
-		if (Archetype && !(GetOuter() && GetOuter()->HasAnyFlags(RF_NeedLoad)))
+		if (Archetype && !(InObject->GetOuter() && InObject->GetOuter()->HasAnyFlags(RF_NeedLoad)))
 		{
 			ArchetypeIndex = GUObjectArray.ObjectToIndex(Archetype);
-			ArchetypeAnnotation.AddAnnotation(this, FArchetypeInfo{ ArchetypeIndex, GUObjectArray.AllocateSerialNumber(ArchetypeIndex) });
+			ArchetypeAnnotation.AddAnnotation(InObject, FArchetypeInfo{ ArchetypeIndex, GUObjectArray.AllocateSerialNumber(ArchetypeIndex) });
 		}
 	}
 	else
@@ -207,7 +229,7 @@ UObject* UObject::GetArchetype() const
 		check(ArchetypeItem != nullptr);
 		Archetype = static_cast<UObject*>(ArchetypeItem->Object);
 #if UE_VERIFY_CACHED_ARCHETYPE
-		UObject* ExpectedArchetype = GetArchetypeFromRequiredInfo(GetClass(), GetOuter(), GetFName(), GetFlags());
+		UObject* ExpectedArchetype = GetArchetypeFromRequiredInfoImpl(InObject->GetClass(), InObject->GetOuter(), InObject->GetFName(), InObject->GetFlags(), bUseUpToDateClass, Policy);
 		if (ExpectedArchetype != Archetype)
 		{
 			UE_LOG(LogClass, Fatal, TEXT("Cached archetype mismatch, expected: %s, cached: %s"), *GetFullNameSafe(ExpectedArchetype), *GetFullNameSafe(Archetype));
@@ -219,10 +241,20 @@ UObject* UObject::GetArchetype() const
 
 	return Archetype;
 #else
-	return GetArchetypeFromRequiredInfo(GetClass(), GetOuter(), GetFName(), GetFlags());
+	return GetArchetypeFromRequiredInfoImpl(InObject->GetClass(), InObject->GetOuter(), InObject->GetFName(), InObject->GetFlags(), bUseUpToDateClass, Policy);
 #endif // UE_CACHE_ARCHETYPE
 }
 
+UObject* UObject::GetArchetype() const
+{
+	return GetArchetypeImpl(this, nullptr);
+}
+
+UObject* FObjectArchetypeHelper::GetArchetype(const UObject* InObject, const FObjectArchetypeHelper::IObjectArchetypePolicy* Policy)
+{
+	return GetArchetypeImpl(InObject, Policy);
+}
+
 /** Removes all cached archetypes to avoid doing it in static exit where it may cause crashes */
 void CleanupCachedArchetypes()
 {
