Use MakeShared instead of MakeSharable to reduce number of allocations

diff --git a/Engine/Source/Editor/PropertyEditor/Private/DetailPropertyRow.cpp b/Engine/Source/Editor/PropertyEditor/Private/DetailPropertyRow.cpp
--- a/Engine/Source/Editor/PropertyEditor/Private/DetailPropertyRow.cpp
+++ b/Engine/Source/Editor/PropertyEditor/Private/DetailPropertyRow.cpp
@@ -200,7 +200,7 @@ void FDetailPropertyRow::GetDefaultWidgets( TSharedPtr<SWidget>& OutNameWidget,
 	TSharedPtr<IPropertyTypeCustomization>& CustomTypeInterface = GetTypeInterface();
 	if ( CustomTypeInterface.IsValid() ) 
 	{
-		CustomTypeRow = MakeShareable(new FDetailWidgetRow);
+		CustomTypeRow = MakeShared<FDetailWidgetRow>();
 
 		CustomTypeInterface->CustomizeHeader(PropertyHandle.ToSharedRef(), *CustomTypeRow, *this);
 	}
@@ -232,7 +232,7 @@ bool FDetailPropertyRow::RequiresTick() const
 FDetailWidgetRow& FDetailPropertyRow::CustomWidget( bool bShowChildren )
 {
 	bShowCustomPropertyChildren = bShowChildren;
-	CustomPropertyWidget = MakeShareable( new FDetailWidgetRow );
+	CustomPropertyWidget = MakeShared<FDetailWidgetRow>();
 	return *CustomPropertyWidget;
 }
 
@@ -341,7 +341,7 @@ void FDetailPropertyRow::OnItemNodeInitialized( TSharedRef<FDetailCategoryImpl>
 	// Don't customize the user already customized
 	if (!CustomPropertyWidget.IsValid() && CustomTypeInterface.IsValid())
 	{
-		CustomPropertyWidget = MakeShareable(new FDetailWidgetRow);
+		CustomPropertyWidget = MakeShared<FDetailWidgetRow>();
 
 		CustomTypeInterface->CustomizeHeader(PropertyHandle.ToSharedRef(), *CustomPropertyWidget, *this);
 
@@ -362,7 +362,7 @@ void FDetailPropertyRow::OnItemNodeInitialized( TSharedRef<FDetailCategoryImpl>
 
 	if( bShowCustomPropertyChildren && CustomTypeInterface.IsValid() )
 	{
-		PropertyTypeLayoutBuilder = MakeShareable(new FCustomChildrenBuilder(InParentCategory, InParentGroup));
+		PropertyTypeLayoutBuilder = MakeShared<FCustomChildrenBuilder>(InParentCategory, InParentGroup);
 
 		/** Does this row pass its custom reset behavior to its children? */
 		if (CustomResetToDefault.IsSet() && CustomResetToDefault->PropagatesToChildren())
@@ -415,7 +415,7 @@ void FDetailPropertyRow::GenerateChildrenForPropertyNode( TSharedPtr<FPropertyNo
 
 		for( int32 ChildIndex = 0; ChildIndex < ChildRows.Num(); ++ChildIndex )
 		{
-			TSharedRef<FDetailItemNode> ChildNodeItem = MakeShareable( new FDetailItemNode( ChildRows[ChildIndex], ParentCategory.Pin().ToSharedRef(), ParentEnabledState ) );
+			TSharedRef<FDetailItemNode> ChildNodeItem = MakeShared<FDetailItemNode>(ChildRows[ChildIndex], ParentCategory.Pin().ToSharedRef(), ParentEnabledState);
 			ChildNodeItem->Initialize();
 			OutChildren.Add( ChildNodeItem );
 		}
@@ -460,14 +460,14 @@ void FDetailPropertyRow::GenerateChildrenForPropertyNode( TSharedPtr<FPropertyNo
 
 					// Create and initialize the child first
 					FDetailLayoutCustomization Customization;
-					Customization.PropertyRow = MakeShareable(new FDetailPropertyRow(ChildNode, ParentCategoryRef));
+					Customization.PropertyRow = MakeShared<FDetailPropertyRow>(ChildNode, ParentCategoryRef);
 
 					if (CustomResetToDefault.IsSet() && CustomResetToDefault->PropagatesToChildren())
 					{
 						Customization.PropertyRow->OverrideResetToDefault(CustomResetToDefault.GetValue());
 					}
 
-					TSharedRef<FDetailItemNode> ChildNodeItem = MakeShareable(new FDetailItemNode(Customization, ParentCategoryRef, ParentEnabledState));
+					TSharedRef<FDetailItemNode> ChildNodeItem = MakeShared<FDetailItemNode>(Customization, ParentCategoryRef, ParentEnabledState);
 					ChildNodeItem->Initialize();
 
 					if ( ChildNode->GetPropertyKeyNode().IsValid() )
@@ -477,8 +477,8 @@ void FDetailPropertyRow::GenerateChildrenForPropertyNode( TSharedPtr<FPropertyNo
 						if ( !Customization.PropertyRow->PropertyKeyEditor.IsValid() )
 						{
 							FDetailLayoutCustomization KeyCustom;
-							KeyCustom.PropertyRow = MakeShareable(new FDetailPropertyRow(ChildNode->GetPropertyKeyNode(), ParentCategoryRef));
-							TSharedRef<FDetailItemNode> KeyNodeItem = MakeShareable(new FDetailItemNode(KeyCustom, ParentCategoryRef, ParentEnabledState));
+							KeyCustom.PropertyRow = MakeShared<FDetailPropertyRow>(ChildNode->GetPropertyKeyNode(), ParentCategoryRef);
+							TSharedRef<FDetailItemNode> KeyNodeItem = MakeShared<FDetailItemNode>(KeyCustom, ParentCategoryRef, ParentEnabledState);
 							KeyNodeItem->Initialize();
 							
 							PropNodes.Add(KeyNodeItem);
@@ -634,7 +634,7 @@ void FDetailPropertyRow::MakeExternalPropertyRowCustomization(TSharedPtr<FStruct
 
 		RootPropertyNode->AddChildNode(ItemNode);
 
-		OutCustomization.PropertyRow = MakeShareable(new FDetailPropertyRow(ItemNode, ParentCategory, RootPropertyNode));
+		OutCustomization.PropertyRow = MakeShared<FDetailPropertyRow>(ItemNode, ParentCategory, RootPropertyNode);
 		OutCustomization.PropertyRow->SetCustomExpansionId(Parameters.GetUniqueId());
 	}
 }
