Fix pasting into the Mobility field of the Details panel

Pasting into the Mobility field of the Details panel was prevented because no property row was found for the Mobility field. With this change, we allow pastes when no property row is found and instead only use the property row when it is found to prevent pastes if the row says so.

* Share "can we paste" checks by having PasteProperty call CanPasteProperty instead of running its own slightly different checks.
* Make CanPasteProperty (via CanPasteFromText) return true if no property row can be found for the property, instead of returning false and failing the paste. The intent here seems to have been allowing an existing property row to prevent the paste, but if none exists that shouldn't fail the paste. Ancient Unreal code didn't prevent pastes when the property row didn't exist either.

diff --git a/Engine/Source/Editor/PropertyEditor/Private/SDetailSingleItemRow.cpp b/Engine/Source/Editor/PropertyEditor/Private/SDetailSingleItemRow.cpp
--- a/Engine/Source/Editor/PropertyEditor/Private/SDetailSingleItemRow.cpp
+++ b/Engine/Source/Editor/PropertyEditor/Private/SDetailSingleItemRow.cpp
@@ -251,14 +251,11 @@ TSharedPtr<FPropertyNode> SDetailSingleItemRow::GetPropertyNode() const
 TSharedPtr<IPropertyHandle> SDetailSingleItemRow::GetPropertyHandle() const
 {
 	TSharedPtr<IPropertyHandle> Handle;
-	const TSharedPtr<FPropertyNode> PropertyNode = GetPropertyNode();
-	if (PropertyNode.IsValid())
+	if (const TSharedPtr<FPropertyNode> PropertyNode = GetPropertyNode())
 	{
-		const TSharedPtr<FDetailTreeNode> OwnerTreeNodePtr = OwnerTreeNode.Pin();
-		if (OwnerTreeNodePtr.IsValid())
+		if (const TSharedPtr<FDetailTreeNode> OwnerTreeNodePtr = OwnerTreeNode.Pin())
 		{
-			IDetailsViewPrivate* DetailsView = OwnerTreeNodePtr->GetDetailsView();
-			if (DetailsView != nullptr)
+			if (IDetailsViewPrivate* DetailsView = OwnerTreeNodePtr->GetDetailsView())
 			{
 				Handle = PropertyEditorHelpers::GetPropertyHandle(PropertyNode.ToSharedRef(), DetailsView->GetNotifyHook(), DetailsView->GetPropertyUtilities());
 			}
@@ -1371,10 +1368,7 @@ bool SDetailSingleItemRow::CanCopyPropertyInternalName()
 bool SDetailSingleItemRow::CanPasteProperty() const
 {
 	FString ClipboardContent;
-	if (OwnerTreeNode.IsValid())
-	{
-		FPropertyEditorClipboard::ClipboardPaste(ClipboardContent);
-	}
+	FPropertyEditorClipboard::ClipboardPaste(ClipboardContent);
 
 	return CanPasteFromText(TEXT(""), ClipboardContent);
 }
@@ -1406,23 +1400,57 @@ bool SDetailSingleItemRow::CanPasteFromText(const FString& InTag, const FString&
 	{
 		return false;
 	}
-	
-	// Prevent paste from working if the property's edit condition is not met.
-	TSharedPtr<FDetailPropertyRow> PropertyRow = Customization->PropertyRow;
-	if (!PropertyRow.IsValid() && Customization->DetailGroup.IsValid())
+
+	const TSharedPtr<FPropertyNode> PropertyNode = GetPropertyNode();
+
+	// Prevent paste if we cannot find the property node to paste into.
+	if (!PropertyNode.IsValid())
 	{
-		PropertyRow = Customization->DetailGroup->GetHeaderPropertyRow();
+		return false;
+	}
+
+	const TSharedPtr<IPropertyHandle> PropertyHandle = GetPropertyHandle();
+
+	// We won't be able to paste without a property handle.
+	if (!PropertyHandle.IsValid())
+	{
+		return false;
 	}
 
-	if (PropertyRow.IsValid())
+	if (const bool bIsTagged = !InTag.IsEmpty();
+		bIsTagged)
 	{
-		if (const FPropertyEditor* PropertyEditor = PropertyRow->GetPropertyEditor().Get())
+		const FString PropertyPath = UE::PropertyEditor::Private::GetPropertyPath(
+			[&]() { return GetPropertyHandle(); },
+			[&]() { return PropertyNode; });
+		
+		// Ensure that if tag is specified, that it matches the subscriber.
+		if (!InTag.Equals(PropertyPath))
 		{
-			return !PropertyEditor->IsEditConst();
+			return false;
 		}
 	}
 
-	return false;
+	// Prevent paste from working if the property's edit condition is not met.
+	// Allow paste if no property row can be found.
+	{
+		TSharedPtr<FDetailPropertyRow> PropertyRow = Customization->PropertyRow;
+		
+		if (!PropertyRow.IsValid() && Customization->DetailGroup.IsValid())
+		{
+			PropertyRow = Customization->DetailGroup->GetHeaderPropertyRow();
+		}
+
+		if (PropertyRow.IsValid())
+		{
+			if (const FPropertyEditor* PropertyEditor = PropertyRow->GetPropertyEditor().Get())
+			{
+				return !PropertyEditor->IsEditConst();
+			}
+		}
+	}
+
+	return true;
 }
 
 void SDetailSingleItemRow::OnPasteProperty()
@@ -1470,38 +1498,10 @@ void SDetailSingleItemRow::OnPasteFromText(const FString& InTag, const FString&
 
 bool SDetailSingleItemRow::PasteFromText(const FString& InTag, const FString& InText)
 {
-	if (InText.IsEmpty() || !OwnerTreeNode.IsValid())
+	if (!CanPasteFromText(InTag, InText))
 	{
 		return false;
 	}
-	
-	TSharedPtr<FPropertyNode> PropertyNode = GetPropertyNode();
-	if (!PropertyNode.IsValid() && Customization->DetailGroup.IsValid())
-	{
-		PropertyNode = Customization->DetailGroup->GetHeaderPropertyNode();
-	}
-
-	// If still invalid, there's nothing to paste to
-	if (!PropertyNode.IsValid())
-	{
-		return false;
-	}
-	
-	const TSharedPtr<IPropertyHandle> PropertyHandle = GetPropertyHandle();
-
-	const bool bIsTagged = !InTag.IsEmpty();
-	if (bIsTagged)
-	{
-		const FString PropertyPath = UE::PropertyEditor::Private::GetPropertyPath(
-			[&]() { return GetPropertyHandle(); },
-			[&]() { return PropertyNode; });
-		
-		// ensure that if tag is specified, that it matches the subscriber
-		if (!InTag.Equals(PropertyPath))
-		{
-			return false;
-		}
-	}
 
 	// The logic below is largely taken from SDisplayClusterColorGradingColorWheel::CommitColor,
 	// which avoids writing to trashed objects
@@ -1510,6 +1510,8 @@ bool SDetailSingleItemRow::PasteFromText(const FString& InTag, const FString& In
 	
 	EPropertyValueSetFlags::Type PropertyValueSetFlags = EPropertyValueSetFlags::InstanceObjects;
 
+	const bool bIsTagged = !InTag.IsEmpty();
+
 	// If tagged, add the InteractiveChange flag so as not to run PECP
 	// @todo: would be better to indicate that this is a batched paste rather than checking for a tag
 	if (bIsTagged)
@@ -1517,7 +1519,13 @@ bool SDetailSingleItemRow::PasteFromText(const FString& InTag, const FString& In
 		PropertyValueSetFlags |= EPropertyValueSetFlags::InteractiveChange;
 		PropertyValueSetFlags |= EPropertyValueSetFlags::NotTransactable;
 	}
-	
+
+	const TSharedPtr<IPropertyHandle> PropertyHandle = GetPropertyHandle();
+	if (!PropertyHandle.IsValid())
+	{
+		return false;
+	}
+
 	if (PropertyHandle->SetValueFromFormattedString(InText, PropertyValueSetFlags) != FPropertyAccess::Success)
 	{
 		return false;
