[PropertyEditor] Fixed a bug causing members of a container property of a struct property not respecting Editable flags

diff --git a/Engine/Source/Editor/PropertyEditor/Private/PropertyNode.cpp b/Engine/Source/Editor/PropertyEditor/Private/PropertyNode.cpp
--- a/Engine/Source/Editor/PropertyEditor/Private/PropertyNode.cpp
+++ b/Engine/Source/Editor/PropertyEditor/Private/PropertyNode.cpp
@@ -1346,7 +1346,12 @@ bool FPropertyNode::IsEditConst() const
 				FStructProperty* StructProperty = CastField<FStructProperty>(CurParent->GetProperty());
 				if (StructProperty == nullptr)
 				{
-					break;
+					const bool bIsContainerProperty = CastField<FArrayProperty>(CurParent->GetProperty()) || CastField<FSetProperty>(CurParent->GetProperty()) || CastField<FMapProperty>(CurParent->GetProperty()) || CastField<FOptionalProperty>(CurParent->GetProperty());
+					
+					if (!bIsContainerProperty)
+					{
+						break;
+					}
 				}
 
 				if (CurParent->IsEditConst())
@@ -1354,7 +1359,7 @@ bool FPropertyNode::IsEditConst() const
 					// An owning struct is edit const, so the child property is too
 					bIsEditConst = true;
 				}
-				else
+				else if (StructProperty)
 				{
 					// See if the struct has a problem with this property being editable
 					UScriptStruct* ScriptStruct = StructProperty->Struct;
