Arrays and sets: Only split the FormattedString when it is an array of objects so it starts with (" and ends with ")
#tests Copy and paste array of structs and array of objects
#rnx

diff --git a/Engine/Source/Editor/PropertyEditor/Private/PropertyHandleImpl.cpp b/Engine/Source/Editor/PropertyEditor/Private/PropertyHandleImpl.cpp
--- a/Engine/Source/Editor/PropertyEditor/Private/PropertyHandleImpl.cpp
+++ b/Engine/Source/Editor/PropertyEditor/Private/PropertyHandleImpl.cpp
@@ -5098,7 +5098,8 @@ FPropertyAccess::Result FPropertyHandleArray::SetValueFromFormattedString(const
 {
 	FPropertyAccess::Result Result = FPropertyAccess::Success;
 
-	if (InValue.StartsWith("((") && InValue.EndsWith("))"))
+	// for array of objects, copy each object individually in case some objects are forbidden
+	if (InValue.StartsWith("(\"") && InValue.EndsWith("\")"))
 	{
 		if (EmptyArray() != FPropertyAccess::Success)
 		{
@@ -5108,7 +5109,7 @@ FPropertyAccess::Result FPropertyHandleArray::SetValueFromFormattedString(const
 		TArray<FString> Values;
 
 		// Remove firsts and lasts parenthesis, and split the string 
-		InValue.LeftChop(2).RightChop(2).ParseIntoArrayWS(Values, TEXT("),("), true);
+		InValue.LeftChop(1).RightChop(1).ParseIntoArrayWS(Values, TEXT(","), true);
 
 		for (const FString& Value : Values)
 		{
@@ -5121,12 +5122,10 @@ FPropertyAccess::Result FPropertyHandleArray::SetValueFromFormattedString(const
 			{
 				return FPropertyAccess::Fail;
 			}
-
-			// Readd trimmed parenthesis
-			FString PasteValue = "(" + Value + ")";
 			
 			const TSharedRef<IPropertyHandle> Property = GetElement( NumElements - 1 );
-			if (Property->IsValidHandle() && Property->SetValueFromFormattedString(PasteValue, Flags) == FPropertyAccess::Fail)
+			ensure(!Property->AsStruct());
+			if (Property->IsValidHandle() && Property->SetValueFromFormattedString(Value, Flags) == FPropertyAccess::Fail)
 			{
 				Result = FPropertyAccess::Fail;
 			}
@@ -5632,7 +5631,8 @@ FPropertyAccess::Result FPropertyHandleSet::SetValueFromFormattedString(const FS
 {
 	FPropertyAccess::Result Result = FPropertyAccess::Success;
 
-	if (InValue.StartsWith("((") && InValue.EndsWith("))"))
+	// for sets of objects, copy each object individually in case some objects are forbidden
+	if (InValue.StartsWith("(\"") && InValue.EndsWith("\")"))
 	{
 		if (Empty() != FPropertyAccess::Success)
 		{
@@ -5642,7 +5642,7 @@ FPropertyAccess::Result FPropertyHandleSet::SetValueFromFormattedString(const FS
 		TArray<FString> Values;
 
 		// Remove firsts and lasts parenthesis  
-		InValue.LeftChop(2).RightChop(2).ParseIntoArrayWS(Values, TEXT("),("), true);
+		InValue.LeftChop(1).RightChop(1).ParseIntoArrayWS(Values, TEXT(","), true);
 		for (const FString& Value : Values)
 		{
 			AddItem();
@@ -5651,12 +5651,9 @@ FPropertyAccess::Result FPropertyHandleSet::SetValueFromFormattedString(const FS
 			{
 				return FPropertyAccess::Fail;
 			}
-
-			// Readd trimmed parenthesis
-			FString PasteValue = "(" + Value + ")";
-			
 			const TSharedRef<IPropertyHandle> Property = GetElement( NumElements - 1 );
-			if ( Property->IsValidHandle() && Property->SetValueFromFormattedString(PasteValue, Flags) == FPropertyAccess::Fail)
+			ensure(!Property->AsStruct());
+			if ( Property->IsValidHandle() && Property->SetValueFromFormattedString(Value, Flags) == FPropertyAccess::Fail)
 			{
 				Result = FPropertyAccess::Fail;
 			}
