Added customization support to properties added as external objects/structs
Also, added the processing of ShouldHideRootObjectNode at DetailCategoryBuilder and DetailPropertyRow

diff --git a/Engine/Source/Editor/PropertyEditor/Private/CustomChildBuilder.cpp b/Engine/Source/Editor/PropertyEditor/Private/CustomChildBuilder.cpp
--- a/Engine/Source/Editor/PropertyEditor/Private/CustomChildBuilder.cpp
+++ b/Engine/Source/Editor/PropertyEditor/Private/CustomChildBuilder.cpp
@@ -76,6 +76,11 @@ IDetailPropertyRow* FCustomChildrenBuilder::AddExternalStructureProperty(TShared
 
 	FDetailPropertyRow::MakeExternalPropertyRowCustomization(ChildStructure, PropertyName, ParentCategoryRef, NewCustomization, Params);
 
+	if (Params.ShouldHideRootObjectNode() && NewCustomization.HasPropertyNode() && NewCustomization.GetPropertyNode()->AsComplexNode())
+	{
+		NewCustomization.PropertyRow->SetForceShowOnlyChildren(true);
+	}
+
 	TSharedPtr<FDetailPropertyRow> NewRow = NewCustomization.PropertyRow;
 
 	if (NewRow.IsValid())
diff --git a/Engine/Source/Editor/PropertyEditor/Private/DetailCategoryBuilderImpl.cpp b/Engine/Source/Editor/PropertyEditor/Private/DetailCategoryBuilderImpl.cpp
--- a/Engine/Source/Editor/PropertyEditor/Private/DetailCategoryBuilderImpl.cpp
+++ b/Engine/Source/Editor/PropertyEditor/Private/DetailCategoryBuilderImpl.cpp
@@ -489,6 +489,11 @@ IDetailPropertyRow* FDetailCategoryImpl::AddExternalStructureProperty(TSharedPtr
 
 	FDetailPropertyRow::MakeExternalPropertyRowCustomization(StructData, PropertyName, AsShared(), NewCustomization, Params);
 
+	if (Params.ShouldHideRootObjectNode() && NewCustomization.HasPropertyNode() && NewCustomization.GetPropertyNode()->AsComplexNode())
+	{
+		NewCustomization.PropertyRow->SetForceShowOnlyChildren(true);
+	}
+
 	TSharedPtr<FDetailPropertyRow> NewRow = NewCustomization.PropertyRow;
 
 	if (NewRow.IsValid())
@@ -510,6 +515,11 @@ IDetailPropertyRow* FDetailCategoryImpl::AddExternalStructureProperty(TSharedPtr
 
 	FDetailPropertyRow::MakeExternalPropertyRowCustomization(StructDataProvider, PropertyName, AsShared(), NewCustomization, Params);
 
+	if (Params.ShouldHideRootObjectNode() && NewCustomization.HasPropertyNode() && NewCustomization.GetPropertyNode()->AsComplexNode())
+	{
+		NewCustomization.PropertyRow->SetForceShowOnlyChildren(true);
+	}
+
 	TSharedPtr<FDetailPropertyRow> NewRow = NewCustomization.PropertyRow;
 
 	if (NewRow.IsValid())
diff --git a/Engine/Source/Editor/PropertyEditor/Private/DetailPropertyRow.cpp b/Engine/Source/Editor/PropertyEditor/Private/DetailPropertyRow.cpp
--- a/Engine/Source/Editor/PropertyEditor/Private/DetailPropertyRow.cpp
+++ b/Engine/Source/Editor/PropertyEditor/Private/DetailPropertyRow.cpp
@@ -563,7 +563,38 @@ TSharedPtr<IPropertyTypeCustomization> FDetailPropertyRow::GetPropertyCustomizat
 		static FName NAME_PropertyEditor("PropertyEditor");
 		FPropertyEditorModule& PropertyEditorModule = FModuleManager::GetModuleChecked<FPropertyEditorModule>(NAME_PropertyEditor);
 
-		FPropertyTypeLayoutCallback LayoutCallback = PropertyEditorModule.GetPropertyTypeCustomization(Property, *PropHandle, InParentCategory->GetCustomPropertyTypeLayoutMap() );
+		FPropertyTypeLayoutCallback LayoutCallback;
+		if (Property != nullptr)
+		{
+			LayoutCallback = PropertyEditorModule.GetPropertyTypeCustomization(Property, *PropHandle, InParentCategory->GetCustomPropertyTypeLayoutMap());
+		}
+		else
+		{
+			// This add support to objects and structs added to the category with AddExternalObjectProperty / AddExternalStructureProperty
+			if (FComplexPropertyNode* ComplexNode = InPropertyNode->AsComplexNode())
+			{
+				if (FObjectPropertyNode* ObjectNode = ComplexNode->AsObjectNode())
+				{
+					UClass* PropertyClass = ObjectNode->GetObjectBaseClass();
+					while (PropertyClass)
+					{
+						LayoutCallback = PropertyEditorModule.FindPropertyTypeLayoutCallback(PropertyClass->GetFName(), *PropHandle, InParentCategory->GetCustomPropertyTypeLayoutMap());
+						if (LayoutCallback.IsValid())
+						{
+							break;
+						}
+
+						PropertyClass = PropertyClass->GetSuperClass();
+					}
+				}
+				else if (FStructurePropertyNode* StructureNode = ComplexNode->AsStructureNode())
+				{
+					const FName PropertyTypeName = StructureNode->GetBaseStructure()->GetFName();
+					LayoutCallback = PropertyEditorModule.FindPropertyTypeLayoutCallback(PropertyTypeName, *PropHandle, InParentCategory->GetCustomPropertyTypeLayoutMap());
+				}
+			}
+		}
+
 		if (LayoutCallback.IsValid())
 		{
 			if (PropHandle->IsValidHandle())
