Convert usage of MakeShareable to MakeShared

diff --git a/Engine/Source/Editor/PropertyEditor/Private/Presentation/PropertyTable/PropertyTableColumn.cpp b/Engine/Source/Editor/PropertyEditor/Private/Presentation/PropertyTable/PropertyTableColumn.cpp
--- a/Engine/Source/Editor/PropertyEditor/Private/Presentation/PropertyTable/PropertyTableColumn.cpp
+++ b/Engine/Source/Editor/PropertyEditor/Private/Presentation/PropertyTable/PropertyTableColumn.cpp
@@ -375,7 +375,7 @@ FORCEINLINE int32 FCompareRowByColumnAscending<FStructProperty>::CompareProperty
 
 FPropertyTableColumn::FPropertyTableColumn( const TSharedRef< IPropertyTable >& InTable, const TWeakObjectPtr< UObject >& InObject )
 	: Cells()
-	, DataSource( MakeShareable( new UObjectDataSource( InObject.Get() ) ) )
+	, DataSource( MakeShared<UObjectDataSource>(InObject.Get() ) )
 	, Table( InTable )
 	, Id( NAME_None )
 	, DisplayName()
diff --git a/Engine/Source/Editor/PropertyEditor/Private/Presentation/PropertyTable/PropertyTableRow.cpp b/Engine/Source/Editor/PropertyEditor/Private/Presentation/PropertyTable/PropertyTableRow.cpp
--- a/Engine/Source/Editor/PropertyEditor/Private/Presentation/PropertyTable/PropertyTableRow.cpp
+++ b/Engine/Source/Editor/PropertyEditor/Private/Presentation/PropertyTable/PropertyTableRow.cpp
@@ -17,7 +17,7 @@
 class UObject;
 
 FPropertyTableRow::FPropertyTableRow( const TSharedRef< class IPropertyTable >& InTable, const TWeakObjectPtr< UObject >& InObject )
-	: DataSource( MakeShareable( new UObjectDataSource( InObject ) ) )
+	: DataSource( MakeShared<UObjectDataSource>( InObject ) )
 	, Table( InTable )
 	, Children()
 	, PartialPath( FPropertyPath::CreateEmpty() )
@@ -26,7 +26,7 @@ FPropertyTableRow::FPropertyTableRow( const TSharedRef< class IPropertyTable >&
 }
 
 FPropertyTableRow::FPropertyTableRow( const TSharedRef< class IPropertyTable >& InTable, const TSharedRef< FPropertyPath >& InPropertyPath )
-	: DataSource( MakeShareable( new PropertyPathDataSource( InPropertyPath ) ) )
+	: DataSource( MakeShared<PropertyPathDataSource>( InPropertyPath ) )
 	, Table( InTable )
 	, Children()
 	, PartialPath( FPropertyPath::CreateEmpty() )
@@ -35,7 +35,7 @@ FPropertyTableRow::FPropertyTableRow( const TSharedRef< class IPropertyTable >&
 }
 
 FPropertyTableRow::FPropertyTableRow( const TSharedRef< class IPropertyTable >& InTable, const TWeakObjectPtr< UObject >& InObject, const TSharedRef< FPropertyPath >& InPartialPropertyPath )
-	: DataSource( MakeShareable( new UObjectDataSource( InObject ) ) )
+	: DataSource( MakeShared<UObjectDataSource>( InObject ) )
 	, Table( InTable )
 	, Children()
 	, PartialPath( InPartialPropertyPath )
@@ -44,7 +44,7 @@ FPropertyTableRow::FPropertyTableRow( const TSharedRef< class IPropertyTable >&
 }
 
 FPropertyTableRow::FPropertyTableRow( const TSharedRef< class IPropertyTable >& InTable, const TSharedRef< FPropertyPath >& InPropertyPath, const TSharedRef< FPropertyPath >& InPartialPropertyPath )
-	: DataSource( MakeShareable( new PropertyPathDataSource( InPropertyPath ) ) )
+	: DataSource( MakeShared<PropertyPathDataSource>( InPropertyPath ) )
 	, Table( InTable )
 	, Children()
 	, PartialPath( InPartialPropertyPath )
@@ -100,7 +100,7 @@ void FPropertyTableRow::GenerateChildren()
 			Extension.Property = ChildNode->GetProperty();
 			Extension.ArrayIndex = ChildNode->GetArrayIndex();
 
-			Children.Add( MakeShareable( new FPropertyTableRow( TableRef, Object, RootPath->ExtendPath( Extension ) ) ) );
+			Children.Add( MakeShared<FPropertyTableRow>( TableRef, Object, RootPath->ExtendPath( Extension ) ) );
 		}
 	}
 }
diff --git a/Engine/Source/Editor/PropertyEditor/Private/PropertyNode.h b/Engine/Source/Editor/PropertyEditor/Private/PropertyNode.h
--- a/Engine/Source/Editor/PropertyEditor/Private/PropertyNode.h
+++ b/Engine/Source/Editor/PropertyEditor/Private/PropertyNode.h
@@ -806,7 +806,7 @@ public:
 
 		if (CurrentNode != nullptr && CurrentNode->AsCategoryNode() != nullptr)
 		{
-			TSharedRef< FPropertyPath > NewPath = MakeShareable(new FPropertyPath());
+			TSharedRef< FPropertyPath > NewPath = MakeShared<FPropertyPath>();
 			return NewPath;
 		}
 
@@ -818,7 +818,7 @@ public:
 				NewPropInfo.Property = CurrentNode->GetProperty();
 				if (!NewPropInfo.Property.IsValid())
 				{
-					return MakeShareable(new FPropertyPath());
+					return MakeShared<FPropertyPath>();
 				}
 				NewPropInfo.ArrayIndex = CurrentNode->GetArrayIndex();
 
@@ -828,7 +828,7 @@ public:
 			CurrentNode = CurrentNode->GetParentNode();
 		}
 
-		TSharedRef< FPropertyPath > NewPath = MakeShareable(new FPropertyPath());
+		TSharedRef< FPropertyPath > NewPath = MakeShared<FPropertyPath>();
 
 		for (int PropertyIndex = Properties.Num() - 1; PropertyIndex >= 0; --PropertyIndex)
 		{
diff --git a/Engine/Source/Editor/PropertyEditor/Private/SPropertyTreeViewImpl.cpp b/Engine/Source/Editor/PropertyEditor/Private/SPropertyTreeViewImpl.cpp
--- a/Engine/Source/Editor/PropertyEditor/Private/SPropertyTreeViewImpl.cpp
+++ b/Engine/Source/Editor/PropertyEditor/Private/SPropertyTreeViewImpl.cpp
@@ -33,7 +33,7 @@ public:
 	FPropertyUtilitiesTreeView( SPropertyTreeViewImpl& InView )
 		: View( InView )
 	{
-		EditConditionParser = MakeShareable(new FEditConditionParser());
+		EditConditionParser = MakeShared<FEditConditionParser>();
 	}
 
 	virtual class FNotifyHook* GetNotifyHook() const override
@@ -160,9 +160,9 @@ void SPropertyTreeViewImpl::Construct(const FArguments& InArgs)
 	bFavoritesEnabled = bFavoritesEnabled && bFavoritesAllowed;
 
 	// Create the root property now
-	RootPropertyNode = MakeShareable( new FObjectPropertyNode );
+	RootPropertyNode = MakeShared<FObjectPropertyNode>();
 
-	PropertySettings = MakeShareable( new FPropertyUtilitiesTreeView(*this) );
+	PropertySettings = MakeShared<FPropertyUtilitiesTreeView>(*this);
 
 	ConstructPropertyTree();
 	
