Add a ClassViewerFilter to the PropertyRestriction. Allow a customizer to easily add extra filter to a PropertyEdtiorClass or PropertyEditorEditInline.

diff --git a/Engine/Source/Editor/PropertyEditor/Private/PropertyRestriction.cpp b/Engine/Source/Editor/PropertyEditor/Private/PropertyRestriction.cpp
--- a/Engine/Source/Editor/PropertyEditor/Private/PropertyRestriction.cpp
+++ b/Engine/Source/Editor/PropertyEditor/Private/PropertyRestriction.cpp
@@ -39,6 +39,11 @@ void FPropertyRestriction::AddDisabledValue(FString InValue)
 	DisabledValues.Add(MoveTemp(InValue));
 }
 
+void FPropertyRestriction::AddClassFilter(TSharedRef<IClassViewerFilter> InValue)
+{
+	ClassViewFilter.Add(MoveTemp(InValue));
+}
+
 void FPropertyRestriction::RemoveHiddenValue(FString Value)
 {
 	HiddenValues.Remove(Value);
@@ -49,8 +54,14 @@ void FPropertyRestriction::RemoveDisabledValue(FString Value)
 	DisabledValues.Remove(Value);
 }
 
+void FPropertyRestriction::RemoveClassFilter(TSharedRef<IClassViewerFilter> Value)
+{
+	ClassViewFilter.Remove(Value);
+}
+
 void FPropertyRestriction::RemoveAll()
 {
 	HiddenValues.Empty();
 	DisabledValues.Empty();
+	ClassViewFilter.Empty();
 }
diff --git a/Engine/Source/Editor/PropertyEditor/Private/UserInterface/PropertyEditor/SPropertyEditorClass.cpp b/Engine/Source/Editor/PropertyEditor/Private/UserInterface/PropertyEditor/SPropertyEditorClass.cpp
--- a/Engine/Source/Editor/PropertyEditor/Private/UserInterface/PropertyEditor/SPropertyEditorClass.cpp
+++ b/Engine/Source/Editor/PropertyEditor/Private/UserInterface/PropertyEditor/SPropertyEditorClass.cpp
@@ -118,7 +118,9 @@ bool SPropertyEditorClass::CanEdit() const
 void SPropertyEditorClass::Construct(const FArguments& InArgs, const TSharedPtr< FPropertyEditor >& InPropertyEditor)
 {
 	PropertyEditor = InPropertyEditor;
-	
+
+	TArray<TSharedRef<class IClassViewerFilter>> ClassViewerFilters = InArgs._ClassViewerFilters;
+
 	if (PropertyEditor.IsValid())
 	{
 		const TSharedRef<FPropertyNode> PropertyNode = PropertyEditor->GetPropertyNode();
@@ -189,6 +191,11 @@ void SPropertyEditorClass::Construct(const FArguments& InArgs, const TSharedPtr<
 					DisallowedClassFilters.Add(DisabledClass);
 				}
 			}
+
+			for (TArray<TSharedRef<IClassViewerFilter>>::TConstIterator Iter = ClassRestriction.Get().GeClassViewFilterIterator(); Iter; ++Iter)
+			{
+				ClassViewerFilters.Add(*Iter);
+			}
 		}
 	}
 	else
@@ -212,7 +219,7 @@ void SPropertyEditorClass::Construct(const FArguments& InArgs, const TSharedPtr<
 		OnSetClass = InArgs._OnSetClass;
 	}
 
-	CreateClassFilter(InArgs._ClassViewerFilters);
+	CreateClassFilter(ClassViewerFilters);
 
 	SAssignNew(ComboButton, SComboButton)
 		.OnGetMenuContent(this, &SPropertyEditorClass::GenerateClassPicker)
diff --git a/Engine/Source/Editor/PropertyEditor/Private/UserInterface/PropertyEditor/SPropertyEditorEditInline.cpp b/Engine/Source/Editor/PropertyEditor/Private/UserInterface/PropertyEditor/SPropertyEditorEditInline.cpp
--- a/Engine/Source/Editor/PropertyEditor/Private/UserInterface/PropertyEditor/SPropertyEditorEditInline.cpp
+++ b/Engine/Source/Editor/PropertyEditor/Private/UserInterface/PropertyEditor/SPropertyEditorEditInline.cpp
@@ -307,6 +307,11 @@ TSharedRef<SWidget> SPropertyEditorEditInline::GenerateClassPicker()
 				DisallowedClassFilters.Add(DisabledClass);
 			}
 		}
+
+		for (TArray<TSharedRef<IClassViewerFilter>>::TConstIterator  Iter = ClassRestriction.Get().GeClassViewFilterIterator(); Iter; ++Iter)
+		{
+			Options.ClassFilters.Add(*Iter);
+		}
 	}
 
 	ClassFilter->AllowedClassFilters = MoveTemp(AllowedClassFilters);
diff --git a/Engine/Source/Editor/PropertyEditor/Public/PropertyRestriction.h b/Engine/Source/Editor/PropertyEditor/Public/PropertyRestriction.h
--- a/Engine/Source/Editor/PropertyEditor/Public/PropertyRestriction.h
+++ b/Engine/Source/Editor/PropertyEditor/Public/PropertyRestriction.h
@@ -7,6 +7,8 @@
 #include "CoreMinimal.h"
 #include "Internationalization/Text.h"
 
+class IClassViewerFilter;
+
 class PROPERTYEDITOR_API FPropertyRestriction
 {
 public:
@@ -22,9 +24,11 @@ public:
 	bool IsValueDisabled(const FString& Value) const;
 	void AddHiddenValue(FString Value);
 	void AddDisabledValue(FString Value);
+	void AddClassFilter(TSharedRef<IClassViewerFilter> ClassFilter);
 
 	void RemoveHiddenValue(FString Value);
 	void RemoveDisabledValue(FString Value);
+	void RemoveClassFilter(TSharedRef<IClassViewerFilter> ClassFilter);
 	void RemoveAll();
 
 	TArray<FString>::TConstIterator GetHiddenValuesIterator() const 
@@ -36,10 +40,15 @@ public:
 	{
 		return DisabledValues.CreateConstIterator();
 	}
+	
+	TArray<TSharedRef<IClassViewerFilter>>::TConstIterator GeClassViewFilterIterator() const
+	{
+		return ClassViewFilter.CreateConstIterator();
+	}
 
 private:
-
 	TArray<FString> HiddenValues;
 	TArray<FString> DisabledValues;
+	TArray<TSharedRef<IClassViewerFilter>> ClassViewFilter;
 	FText Reason;
 };
