Fixed SPropertyEditorAsset asset thumbnail not using the proper class when getting the asset type actions which caused the thumbnail to have the wrong asset type color.

diff --git a/Engine/Source/Editor/PropertyEditor/Private/UserInterface/PropertyEditor/SPropertyEditorAsset.cpp b/Engine/Source/Editor/PropertyEditor/Private/UserInterface/PropertyEditor/SPropertyEditorAsset.cpp
--- a/Engine/Source/Editor/PropertyEditor/Private/UserInterface/PropertyEditor/SPropertyEditorAsset.cpp
+++ b/Engine/Source/Editor/PropertyEditor/Private/UserInterface/PropertyEditor/SPropertyEditorAsset.cpp
@@ -549,8 +549,16 @@ void SPropertyEditorAsset::Construct(const FArguments& InArgs, const TSharedPtr<
 		TSharedPtr<IAssetTypeActions> AssetTypeActions;
 		if (ObjectClass != nullptr)
 		{
+			const UClass* EffectiveClass = ObjectClass;
+			if (EffectiveClass->GetPathName() != Value.AssetData.AssetClassPath.ToString())
+			{
+				if (UClass* AssetDataClass = FindObject<UClass>(Value.AssetData.AssetClassPath); AssetDataClass && AssetDataClass->IsChildOf(EffectiveClass))
+				{
+					EffectiveClass = AssetDataClass;
+				}
+			}
 			FAssetToolsModule& AssetToolsModule = FModuleManager::LoadModuleChecked<FAssetToolsModule>(TEXT("AssetTools"));
-			AssetTypeActions = AssetToolsModule.Get().GetAssetTypeActionsForClass(ObjectClass).Pin();
+			AssetTypeActions = AssetToolsModule.Get().GetAssetTypeActionsForClass(EffectiveClass).Pin();
 
 			if (AssetTypeActions.IsValid())
 			{
