Add new function for for DetailLayoutBuilder::EditCategory that does not hide category-less properties
* It wasn't possible to get at the IPropertyHandles for category-less properties since passing a "None" FName would default the category to "General".  However, it is possible to construct struct/class metadata that has properties with no category.
* Allows detail panel customizations to be written that affect all properties.

diff --git a/Engine/Source/Editor/PropertyEditor/Private/DetailLayoutBuilderImpl.cpp b/Engine/Source/Editor/PropertyEditor/Private/DetailLayoutBuilderImpl.cpp
--- a/Engine/Source/Editor/PropertyEditor/Private/DetailLayoutBuilderImpl.cpp
+++ b/Engine/Source/Editor/PropertyEditor/Private/DetailLayoutBuilderImpl.cpp
@@ -38,11 +38,16 @@ IDetailCategoryBuilder& FDetailLayoutBuilderImpl::EditCategory(FName CategoryNam
 	{
 		static const FText GeneralString = NSLOCTEXT("DetailLayoutBuilderImpl", "General", "General");
 		static const FName GeneralName = TEXT("General");
-
+	
 		CategoryName = GeneralName;
 		LocalizedDisplayName = GeneralString;
 	}
 
+	return EditCategoryAllowNone(CategoryName, LocalizedDisplayName, CategoryType);
+}
+
+IDetailCategoryBuilder& FDetailLayoutBuilderImpl::EditCategoryAllowNone(FName CategoryName, const FText& NewLocalizedDisplayName, ECategoryPriority::Type CategoryType)
+{
 	TSharedPtr<FDetailCategoryImpl> CategoryImpl;
 	// If the default category map had a category by the provided name, remove it from the map as it is now customized
 	if (!DefaultCategoryMap.RemoveAndCopyValue(CategoryName, CategoryImpl))
@@ -69,7 +74,7 @@ IDetailCategoryBuilder& FDetailLayoutBuilderImpl::EditCategory(FName CategoryNam
 		const int32 SortOrder = CategoryType * 1000 + (CustomCategoryMap.Num() - 1);
 		CategoryImpl->SetSortOrder( SortOrder );
 	}
-	CategoryImpl->SetDisplayName(CategoryName, LocalizedDisplayName);
+	CategoryImpl->SetDisplayName(CategoryName, NewLocalizedDisplayName);
 
 	return *CategoryImpl;
 }
diff --git a/Engine/Source/Editor/PropertyEditor/Private/DetailLayoutBuilderImpl.h b/Engine/Source/Editor/PropertyEditor/Private/DetailLayoutBuilderImpl.h
--- a/Engine/Source/Editor/PropertyEditor/Private/DetailLayoutBuilderImpl.h
+++ b/Engine/Source/Editor/PropertyEditor/Private/DetailLayoutBuilderImpl.h
@@ -32,6 +32,7 @@ public:
 	virtual void GetObjectsBeingCustomized(TArray< TWeakObjectPtr<UObject> >& OutObjects) const override;
 	virtual void GetStructsBeingCustomized(TArray< TSharedPtr<FStructOnScope> >& OutStructs) const override;
 	virtual IDetailCategoryBuilder& EditCategory(FName CategoryName, const FText& NewLocalizedDisplayName = FText::GetEmpty(), ECategoryPriority::Type CategoryType = ECategoryPriority::Default) override;
+	virtual IDetailCategoryBuilder& EditCategoryAllowNone(FName CategoryName, const FText& NewLocalizedDisplayName = FText::GetEmpty(), ECategoryPriority::Type CategoryType = ECategoryPriority::Default) override;
 	virtual void GetCategoryNames(TArray<FName>& OutCategoryNames) const override;
 	virtual IDetailPropertyRow& AddPropertyToCategory(TSharedPtr<IPropertyHandle> InPropertyHandle) override;
 	virtual FDetailWidgetRow& AddCustomRowToCategory(TSharedPtr<IPropertyHandle> InPropertyHandle, const FText& InCustomSearchString, bool bForAdvanced = false) override;
diff --git a/Engine/Source/Editor/PropertyEditor/Public/DetailLayoutBuilder.h b/Engine/Source/Editor/PropertyEditor/Public/DetailLayoutBuilder.h
--- a/Engine/Source/Editor/PropertyEditor/Public/DetailLayoutBuilder.h
+++ b/Engine/Source/Editor/PropertyEditor/Public/DetailLayoutBuilder.h
@@ -121,6 +121,16 @@ public:
 	 */
 	virtual IDetailCategoryBuilder& EditCategory(FName CategoryName, const FText& NewLocalizedDisplayName = FText::GetEmpty(), ECategoryPriority::Type CategoryType = ECategoryPriority::Default) = 0;
 
+	/**
+	* Edits an existing category or creates a new one
+	* If CategoryName is NAME_None, will enable access to properties without categories
+	* 
+	* @param CategoryName				The name of the category
+	* @param NewLocalizedDisplayName	The new display name of the category (optional)
+	* @param CategoryType				Category type to define sort order.  Category display order is sorted by this type (optional)
+	*/
+	virtual IDetailCategoryBuilder& EditCategoryAllowNone(FName CategoryName, const FText& NewLocalizedDisplayName = FText::GetEmpty(), ECategoryPriority::Type CategoryType = ECategoryPriority::Default) = 0;
+
 	/**
 	 * Gets the current set of existing category names. This includes both categories derived from properties and categories added via EditCategory.
 	 * @param	OutCategoryNames	 The array of category names
