Disable multi-select for entities.

diff --git a/Engine/Source/Editor/PropertyEditor/Public/DetailsViewObjectFilter.h b/Engine/Source/Editor/PropertyEditor/Public/DetailsViewObjectFilter.h
--- a/Engine/Source/Editor/PropertyEditor/Public/DetailsViewObjectFilter.h
+++ b/Engine/Source/Editor/PropertyEditor/Public/DetailsViewObjectFilter.h
@@ -26,14 +26,9 @@ struct FDetailsViewObjectRoot
 		Objects.Add(InObject);
 	}
 
-	FDetailsViewObjectRoot(const TArray<UObject*>& InObjects)
-	{
-		Objects.Reserve(InObjects.Num());
-		for (UObject* Object : InObjects)
-		{
-			Objects.Add(Object);
-		}
-	}
+	FDetailsViewObjectRoot(TArray<UObject*> InObjects)
+		: Objects(MoveTemp(InObjects))
+	{}
 
 	TArray<UObject*> Objects;
 };
