Add IStructureDataProvide version of FCustomChildrenBuilder::AddAllExternalStructureProperties()

diff --git a/Engine/Source/Editor/PropertyEditor/Private/CustomChildBuilder.cpp b/Engine/Source/Editor/PropertyEditor/Private/CustomChildBuilder.cpp
--- a/Engine/Source/Editor/PropertyEditor/Private/CustomChildBuilder.cpp
+++ b/Engine/Source/Editor/PropertyEditor/Private/CustomChildBuilder.cpp
@@ -98,7 +98,14 @@ IDetailPropertyRow* FCustomChildrenBuilder::AddExternalStructureProperty(TShared
 
 TArray<TSharedPtr<IPropertyHandle>> FCustomChildrenBuilder::AddAllExternalStructureProperties(TSharedRef<FStructOnScope> ChildStructure)
 {
-	return ParentCategory.Pin()->AddAllExternalStructureProperties(ChildStructure);
+	const TSharedPtr<FDetailCategoryImpl> ParentCategoryPinned = ParentCategory.Pin();
+	return ParentCategoryPinned ? ParentCategoryPinned->AddAllExternalStructureProperties(ChildStructure) : TArray<TSharedPtr<IPropertyHandle>>();
+}
+
+TArray<TSharedPtr<IPropertyHandle>> FCustomChildrenBuilder::AddAllExternalStructureProperties(TSharedPtr<IStructureDataProvider> ChildStructure)
+{
+	const TSharedPtr<FDetailCategoryImpl> ParentCategoryPinned = ParentCategory.Pin();
+	return ParentCategoryPinned ? ParentCategoryPinned->AddAllExternalStructureProperties(ChildStructure) : TArray<TSharedPtr<IPropertyHandle>>();
 }
 
 IDetailPropertyRow* FCustomChildrenBuilder::AddExternalObjects(const TArray<UObject*>& Objects, FName UniqueIdName)
diff --git a/Engine/Source/Editor/PropertyEditor/Private/CustomChildBuilder.h b/Engine/Source/Editor/PropertyEditor/Private/CustomChildBuilder.h
--- a/Engine/Source/Editor/PropertyEditor/Private/CustomChildBuilder.h
+++ b/Engine/Source/Editor/PropertyEditor/Private/CustomChildBuilder.h
@@ -31,6 +31,7 @@ public:
 	virtual IDetailPropertyRow* AddExternalStructure(TSharedRef<FStructOnScope> ChildStructure, FName UniqueIdName) override;
 	virtual IDetailPropertyRow* AddExternalStructureProperty(TSharedRef<FStructOnScope> ChildStructure, FName PropertyName, const FAddPropertyParams& Params) override;
 	virtual TArray<TSharedPtr<IPropertyHandle>> AddAllExternalStructureProperties(TSharedRef<FStructOnScope> ChildStructure) override;
+	virtual TArray<TSharedPtr<IPropertyHandle>> AddAllExternalStructureProperties(TSharedPtr<IStructureDataProvider> ChildStructure) override;
 	virtual TSharedRef<SWidget> GenerateStructValueWidget(TSharedRef<IPropertyHandle> StructPropertyHandle) override;
 	virtual IDetailCategoryBuilder& GetParentCategory() const override;
 	virtual IDetailGroup* GetParentGroup() const override;
diff --git a/Engine/Source/Editor/PropertyEditor/Private/DetailCategoryBuilderImpl.h b/Engine/Source/Editor/PropertyEditor/Private/DetailCategoryBuilderImpl.h
--- a/Engine/Source/Editor/PropertyEditor/Private/DetailCategoryBuilderImpl.h
+++ b/Engine/Source/Editor/PropertyEditor/Private/DetailCategoryBuilderImpl.h
@@ -176,7 +176,7 @@ public:
 	virtual IDetailPropertyRow* AddExternalStructureProperty(TSharedPtr<FStructOnScope> StructData, FName PropertyName, EPropertyLocation::Type Location = EPropertyLocation::Default, const FAddPropertyParams& Params = FAddPropertyParams()) override;
 	virtual IDetailPropertyRow* AddExternalStructureProperty(TSharedPtr<IStructureDataProvider> StructData, FName PropertyName, EPropertyLocation::Type Location = EPropertyLocation::Default, const FAddPropertyParams& Params = FAddPropertyParams()) override;
 	virtual TArray<TSharedPtr<IPropertyHandle>> AddAllExternalStructureProperties(TSharedRef<FStructOnScope> StructData, EPropertyLocation::Type Location = EPropertyLocation::Default, TArray<IDetailPropertyRow*>* OutPropertiesRow = nullptr) override;
-	virtual TArray<TSharedPtr<IPropertyHandle>> AddAllExternalStructureProperties(TSharedPtr<IStructureDataProvider> StructProvider, EPropertyLocation::Type Location, TArray<IDetailPropertyRow*>* OutPropertiesRow = nullptr) override;
+	virtual TArray<TSharedPtr<IPropertyHandle>> AddAllExternalStructureProperties(TSharedPtr<IStructureDataProvider> StructProvider, EPropertyLocation::Type Location = EPropertyLocation::Default, TArray<IDetailPropertyRow*>* OutPropertiesRow = nullptr) override;
 	virtual bool IsParentLayoutValid() const override { return DetailLayoutBuilder.IsValid(); }
 	virtual IDetailLayoutBuilder& GetParentLayout() const override { return *DetailLayoutBuilder.Pin(); }
 	virtual FDetailWidgetRow& AddCustomRow(const FText& FilterString, bool bForAdvanced = false) override;
diff --git a/Engine/Source/Editor/PropertyEditor/Public/IDetailChildrenBuilder.h b/Engine/Source/Editor/PropertyEditor/Public/IDetailChildrenBuilder.h
--- a/Engine/Source/Editor/PropertyEditor/Public/IDetailChildrenBuilder.h
+++ b/Engine/Source/Editor/PropertyEditor/Public/IDetailChildrenBuilder.h
@@ -12,6 +12,7 @@ class IDetailCustomNodeBuilder;
 class IDetailGroup;
 class IDetailPropertyRow;
 class IPropertyHandle;
+class IStructureDataProvider;
 class SWidget;
 
 /**
@@ -110,6 +111,8 @@ public:
 	 */
 	virtual TArray<TSharedPtr<IPropertyHandle>> AddAllExternalStructureProperties(TSharedRef<FStructOnScope> ChildStructure) = 0;
 
+	virtual TArray<TSharedPtr<IPropertyHandle>> AddAllExternalStructureProperties(TSharedPtr<IStructureDataProvider> ChildStructure) = 0;
+	
 	/**
 	 * Generates a value widget from a customized struct
 	 * If the customized struct has no value widget an empty widget will be returned
