[PropertyEditor] Fixed crash when clearing a soft object reference property with FEditorPathHelper::IsEnabled()==true

diff --git a/Engine/Source/Editor/PropertyEditor/Private/UserInterface/PropertyEditor/SPropertyEditorAsset.cpp b/Engine/Source/Editor/PropertyEditor/Private/UserInterface/PropertyEditor/SPropertyEditorAsset.cpp
--- a/Engine/Source/Editor/PropertyEditor/Private/UserInterface/PropertyEditor/SPropertyEditorAsset.cpp
+++ b/Engine/Source/Editor/PropertyEditor/Private/UserInterface/PropertyEditor/SPropertyEditorAsset.cpp
@@ -1246,7 +1246,7 @@ SPropertyEditorAsset::FObjectOrAssetData::FObjectOrAssetData(UObject* InObject,
 
 void SPropertyEditorAsset::OnActorSelected( AActor* InActor )
 {
-	if (FEditorPathHelper::IsEnabled() && bIsSoftObjectPath)
+	if (InActor && FEditorPathHelper::IsEnabled() && bIsSoftObjectPath)
 	{
 		// Even if SetValue ends up calling FSoftObjectProperty::ImportText_Internal the FAssetData validation needs to validate the reference domain which is /Temp when referencing Level Instance objects. So we convert the FAssetData to the EditorPath version to pass validation.
 		FSoftObjectPath EditorPath = FEditorPathHelper::GetEditorPathFromEditorPathOwner(InActor, EditorPathOwner);
