Property Editor: Added property type customization for instanced UObjects.

diff --git a/Engine/Source/Editor/PropertyEditor/Private/PropertyEditorModule.cpp b/Engine/Source/Editor/PropertyEditor/Private/PropertyEditorModule.cpp
--- a/Engine/Source/Editor/PropertyEditor/Private/PropertyEditorModule.cpp
+++ b/Engine/Source/Editor/PropertyEditor/Private/PropertyEditorModule.cpp
@@ -986,7 +986,21 @@ FPropertyTypeLayoutCallback FPropertyEditorModule::GetPropertyTypeCustomization(
 		}
 		else if ( bObjectProperty )
 		{
-			UClass* PropertyClass = ObjectProperty->PropertyClass;
+			// Try to find the PropertyClass as common base class of the current instances if possible.
+			// That way we show the selected class' customization, not the based class customization.
+			// If no instances are present, use the base class from the property. 
+			UClass* InstanceBaseClass = nullptr;
+			PropertyHandle.EnumerateConstRawData([&InstanceBaseClass, ObjectProperty](const void* RawData, const int32 /*DataIndex*/, const int32 /*NumDatas*/)
+			{
+				if (const UObject* Object = ObjectProperty->GetObjectPropertyValue(RawData))
+				{
+					UClass* Class = Object->GetClass();
+					InstanceBaseClass = InstanceBaseClass ? UClass::FindCommonBase(InstanceBaseClass, Class) : Class;
+				}
+				return true;
+			});
+			const UClass* PropertyClass = InstanceBaseClass ? InstanceBaseClass : ObjectProperty->PropertyClass.Get();
+			
 			while (PropertyClass)
 			{
 				const FPropertyTypeLayoutCallback& Callback = FindPropertyTypeLayoutCallback(PropertyClass->GetFName(), PropertyHandle, InstancedPropertyTypeLayoutMap);
