Add documentation for IPropertyTypeCustomization

diff --git a/Engine/Source/Editor/PropertyEditor/Public/IPropertyTypeCustomization.h b/Engine/Source/Editor/PropertyEditor/Public/IPropertyTypeCustomization.h
--- a/Engine/Source/Editor/PropertyEditor/Public/IPropertyTypeCustomization.h
+++ b/Engine/Source/Editor/PropertyEditor/Public/IPropertyTypeCustomization.h
@@ -14,7 +14,15 @@ class IPropertyTypeCustomizationUtils;
 class IPropertyUtilities;
 
 /**
- * Base class for property type customizations
+ * Base class for property type customizations. The property type customizations are used for Struct and instanced Object properties.
+ *
+ * The CustomizeHeader() is used to customize the property row, and the CustomizeChildren() is called to add property rows under the header row.
+ * If CustomizeHeader() does not populate the row, the child rows will be added inline at the level where the header row would have gone.
+ *
+ * Note for Object properties:
+ *  - The customization is looked up based on the instanced Object type, or if multiple objects are edited, the common base class.
+ *  - The PropertyHandle points to the property value (object pointer), the first child of PropertyHandle is the instanced object,
+ *    and the instanced object properties are child of that.
  */
 class IPropertyTypeCustomization
 	: public TSharedFromThis<IPropertyTypeCustomization>
