Slate: Add ability to customize details for 'edit-inline' instanced UObjects. This is locked behind 'AllowEditInlineCustomization' property metadata.

diff --git a/Engine/Source/Editor/PropertyEditor/Private/DetailPropertyRow.cpp b/Engine/Source/Editor/PropertyEditor/Private/DetailPropertyRow.cpp
--- a/Engine/Source/Editor/PropertyEditor/Private/DetailPropertyRow.cpp
+++ b/Engine/Source/Editor/PropertyEditor/Private/DetailPropertyRow.cpp
@@ -75,6 +75,7 @@ FDetailPropertyRow::FDetailPropertyRow(TSharedPtr<FPropertyNode> InPropertyNode,
 			MakePropertyEditor(PropertyNodeRef, Utilities, PropertyEditor);
 		}
 		
+		static FName InlineCustomizationKeyMeta("AllowEditInlineCustomization");
 		if (PropertyNode->AsComplexNode() && ExternalRootNode.IsValid()) // AsComplexNode works both for objects and structs
 		{
 			// We are showing an entirely different object inline.  Generate a layout for it now.
@@ -84,6 +85,22 @@ FDetailPropertyRow::FDetailPropertyRow(TSharedPtr<FPropertyNode> InPropertyNode,
 				DetailsView->UpdateSinglePropertyMap(InExternalRootNode, *ExternalObjectLayout, true);
 			}
 		}
+		else if (PropertyNode->HasNodeFlags(EPropertyNodeFlags::EditInlineNew) && PropertyNode->GetProperty()->HasMetaData(InlineCustomizationKeyMeta))
+		{
+			// Allow customization of 'edit inline new' objects if the metadata key has been specified.
+			// The child of this node, if set, will be an object node that we will want to treat as an 'external object layout'
+			TSharedPtr<FPropertyNode> ChildNode = PropertyNode->GetNumChildNodes() > 0 ? PropertyNode->GetChildNode(0) : nullptr;
+			TSharedPtr<FComplexPropertyNode> ComplexChildNode = StaticCastSharedPtr<FComplexPropertyNode>(ChildNode);
+			if (ComplexChildNode.IsValid())
+			{
+				// We are showing an entirely different object inline.  Generate a layout for it now.
+				if (IDetailsViewPrivate* DetailsView = InParentCategory->GetDetailsView())
+				{
+					ExternalObjectLayout = MakeShared<FDetailLayoutData>();
+					DetailsView->UpdateSinglePropertyMap(ComplexChildNode, *ExternalObjectLayout, true);
+				}
+			}
+		}
 
 		if (PropertyNode->GetPropertyKeyNode().IsValid())
 		{							
diff --git a/Engine/Source/Runtime/CoreUObject/Public/UObject/ObjectMacros.h b/Engine/Source/Runtime/CoreUObject/Public/UObject/ObjectMacros.h
--- a/Engine/Source/Runtime/CoreUObject/Public/UObject/ObjectMacros.h
+++ b/Engine/Source/Runtime/CoreUObject/Public/UObject/ObjectMacros.h
@@ -1518,7 +1518,10 @@ namespace UM
 		///
 		/// UFUNCTION()
 		/// TArray<UClass*> FuncName() const;
-		GetDisallowedClasses
+		GetDisallowedClasses,
+
+		/// [PropertyMetadata] Used for edit inline UObject* properties, enables detail class customization for instanced UObject properties.
+		AllowEditInlineCustomization,
 	};
 
 	// Metadata usable in UPROPERTY for customizing the behavior of Persona and UMG
