PropertyEditor: Expose DisplayUseSelected from the Actor Picker through the PropertyEditorAsset for when there's an instanced actor inside.

diff --git a/Engine/Source/Editor/PropertyEditor/Private/PropertyCustomizationHelpers.cpp b/Engine/Source/Editor/PropertyEditor/Private/PropertyCustomizationHelpers.cpp
--- a/Engine/Source/Editor/PropertyEditor/Private/PropertyCustomizationHelpers.cpp
+++ b/Engine/Source/Editor/PropertyEditor/Private/PropertyCustomizationHelpers.cpp
@@ -359,7 +359,7 @@ namespace PropertyCustomizationHelpers
 			.OnUseSelected(OnUseSelected);
 	}
 
-	TSharedRef<SWidget> MakeActorPickerWithMenu(AActor* const InitialActor, const bool AllowClear, const bool AllowPickingLevelInstanceContent, FOnShouldFilterActor ActorFilter, FOnActorSelected OnSet, FSimpleDelegate OnClose, FSimpleDelegate OnUseSelected)
+	TSharedRef<SWidget> MakeActorPickerWithMenu(AActor* const InitialActor, const bool AllowClear, const bool AllowPickingLevelInstanceContent, FOnShouldFilterActor ActorFilter, FOnActorSelected OnSet, FSimpleDelegate OnClose, FSimpleDelegate OnUseSelected, bool bDisplayUseSelected)
 	{
 		return
 			SNew(SPropertyMenuActorPicker)
@@ -369,7 +369,8 @@ namespace PropertyCustomizationHelpers
 			.ActorFilter(ActorFilter)
 			.OnSet(OnSet)
 			.OnClose(OnClose)
-			.OnUseSelected(OnUseSelected);
+			.OnUseSelected(OnUseSelected)
+			.DisplayUseSelected(bDisplayUseSelected);
 	}
 
 	TSharedRef<SWidget> MakeComponentPickerWithMenu( UActorComponent* const InitialComponent, const bool AllowClear, FOnShouldFilterActor ActorFilter, FOnShouldFilterComponent ComponentFilter, FOnComponentSelected OnSet, FSimpleDelegate OnClose )
@@ -384,14 +385,14 @@ namespace PropertyCustomizationHelpers
 			.OnClose(OnClose);
 	}
 
-	TSharedRef<SWidget> MakeInteractiveActorPicker( FOnGetAllowedClasses OnGetAllowedClasses, FOnShouldFilterActor OnShouldFilterActor, FOnActorSelected OnActorSelectedFromPicker )
+	TSharedRef<SWidget> MakeInteractiveActorPicker( FOnGetAllowedClasses OnGetAllowedClasses, FOnShouldFilterActor OnShouldFilterActor, FOnActorSelected OnActorSelectedFromPicker)
 	{
-		return 
-			SNew( SPropertyEditorInteractiveActorPicker )
-			.ToolTipText( LOCTEXT( "PickButtonLabel", "Pick Actor from scene") )
-			.OnGetAllowedClasses( OnGetAllowedClasses )
-			.OnShouldFilterActor( OnShouldFilterActor )
-			.OnActorSelected( OnActorSelectedFromPicker );
+		return
+			SNew(SPropertyEditorInteractiveActorPicker)
+			.ToolTipText(LOCTEXT("PickButtonLabel", "Pick Actor from scene"))
+			.OnGetAllowedClasses(OnGetAllowedClasses)
+			.OnShouldFilterActor(OnShouldFilterActor)
+			.OnActorSelected(OnActorSelectedFromPicker);
 	}
 
 	TSharedRef<SWidget> MakeSceneDepthPicker(FOnSceneDepthLocationSelected OnSceneDepthLocationSelected)
diff --git a/Engine/Source/Editor/PropertyEditor/Private/UserInterface/PropertyEditor/SPropertyEditorAsset.cpp b/Engine/Source/Editor/PropertyEditor/Private/UserInterface/PropertyEditor/SPropertyEditorAsset.cpp
--- a/Engine/Source/Editor/PropertyEditor/Private/UserInterface/PropertyEditor/SPropertyEditorAsset.cpp
+++ b/Engine/Source/Editor/PropertyEditor/Private/UserInterface/PropertyEditor/SPropertyEditorAsset.cpp
@@ -319,6 +319,7 @@ void SPropertyEditorAsset::Construct(const FArguments& InArgs, const TSharedPtr<
 	OnSetObject = InArgs._OnSetObject;
 	OnShouldFilterActor = InArgs._OnShouldFilterActor;
 	ObjectPath = InArgs._ObjectPath;
+	bDisplayUseSelected = InArgs._DisplayUseSelected;
 
 	// Override this as we stole the value to use as OnIsEnabled for the inner widgets
 	SetEnabled(true);
@@ -675,7 +676,7 @@ void SPropertyEditorAsset::Construct(const FArguments& InArgs, const TSharedPtr<
 		];
 	}
 
-	if( InArgs._DisplayUseSelected )
+	if( bDisplayUseSelected )
 	{
 		ButtonBox->AddSlot()
 		.VAlign(VAlign_Center)
@@ -863,7 +864,8 @@ TSharedRef<SWidget> SPropertyEditorAsset::OnGetMenuContent()
 																	 FOnShouldFilterActor::CreateSP( this, &SPropertyEditorAsset::IsFilteredActor ),
 																	 FOnActorSelected::CreateSP( this, &SPropertyEditorAsset::OnActorSelected),
 																	 FSimpleDelegate::CreateSP( this, &SPropertyEditorAsset::CloseComboButton ),
-																	 FSimpleDelegate::CreateSP( this, &SPropertyEditorAsset::OnUse ) );
+																	 FSimpleDelegate::CreateSP( this, &SPropertyEditorAsset::OnUse), 
+																	 bDisplayUseSelected);
 	}
 	else
 	{
diff --git a/Engine/Source/Editor/PropertyEditor/Private/UserInterface/PropertyEditor/SPropertyEditorAsset.h b/Engine/Source/Editor/PropertyEditor/Private/UserInterface/PropertyEditor/SPropertyEditorAsset.h
--- a/Engine/Source/Editor/PropertyEditor/Private/UserInterface/PropertyEditor/SPropertyEditorAsset.h
+++ b/Engine/Source/Editor/PropertyEditor/Private/UserInterface/PropertyEditor/SPropertyEditorAsset.h
@@ -438,4 +438,7 @@ private:
 	 * The system will test first the PropertyHandle follow by the PropertyEditor and if nothing this array of assets
 	 */
 	TArray<FAssetData> OwnerAssetDataArray;
+
+	/*Whether to show the 'Use Selected' item*/
+	bool bDisplayUseSelected;
 };
diff --git a/Engine/Source/Editor/PropertyEditor/Private/UserInterface/PropertyEditor/SPropertyMenuActorPicker.cpp b/Engine/Source/Editor/PropertyEditor/Private/UserInterface/PropertyEditor/SPropertyMenuActorPicker.cpp
--- a/Engine/Source/Editor/PropertyEditor/Private/UserInterface/PropertyEditor/SPropertyMenuActorPicker.cpp
+++ b/Engine/Source/Editor/PropertyEditor/Private/UserInterface/PropertyEditor/SPropertyMenuActorPicker.cpp
@@ -31,16 +31,20 @@ void SPropertyMenuActorPicker::Construct( const FArguments& InArgs )
 	OnSet = InArgs._OnSet;
 	OnClose = InArgs._OnClose;
 	OnUseSelected = InArgs._OnUseSelected;
+	bDisplayUseSelected = InArgs._DisplayUseSelected;
 
 	FMenuBuilder MenuBuilder(true, NULL);
 
 	MenuBuilder.BeginSection(NAME_None, LOCTEXT("CurrentActorOperationsHeader", "Current Actor"));
 	{
-		MenuBuilder.AddMenuEntry(
-			LOCTEXT("UseSelected", "Use Selected"), 
-			LOCTEXT("UseSelected_Tooltip", "Use the currently selected Actor"),
-			FSlateIcon(),
-			FUIAction( FExecuteAction::CreateSP( this, &SPropertyMenuActorPicker::HandleUseSelected ) ) );
+		if (bDisplayUseSelected)
+		{
+			MenuBuilder.AddMenuEntry(
+				LOCTEXT("UseSelected", "Use Selected"),
+				LOCTEXT("UseSelected_Tooltip", "Use the currently selected Actor"),
+				FSlateIcon(),
+				FUIAction(FExecuteAction::CreateSP(this, &SPropertyMenuActorPicker::HandleUseSelected)));
+		}
 
 		if( CurrentActor )
 		{
diff --git a/Engine/Source/Editor/PropertyEditor/Private/UserInterface/PropertyEditor/SPropertyMenuActorPicker.h b/Engine/Source/Editor/PropertyEditor/Private/UserInterface/PropertyEditor/SPropertyMenuActorPicker.h
--- a/Engine/Source/Editor/PropertyEditor/Private/UserInterface/PropertyEditor/SPropertyMenuActorPicker.h
+++ b/Engine/Source/Editor/PropertyEditor/Private/UserInterface/PropertyEditor/SPropertyMenuActorPicker.h
@@ -17,11 +17,13 @@ public:
 		, _AllowClear(true)
 		, _AllowPickingLevelInstanceContent(false)
 		, _ActorFilter()
+		, _DisplayUseSelected(true)
 	{}
 		SLATE_ARGUMENT( AActor*, InitialActor )
 		SLATE_ARGUMENT( bool, AllowClear )
 		SLATE_ARGUMENT( bool, AllowPickingLevelInstanceContent )
 		SLATE_ARGUMENT( FOnShouldFilterActor, ActorFilter )
+		SLATE_ARGUMENT(bool, DisplayUseSelected)
 		SLATE_EVENT( FOnActorSelected, OnSet )
 		SLATE_EVENT( FSimpleDelegate, OnClose )
 		SLATE_EVENT( FSimpleDelegate, OnUseSelected )
@@ -98,4 +100,7 @@ private:
 
 	/** Delegate for using the currently selected actor */
 	FSimpleDelegate OnUseSelected;
+
+	/* Whether to display the 'Use Selected' item*/
+	bool bDisplayUseSelected;
 };
diff --git a/Engine/Source/Editor/PropertyEditor/Public/PropertyCustomizationHelpers.h b/Engine/Source/Editor/PropertyEditor/Public/PropertyCustomizationHelpers.h
--- a/Engine/Source/Editor/PropertyEditor/Public/PropertyCustomizationHelpers.h
+++ b/Engine/Source/Editor/PropertyEditor/Public/PropertyCustomizationHelpers.h
@@ -109,7 +109,7 @@ namespace PropertyCustomizationHelpers
 	PROPERTYEDITOR_API TSharedRef<SWidget> MakeAssetPickerWithMenu( const FAssetData& InitialObject, const bool AllowClear, const bool AllowCopyPaste, const TArray<const UClass*>& AllowedClasses, const TArray<const UClass*>& DisallowedClasses, const TArray<UFactory*>& NewAssetFactories, FOnShouldFilterAsset OnShouldFilterAsset, FOnAssetSelected OnSet, FSimpleDelegate OnClose, const TSharedPtr<IPropertyHandle>& PropertyHandle = TSharedPtr<IPropertyHandle>(), const TArray<FAssetData>& OwnerAssetArray = TArray<FAssetData>());
 	PROPERTYEDITOR_API TSharedRef<SWidget> MakeActorPickerAnchorButton( FOnGetActorFilters OnGetActorFilters, FOnActorSelected OnActorSelectedFromPicker );
 	PROPERTYEDITOR_API TSharedRef<SWidget> MakeActorPickerWithMenu( AActor* const InitialActor, const bool AllowClear, FOnShouldFilterActor ActorFilter, FOnActorSelected OnSet, FSimpleDelegate OnClose, FSimpleDelegate OnUseSelected );
-	PROPERTYEDITOR_API TSharedRef<SWidget> MakeActorPickerWithMenu(AActor* const InitialActor, const bool AllowClear, const bool AllowPickingLevelInstanceContent, FOnShouldFilterActor ActorFilter, FOnActorSelected OnSet, FSimpleDelegate OnClose, FSimpleDelegate OnUseSelected);
+	PROPERTYEDITOR_API TSharedRef<SWidget> MakeActorPickerWithMenu(AActor* const InitialActor, const bool AllowClear, const bool AllowPickingLevelInstanceContent, FOnShouldFilterActor ActorFilter, FOnActorSelected OnSet, FSimpleDelegate OnClose, FSimpleDelegate OnUseSelected, bool bDisplayUseSelected=true);
 	PROPERTYEDITOR_API TSharedRef<SWidget> MakeComponentPickerWithMenu( UActorComponent* const InitialComponent, const bool AllowClear, FOnShouldFilterActor ActorFilter, FOnShouldFilterComponent ComponentFilter, FOnComponentSelected OnSet, FSimpleDelegate OnClose );
 	PROPERTYEDITOR_API TSharedRef<SWidget> MakeInteractiveActorPicker(FOnGetAllowedClasses OnGetAllowedClasses, FOnShouldFilterActor OnShouldFilterActor, FOnActorSelected OnActorSelectedFromPicker);
 	PROPERTYEDITOR_API TSharedRef<SWidget> MakeSceneDepthPicker(FOnSceneDepthLocationSelected OnSceneDepthLocationSelected);
