Pass through PropertyPath to the PropertyUpdatedWidgetBuilder instead of FEditPropertyChain
* PropertyPath allows for property traversal through nested array properties within an object - FEditPropertyChain would also require an FPropertyChangedEvent.
* Deprecated the GetPropertyUpdatedWidget API in favor of one that passes a struct of arguments to ease refactoring in future.
* All added public APIs deprecated preemptively as experimental.  Some refinements here to ensure this works generically for all property editor rows.

diff --git a/Engine/Source/Editor/PropertyEditor/Private/DetailsDisplayManager.cpp b/Engine/Source/Editor/PropertyEditor/Private/DetailsDisplayManager.cpp
--- a/Engine/Source/Editor/PropertyEditor/Private/DetailsDisplayManager.cpp
+++ b/Engine/Source/Editor/PropertyEditor/Private/DetailsDisplayManager.cpp
@@ -16,6 +16,8 @@ static TAutoConsoleVariable<bool> CVarForceShowComponentEditor(
 	TEXT("If true, forces the component editor to show in the main viewport and blueprint details panel for UObjects which normally have it hidden."));
 
 
+FConstructPropertyUpdatedWidgetBuilderArgs::~FConstructPropertyUpdatedWidgetBuilderArgs() = default;
+
 FDetailsDisplayManager::FDetailsDisplayManager(): bIsOuterCategory(false)
 {
 	PrimaryStyleKey = SDetailsView::GetPrimaryDetailsViewStyleKey();
@@ -160,6 +162,16 @@ TSharedPtr<FPropertyUpdatedWidgetBuilder> FDetailsDisplayManager::GetPropertyUpd
 	return nullptr;
 }
 
+bool FDetailsDisplayManager::CanConstructPropertyUpdatedWidgetBuilder() const
+{
+	return false;
+}
+
+TSharedPtr<FPropertyUpdatedWidgetBuilder> FDetailsDisplayManager::ConstructPropertyUpdatedWidgetBuilder(const FConstructPropertyUpdatedWidgetBuilderArgs& Args)
+{
+	return nullptr;
+}
+
 bool FDetailsDisplayManager::GetIsScrollBarNeeded() const
 {
 	return bIsScrollBarNeeded;
diff --git a/Engine/Source/Editor/PropertyEditor/Private/SDetailCategoryTableRow.cpp b/Engine/Source/Editor/PropertyEditor/Private/SDetailCategoryTableRow.cpp
--- a/Engine/Source/Editor/PropertyEditor/Private/SDetailCategoryTableRow.cpp
+++ b/Engine/Source/Editor/PropertyEditor/Private/SDetailCategoryTableRow.cpp
@@ -108,18 +108,31 @@ void SDetailCategoryTableRow::Construct(const FArguments& InArgs, TSharedRef<FDe
 	}
 
 	OwnerTableViewWeak = InOwnerTableView;
-	PropertyUpdatedWidgetBuilder = DisplayManager->GetPropertyUpdatedWidget(FExecuteAction::CreateLambda( [this]
+
+
+
+	PRAGMA_DISABLE_DEPRECATION_WARNINGS
+	if (DisplayManager->CanConstructPropertyUpdatedWidgetBuilder())
+	{
+		FConstructPropertyUpdatedWidgetBuilderArgs Args;
+		Args.ResetToDefaultAction = FExecuteAction::CreateLambda( [this]
 		{
 			if( ResetToDefault.CanExecute())
 			{
 				ResetToDefault.Execute();
 			}
-		}), true, ObjectName );
-	if (PropertyUpdatedWidgetBuilder.IsValid())
-	{
-		TAttribute<bool> IsHovered = TAttribute<bool>::CreateSP( this, &SDetailCategoryTableRow::IsHovered);
-		PropertyUpdatedWidgetBuilder->Bind_IsRowHovered(IsHovered);
+		});
+		Args.PropertyPath = nullptr;
+		Args.CategoryObjectName = ObjectName;
+		
+		PropertyUpdatedWidgetBuilder = DisplayManager->ConstructPropertyUpdatedWidgetBuilder(Args);
+		if (PropertyUpdatedWidgetBuilder.IsValid())
+		{
+			TAttribute<bool> IsHovered = TAttribute<bool>::CreateSP( this, &SDetailCategoryTableRow::IsHovered);
+			PropertyUpdatedWidgetBuilder->Bind_IsRowHovered(IsHovered);
+		}
 	}
+	PRAGMA_ENABLE_DEPRECATION_WARNINGS
 
 	this->ChildSlot
 	[
diff --git a/Engine/Source/Editor/PropertyEditor/Private/SDetailSingleItemRow.cpp b/Engine/Source/Editor/PropertyEditor/Private/SDetailSingleItemRow.cpp
--- a/Engine/Source/Editor/PropertyEditor/Private/SDetailSingleItemRow.cpp
+++ b/Engine/Source/Editor/PropertyEditor/Private/SDetailSingleItemRow.cpp
@@ -12,6 +12,7 @@
 #include "IDetailPropertyExtensionHandler.h"
 #include "Modules/ModuleInterface.h"
 #include "Modules/ModuleManager.h"
+#include "ObjectPropertyNode.h"
 #include "PropertyEditorClipboard.h"
 #include "PropertyEditorClipboardPrivate.h"
 #include "PropertyEditorCopyPastePrivate.h"
@@ -634,18 +635,27 @@ void SDetailSingleItemRow::Construct( const FArguments& InArgs, FDetailLayoutCus
 				DetailsView &&
 				DetailsView->GetDisplayManager().IsValid())
 			{
-				TSharedPtr<FPropertyNode>  PropertyNode = GetPropertyNode(); 	
-				Property = PropertyNode->GetProperty();
-				DisplayManager = DetailsView->GetDisplayManager();
-				TSharedRef<FEditPropertyChain> EditPropertyChain = PropertyNode->BuildPropertyChain( Property ); 
-				PropertyUpdatedWidgetBuilder = DisplayManager->GetPropertyUpdatedWidget(
-					FExecuteAction::CreateSP(this, &SDetailSingleItemRow::OnResetToDefaultClicked), EditPropertyChain, Category->GetObjectName());
-
-				if (PropertyUpdatedWidgetBuilder.IsValid())
+				TSharedPtr<FDetailsDisplayManager> DisplayManagerLocal = DetailsView->GetDisplayManager();
+
+				PRAGMA_DISABLE_DEPRECATION_WARNINGS
+				if (DisplayManagerLocal->CanConstructPropertyUpdatedWidgetBuilder())
 				{
-					TAttribute<bool> IsHovered = TAttribute<bool>::CreateSP( this, &SDetailSingleItemRow::IsHovered);
-					PropertyUpdatedWidgetBuilder->Bind_IsRowHovered(IsHovered);
+					DisplayManager = MoveTemp(DisplayManagerLocal);
+				
+					FConstructPropertyUpdatedWidgetBuilderArgs Args;
+					Args.ResetToDefaultAction = FExecuteAction::CreateSP(this, &SDetailSingleItemRow::OnResetToDefaultClicked);
+					Args.PropertyPath = MakeShared<FPropertyPath>(InOwnerTreeNode->GetPropertyPath());
+					Args.CategoryObjectName = Category->GetObjectName();
+					
+					PropertyUpdatedWidgetBuilder = DisplayManager->ConstructPropertyUpdatedWidgetBuilder(Args);
+
+					if (PropertyUpdatedWidgetBuilder.IsValid())
+					{
+						TAttribute<bool> IsHovered = TAttribute<bool>::CreateSP( this, &SDetailSingleItemRow::IsHovered);
+						PropertyUpdatedWidgetBuilder->Bind_IsRowHovered(IsHovered);
+					}
 				}
+				PRAGMA_ENABLE_DEPRECATION_WARNINGS
 			}
 
 			Splitter->AddSlot()
diff --git a/Engine/Source/Editor/PropertyEditor/Public/DetailsDisplayManager.h b/Engine/Source/Editor/PropertyEditor/Public/DetailsDisplayManager.h
--- a/Engine/Source/Editor/PropertyEditor/Public/DetailsDisplayManager.h
+++ b/Engine/Source/Editor/PropertyEditor/Public/DetailsDisplayManager.h
@@ -2,16 +2,20 @@
 #pragma once
 
 #include "DetailsViewStyleKey.h"
+#include "Framework/Commands/UIAction.h"
 #include "Templates/SharedPointer.h"
 #include "UserInterface/Widgets/PropertyUpdatedWidgetBuilder.h"
 #include "Widgets/SWidget.h"
 
+class FPropertyPath;
 DECLARE_DELEGATE(FOnDetailsNeedsUpdate)
 
 class FDetailsViewStyle;
 class FComplexPropertyNode;
 class FDetailLayoutBuilderImpl;
 
+struct FConstructPropertyUpdatedWidgetBuilderArgs;
+
 /** An @code FDetailsDisplayManager @endcode provides an API to tweak various settings of your details view, and
  * provides some utility methods to work with Details.  */
 class FDetailsDisplayManager : public TSharedFromThis<FDetailsDisplayManager>
@@ -114,7 +118,8 @@ public:
 	* @param bIsCategoryUpdateWidget if true this widget builder is for a Category instead of a property within
 	* @param InCategoryObjectName  the name of the UObject associated with the Category for the widget builder, if one exists, else it is NAME_NONE
 	*/
-	virtual TSharedPtr<FPropertyUpdatedWidgetBuilder> GetPropertyUpdatedWidget(FResetToDefault ResetToDefault, bool bIsCategoryUpdateWidget = false, FName InCategoryObjectName = NAME_None);
+	UE_DEPRECATED(5.5, "Use ConstructPropertyUpdatedWidgetBuilder")
+	PROPERTYEDITOR_API virtual TSharedPtr<FPropertyUpdatedWidgetBuilder> GetPropertyUpdatedWidget(FResetToDefault ResetToDefault, bool bIsCategoryUpdateWidget = false, FName InCategoryObjectName = NAME_None);
 
 	/**
 	* Returns a widget which will show in place of the reset to default button, or a nullptr if the default reset button should be used
@@ -123,7 +128,22 @@ public:
 	* @param InEditorPropertyChain the FEditorPropertyChain for the FPropertyNode whose state is visualized by this  property updated widget 
 	* @param InCategoryObjectName the name of the UObject for which the Category is displayed, if one is associated with the Category
 	*/
-	virtual TSharedPtr<FPropertyUpdatedWidgetBuilder> GetPropertyUpdatedWidget(FResetToDefault ResetToDefault, TSharedRef<FEditPropertyChain> InEditorPropertyChain, FName InCategoryObjectName);
+	UE_DEPRECATED(5.5, "Use ConstructPropertyUpdatedWidgetBuilder")
+	PROPERTYEDITOR_API virtual TSharedPtr<FPropertyUpdatedWidgetBuilder> GetPropertyUpdatedWidget(FResetToDefault ResetToDefault, TSharedRef<FEditPropertyChain> InEditorPropertyChain, FName InCategoryObjectName);
+
+	/**
+	 * If returning true, the row widgets (ie. SDetailSingleItemRow, SDetailCategoryTableRow) will call ConstructPropertyUpdatedWidgetBuilder.
+	 * Enables potentially expensive setup to be elided
+	 */
+	UE_DEPRECATED(5.5, "Experimental API")
+	PROPERTYEDITOR_API virtual bool CanConstructPropertyUpdatedWidgetBuilder() const;
+
+	/**
+	 * Returns a builder class which will generate a widget in place of the Extension widgets on an item row.
+	 * For non-item rows (ie. category), a widget will be placed in the same location
+	 */
+	UE_DEPRECATED(5.5, "Experimental API")
+	PROPERTYEDITOR_API virtual TSharedPtr<FPropertyUpdatedWidgetBuilder> ConstructPropertyUpdatedWidgetBuilder(const FConstructPropertyUpdatedWidgetBuilderArgs& Args);
 
 	void UpdatePropertyForCategory(FName InCategoryObjectName, FProperty* Property, bool bAddProperty);
 
@@ -165,3 +185,11 @@ protected:
 	TMap<FName, TSet<FProperty*>> CategoryNameToUpdatePropertySetMap;
 };
 
+struct FConstructPropertyUpdatedWidgetBuilderArgs
+{
+	~FConstructPropertyUpdatedWidgetBuilderArgs();
+	
+	FExecuteAction ResetToDefaultAction;
+	TSharedPtr<FPropertyPath> PropertyPath;
+	FName CategoryObjectName;
+};
\ No newline at end of file
