Enable edit propagation for changes to sparse class data, explicitly identify the set of objects we want to propagate to

diff --git a/Engine/Source/Editor/PropertyEditor/Private/PropertyHandleImpl.cpp b/Engine/Source/Editor/PropertyEditor/Private/PropertyHandleImpl.cpp
--- a/Engine/Source/Editor/PropertyEditor/Private/PropertyHandleImpl.cpp
+++ b/Engine/Source/Editor/PropertyEditor/Private/PropertyHandleImpl.cpp
@@ -514,9 +514,9 @@ FPropertyAccess::Result FPropertyValueImpl::ImportText( const TArray<FObjectBase
 					(CurObject->HasAnyFlags(RF_DefaultSubObject) && CurObject->GetOuter()->HasAnyFlags(RF_ClassDefaultObject | RF_ArchetypeObject))) &&
 					!bIsGameWorld)
 				{
-					// propagate the changes to instances unless we're modifying class shared data
-					// or the object is using overridable serialization, the propagation is done via reinstantiation
-					if (!bIsSparseClassData && !FOverridableManager::Get().IsEnabled(*CurObject))
+					// propagate the changes to instances unless the object is using overridable serialization, 
+					// the propagation is done via reinstantiation
+					if (!FOverridableManager::Get().IsEnabled(*CurObject))
 					{
 						InPropertyNode->PropagatePropertyChange(CurObject, *NewValue, PreviousContainerValue.IsEmpty() ? PreviousValue : PreviousContainerValue);
 					}
diff --git a/Engine/Source/Editor/PropertyEditor/Private/PropertyNode.cpp b/Engine/Source/Editor/PropertyEditor/Private/PropertyNode.cpp
--- a/Engine/Source/Editor/PropertyEditor/Private/PropertyNode.cpp
+++ b/Engine/Source/Editor/PropertyEditor/Private/PropertyNode.cpp
@@ -3961,7 +3961,25 @@ void FPropertyNode::PropagatePropertyChange( UObject* ModifiedObject, const TCHA
 	FPropertyNode* SubobjectPropertyNode = NULL;
 	UObject* Object = ModifiedObject;
 
-	if (Object->HasAnyFlags(RF_ClassDefaultObject|RF_ArchetypeObject))
+	if (HasNodeFlags(EPropertyNodeFlags::IsSparseClassData))
+	{
+		// Propagate only to child types with the CDO serving as a 'dummy' object to identify
+		// the class (and consequently edit the SCD)
+		if (ensure(Object->HasAnyFlags(RF_ClassDefaultObject)))
+		{
+			TArray<UClass*> Children;
+			GetDerivedClasses(Object->GetClass(), Children);
+			for (UClass* ChildClass : Children)
+			{
+				if (ChildClass->ClassDefaultObject &&
+					!ChildClass->GetPackage()->HasAnyFlags(RF_Transient))
+				{
+					ArchetypeInstances.Add(ChildClass->ClassDefaultObject);
+				}
+			}
+		}
+	}
+	else if (Object->HasAnyFlags(RF_ClassDefaultObject|RF_ArchetypeObject))
 	{
 		// Object is a default subobject, collect all instances.
 		Object->GetArchetypeInstances(ArchetypeInstances);
