SOL-5786 - "Make [REDACTED_HANDLE] work with type and subtype"

- Adds UI support for [REDACTED_HANDLE] `type` except for types with a positive type of 'any' (This could be allowed if desired)
- This relies on adding the meta-class for "VerseClassPositiveType" to FPropertyClass's generated from CTypeType

#test SceneGraph tests

diff --git a/Engine/Source/Editor/PropertyEditor/Private/PropertyHandleImpl.cpp b/Engine/Source/Editor/PropertyEditor/Private/PropertyHandleImpl.cpp
--- a/Engine/Source/Editor/PropertyEditor/Private/PropertyHandleImpl.cpp
+++ b/Engine/Source/Editor/PropertyEditor/Private/PropertyHandleImpl.cpp
@@ -5244,7 +5244,13 @@ FPropertyAccess::Result FPropertyHandleOptional::SetOptionalValue(FProperty* New
 		// If our OptionalValue is a ptr to an object and we are not setting to a passed in value
 		// we need to initialize a default of that object and set the ptr to it.
 		FObjectProperty* ObjectProperty = CastField<FObjectProperty>(OptionalProperty->GetValueProperty());
-		if (ObjectProperty)
+		FClassProperty* ClassProperty = CastField<FClassProperty>(OptionalProperty->GetValueProperty());
+		if (ClassProperty)
+		{
+			void* ClassPropertyValuePtr = ClassProperty->ContainerPtrToValuePtr<void>(Addr);
+			ClassProperty->SetObjectPropertyValue(ClassPropertyValuePtr, ClassProperty->MetaClass);
+		}
+		else if (ObjectProperty)
 		{
 			UObject* Outer = Obj;
 
