Allow single property views to hide their reset to default button

diff --git a/Engine/Source/Editor/PropertyEditor/Private/PropertyEditorModule.cpp b/Engine/Source/Editor/PropertyEditor/Private/PropertyEditorModule.cpp
--- a/Engine/Source/Editor/PropertyEditor/Private/PropertyEditorModule.cpp
+++ b/Engine/Source/Editor/PropertyEditor/Private/PropertyEditorModule.cpp
@@ -324,7 +324,8 @@ TSharedPtr<class ISinglePropertyView> FPropertyEditorModule::CreateSinglePropert
 		.NameOverride(InitParams.NameOverride)
 		.NotifyHook(InitParams.NotifyHook)
 		.PropertyFont(InitParams.Font)
-		.bShouldHideAssetThumbnail(InitParams.bHideAssetThumbnail);
+		.bShouldHideAssetThumbnail(InitParams.bHideAssetThumbnail)
+		.bShouldHideResetToDefault(InitParams.bHideResetToDefault);
 
 	if (Property->HasValidProperty())
 	{
diff --git a/Engine/Source/Editor/PropertyEditor/Private/SSingleProperty.cpp b/Engine/Source/Editor/PropertyEditor/Private/SSingleProperty.cpp
--- a/Engine/Source/Editor/PropertyEditor/Private/SSingleProperty.cpp
+++ b/Engine/Source/Editor/PropertyEditor/Private/SSingleProperty.cpp
@@ -117,6 +117,7 @@ void SSingleProperty::Construct( const FArguments& InArgs )
 	NamePlacement = InArgs._NamePlacement;
 	NotifyHook = InArgs._NotifyHook;
 	PropertyFont = InArgs._PropertyFont;
+	bShouldHideResetToDefault = InArgs._bShouldHideResetToDefault;
 
 	PropertyUtilities = MakeShareable( new FSinglePropertyUtilities( SharedThis( this ), InArgs._bShouldHideAssetThumbnail ) );
 
@@ -270,7 +271,7 @@ bool SSingleProperty::GeneratePropertyCustomization()
 			];			
 		}
 
-		if (!PropertyEditor->GetPropertyHandle()->HasMetaData(TEXT("NoResetToDefault")))
+		if (!PropertyEditor->GetPropertyHandle()->HasMetaData(TEXT("NoResetToDefault")) && !bShouldHideResetToDefault)
 		{
 			HorizontalBox->AddSlot()
 			.Padding( 2.0f )
@@ -298,7 +299,7 @@ bool SSingleProperty::GeneratePropertyCustomization()
 }
 
 
-void SSingleProperty::SetOnPropertyValueChanged( FSimpleDelegate& InOnPropertyValueChanged )
+void SSingleProperty::SetOnPropertyValueChanged( const FSimpleDelegate& InOnPropertyValueChanged )
 {
 	if( HasValidProperty() )
 	{
diff --git a/Engine/Source/Editor/PropertyEditor/Private/SSingleProperty.h b/Engine/Source/Editor/PropertyEditor/Private/SSingleProperty.h
--- a/Engine/Source/Editor/PropertyEditor/Private/SSingleProperty.h
+++ b/Engine/Source/Editor/PropertyEditor/Private/SSingleProperty.h
@@ -25,6 +25,7 @@ public:
 		, _NamePlacement( EPropertyNamePlacement::Left )
 		, _NameOverride()
 		, _bShouldHideAssetThumbnail(false)
+		, _bShouldHideResetToDefault(false)
 	{}
 
 		SLATE_ARGUMENT( UObject*, Object )
@@ -35,6 +36,7 @@ public:
 		SLATE_ARGUMENT( EPropertyNamePlacement::Type, NamePlacement )
 		SLATE_ARGUMENT( FText, NameOverride )
 		SLATE_ARGUMENT( bool, bShouldHideAssetThumbnail )
+		SLATE_ARGUMENT( bool, bShouldHideResetToDefault )
 	SLATE_END_ARGS()	
 
 	void Construct( const FArguments& InArgs );
@@ -43,7 +45,7 @@ public:
 	virtual bool HasValidProperty() const override { return (RootPropertyNode.IsValid() || RootPropertyNode) && ValueNode.IsValid(); }
 	virtual void SetObject( UObject* InObject ) override;
 	virtual void SetStruct( const TSharedPtr<IStructureDataProvider>& InStruct ) override;
-	virtual void SetOnPropertyValueChanged( FSimpleDelegate& InOnPropertyValueChanged ) override;
+	virtual void SetOnPropertyValueChanged( const FSimpleDelegate& InOnPropertyValueChanged ) override;
 
 	/**
 	 * Replaces objects being observed by the view with new objects
@@ -99,4 +101,6 @@ private:
 	/** Location of the name in the view */
 	EPropertyNamePlacement::Type NamePlacement;
 	TSharedPtr<class IPropertyHandle> PropertyHandle;
+	/** Whether the 'reset to default' button should be hidden */
+	bool bShouldHideResetToDefault;
 };
diff --git a/Engine/Source/Editor/PropertyEditor/Public/ISinglePropertyView.h b/Engine/Source/Editor/PropertyEditor/Public/ISinglePropertyView.h
--- a/Engine/Source/Editor/PropertyEditor/Public/ISinglePropertyView.h
+++ b/Engine/Source/Editor/PropertyEditor/Public/ISinglePropertyView.h
@@ -28,13 +28,17 @@ struct FSinglePropertyParams
 
 	/** Whether to hide an asset thumbnail, if available */
 	bool bHideAssetThumbnail;
-		
+
+	/** Whether to hide the 'reset to default' button */
+	bool bHideResetToDefault;
+
 	FSinglePropertyParams()
 		: NameOverride(FText::GetEmpty())
 		, Font()
 		, NotifyHook( NULL )
 		, NamePlacement( EPropertyNamePlacement::Left )
 		, bHideAssetThumbnail( false )
+		, bHideResetToDefault( false )
 	{
 	}
 };
@@ -54,7 +58,7 @@ public:
 	virtual void SetStruct( const TSharedPtr<class IStructureDataProvider>& InStruct) = 0;
 
 	/** Sets a delegate called when the property value changes */
-	virtual void SetOnPropertyValueChanged( FSimpleDelegate& InOnPropertyValueChanged ) = 0;
+	virtual void SetOnPropertyValueChanged( const FSimpleDelegate& InOnPropertyValueChanged ) = 0;
 
 	/** Whether or not this widget has a valid property */	
 	virtual bool HasValidProperty() const = 0;
