Details panel - Fixed perf regression when searching:
- Filter against the filter text directly instead of getting the widget row and filtering against it. This way, we don't create the widgets just to do text filtering.

diff --git a/Engine/Source/Editor/PropertyEditor/Private/DetailCategoryBuilderImpl.cpp b/Engine/Source/Editor/PropertyEditor/Private/DetailCategoryBuilderImpl.cpp
--- a/Engine/Source/Editor/PropertyEditor/Private/DetailCategoryBuilderImpl.cpp
+++ b/Engine/Source/Editor/PropertyEditor/Private/DetailCategoryBuilderImpl.cpp
@@ -144,6 +144,27 @@ TArrayView<TSharedPtr<IPropertyHandle>> FDetailLayoutCustomization::GetPropertyH
 	return TArrayView<TSharedPtr<IPropertyHandle>>();
 }
 
+FText FDetailLayoutCustomization::GetFilterTextString() const
+{
+	if (HasCustomWidget())
+	{
+		return WidgetDecl->FilterTextString;
+	}
+	else if (HasCustomBuilder() && CustomBuilderRow->GetWidgetRow())
+	{
+		return CustomBuilderRow->GetWidgetRow()->FilterTextString;
+	}
+	else if (HasPropertyNode())
+	{
+		return PropertyRow->GetFilterTextString();
+	}
+	else if (DetailGroup->GetHeaderPropertyRow())
+	{
+		return DetailGroup->GetHeaderPropertyRow()->GetFilterTextString();
+	}
+	return {};
+}
+
 const IDetailLayoutRow* FDetailLayoutCustomization::GetDetailLayoutRow() const
 {
 	if (HasCustomWidget())
diff --git a/Engine/Source/Editor/PropertyEditor/Private/DetailCategoryBuilderImpl.h b/Engine/Source/Editor/PropertyEditor/Private/DetailCategoryBuilderImpl.h
--- a/Engine/Source/Editor/PropertyEditor/Private/DetailCategoryBuilderImpl.h
+++ b/Engine/Source/Editor/PropertyEditor/Private/DetailCategoryBuilderImpl.h
@@ -58,6 +58,8 @@ struct FDetailLayoutCustomization
 	FDetailWidgetRow GetWidgetRow() const;
 	/** @return properties being customized */
 	TArrayView<TSharedPtr<IPropertyHandle>> GetPropertyHandles() const;
+	/** @return the filter text associated with this row, if any */
+	FText GetFilterTextString() const;
 	/** Whether or not this customization is considered an advanced property. */
 	bool bAdvanced { false };
 	/** Whether or not this customization is custom or a default one. */
diff --git a/Engine/Source/Editor/PropertyEditor/Private/DetailItemNode.cpp b/Engine/Source/Editor/PropertyEditor/Private/DetailItemNode.cpp
--- a/Engine/Source/Editor/PropertyEditor/Private/DetailItemNode.cpp
+++ b/Engine/Source/Editor/PropertyEditor/Private/DetailItemNode.cpp
@@ -123,9 +123,9 @@ TSharedPtr<IPropertyHandle> FDetailItemNode::CreatePropertyHandle() const
 
 void FDetailItemNode::GetFilterStrings(TArray<FString>& OutFilterStrings) const
 {
-	if (Customization.HasCustomWidget())
+	if (!Customization.GetFilterTextString().IsEmpty())
 	{
-		OutFilterStrings.Add(Customization.GetWidgetRow().FilterTextString.ToString());
+		OutFilterStrings.Add(Customization.GetFilterTextString().ToString());
 	}
 
 	if (Customization.HasPropertyNode())
@@ -138,18 +138,8 @@ void FDetailItemNode::GetFilterStrings(TArray<FString>& OutFilterStrings) const
 			{
 				OutFilterStrings.Add(PropertyNode->GetProperty()->GetName());
 			}
-
-			if (!Customization.GetWidgetRow().FilterTextString.IsEmpty())
-			{
-				OutFilterStrings.Add(Customization.GetWidgetRow().FilterTextString.ToString());
-			}
 		}
 	}
-
-	if (Customization.HasCustomBuilder() && Customization.CustomBuilderRow->GetWidgetRow())
-	{
-		OutFilterStrings.Add(Customization.CustomBuilderRow->GetWidgetRow()->FilterTextString.ToString());
-	}
 }
 
 bool FDetailItemNode::GetInitiallyCollapsed() const
@@ -708,7 +698,7 @@ static bool PassesAllFilters( FDetailItemNode* ItemNode, const FDetailLayoutCust
 			const bool bIsSeenDueToFiltering = PropertyNodePin->HasNodeFlags(EPropertyNodeFlags::IsSeenDueToFiltering) != 0;
 			const bool bIsParentSeenDueToFiltering = PropertyNodePin->HasNodeFlags(EPropertyNodeFlags::IsParentSeenDueToFiltering) != 0;
 
-			const bool bPassesTextFilter = bPassesCategoryFilter || bPassesValueFilter || Local::StringPassesFilter(InFilter, InCustomization.GetWidgetRow().FilterTextString.ToString());
+			const bool bPassesTextFilter = bPassesCategoryFilter || bPassesValueFilter || Local::StringPassesFilter(InFilter, InCustomization.GetFilterTextString().ToString());
 			const bool bPassesSearchFilter = bPassesTextFilter || bSearchFilterIsEmpty || ( bIsNotBeingFiltered || bIsSeenDueToFiltering || bIsParentSeenDueToFiltering );
 
 			bool bPassesModifiedFilter = true;
diff --git a/Engine/Source/Editor/PropertyEditor/Private/DetailPropertyRow.cpp b/Engine/Source/Editor/PropertyEditor/Private/DetailPropertyRow.cpp
--- a/Engine/Source/Editor/PropertyEditor/Private/DetailPropertyRow.cpp
+++ b/Engine/Source/Editor/PropertyEditor/Private/DetailPropertyRow.cpp
@@ -309,6 +309,18 @@ TArrayView<TSharedPtr<IPropertyHandle>> FDetailPropertyRow::GetPropertyHandles()
 	return TArrayView<TSharedPtr<IPropertyHandle>>(const_cast<TSharedPtr<IPropertyHandle>*>(&PropertyHandle), 1);
 }
 
+FText FDetailPropertyRow::GetFilterTextString() const
+{
+	if (CustomPropertyWidget)
+	{
+		return CustomPropertyWidget->FilterTextString;
+	}
+	else
+	{
+		return {};
+	}
+}
+
 static bool IsHeaderRowRequired(const TSharedPtr<IPropertyHandle>& PropertyHandle)
 {
 	TSharedPtr<IPropertyHandle> ParentHandle = PropertyHandle->GetParentHandle();
diff --git a/Engine/Source/Editor/PropertyEditor/Private/DetailPropertyRow.h b/Engine/Source/Editor/PropertyEditor/Private/DetailPropertyRow.h
--- a/Engine/Source/Editor/PropertyEditor/Private/DetailPropertyRow.h
+++ b/Engine/Source/Editor/PropertyEditor/Private/DetailPropertyRow.h
@@ -88,9 +88,12 @@ public:
 	/** @return The widget row that should be displayed for this property row */
 	FDetailWidgetRow GetWidgetRow();
 
-	/** returns only the widget for editing this property. Use GetWidgetRow to get the full row with property name */
+	/** @return properties being customized */
 	TArrayView<TSharedPtr<IPropertyHandle>> GetPropertyHandles() const;
 
+	/** @return the filter text associated with this row, if any */
+	FText GetFilterTextString() const;
+
 	/**
 	 * @return The property node for this row
 	 */
@@ -102,7 +105,7 @@ public:
 	TSharedPtr<FComplexPropertyNode> GetExternalRootNode() const { return ExternalRootNode; }
 
 	/**
-	 * @return The property node for this row
+	 * @return The property editor for this row
 	 */
 	TSharedPtr<FPropertyEditor> GetPropertyEditor() { return PropertyEditor; }
 
