[RegressionFix] type customizations broke when PresentAsType changed to using a more complete type info

diff --git a/Engine/Source/Editor/PropertyEditor/Private/PropertyEditorModule.cpp b/Engine/Source/Editor/PropertyEditor/Private/PropertyEditorModule.cpp
--- a/Engine/Source/Editor/PropertyEditor/Private/PropertyEditorModule.cpp
+++ b/Engine/Source/Editor/PropertyEditor/Private/PropertyEditorModule.cpp
@@ -1042,13 +1042,18 @@ FPropertyTypeLayoutCallback FPropertyEditorModule::FindPropertyTypeLayoutCallbac
 				static const FName NAME_PresentAsTypeMetadata(TEXT("PresentAsType"));
 				if (const FString* DisplayType = AsStructProperty->Struct->FindMetaData(NAME_PresentAsTypeMetadata))
 				{
-					// try finding DisplayType instead
-					LayoutCallbacks = InstancedPropertyTypeLayoutMap.Find(FName(*DisplayType));
-	
-					if( !LayoutCallbacks )
+					if (UE::FPropertyTypeNameBuilder Type; Type.TryParse(*DisplayType))
 					{
-						LayoutCallbacks = GlobalPropertyTypeToLayoutMap.Find(FName(*DisplayType));
+						const FName DisplayStruct = Type.Build().GetName();
+						// try finding DisplayType instead
+						LayoutCallbacks = InstancedPropertyTypeLayoutMap.Find(DisplayStruct);
+		
+						if( !LayoutCallbacks )
+						{
+							LayoutCallbacks = GlobalPropertyTypeToLayoutMap.Find(DisplayStruct);
+						}
 					}
+					
 				}
 			}
 		}
