Added null check in SDetailSingleItemRow::Construct.

review [REDACTED_HANDLE], [REDACTED_HANDLE]

diff --git a/Engine/Source/Editor/PropertyEditor/Private/SDetailSingleItemRow.cpp b/Engine/Source/Editor/PropertyEditor/Private/SDetailSingleItemRow.cpp
--- a/Engine/Source/Editor/PropertyEditor/Private/SDetailSingleItemRow.cpp
+++ b/Engine/Source/Editor/PropertyEditor/Private/SDetailSingleItemRow.cpp
@@ -307,12 +307,12 @@ void SDetailSingleItemRow::Construct( const FArguments& InArgs, FDetailLayoutCus
 	FOnCanAcceptDrop CanAcceptDropDelegate;
 
 	IDetailsViewPrivate* DetailsView = InOwnerTreeNode->GetDetailsView();
-	FDetailColumnSizeData& ColumnSizeData = DetailsView->GetColumnSizeData();
+	const FDetailColumnSizeData& ColumnSizeData = DetailsView ? DetailsView->GetColumnSizeData() : FDetailColumnSizeData();
 
 	PulseAnimation.AddCurve(0.0f, UE::PropertyEditor::Private::PulseAnimationLength, ECurveEaseFunction::CubicInOut);
 
 	// Play on construction if animation was started from a behavior the re-constructs this widget
-	if (DetailsView->IsNodeAnimating(GetPropertyNode()))
+	if (DetailsView && DetailsView->IsNodeAnimating(GetPropertyNode()))
 	{
 		PulseAnimation.Play(SharedThis(this));
 	}
