Struct Utils: Add Property Type Customization support for Instanced Structs

diff --git a/Engine/Plugins/Experimental/StructUtils/Source/StructUtilsEditor/Private/Customizations/InstancedStructDetails.cpp b/Engine/Plugins/Experimental/StructUtils/Source/StructUtilsEditor/Private/Customizations/InstancedStructDetails.cpp
--- a/Engine/Plugins/Experimental/StructUtils/Source/StructUtilsEditor/Private/Customizations/InstancedStructDetails.cpp
+++ b/Engine/Plugins/Experimental/StructUtils/Source/StructUtilsEditor/Private/Customizations/InstancedStructDetails.cpp
@@ -323,12 +323,48 @@ void FInstancedStructDataDetails::GenerateChildContent(IDetailChildrenBuilder& C
 {
 	// Add the rows for the struct
 	TSharedRef<FInstancedStructProvider> NewStructProvider = MakeShared<FInstancedStructProvider>(StructProperty);
-	
-	TArray<TSharedPtr<IPropertyHandle>> ChildProperties = StructProperty->AddChildStructure(NewStructProvider);
-	for (TSharedPtr<IPropertyHandle> ChildHandle : ChildProperties)
+
+	bool bCustomizedProperty = false;
+
+	const UStruct* BaseStruct = NewStructProvider->GetBaseStructure();
+	if (BaseStruct)
+	{
+		static const FName EditModuleName("PropertyEditor");
+		if (const FPropertyEditorModule* EditModule = FModuleManager::GetModulePtr<FPropertyEditorModule>(EditModuleName))
+		{
+			if (EditModule->IsCustomizedStruct(BaseStruct, FCustomPropertyTypeLayoutMap()))
+			{
+				bCustomizedProperty = true;
+			}
+		}
+	}
+
+	if (bCustomizedProperty)
+	{
+		// Use the struct name instead of the fully-qualified property name
+		const FText Label = BaseStruct->GetDisplayNameText();
+		const FName PropertyName = StructProperty->GetProperty()->GetFName();
+
+		// If the struct has a property customization, then we'll route through AddChildStructure, as it supports
+		// IPropertyTypeCustomization. The other branch is mostly kept as-is for legacy support purposes.
+		IDetailPropertyRow* PropertyRow = ChildBuilder.AddChildStructure(
+			StructProperty.ToSharedRef(), NewStructProvider, PropertyName, Label
+		);
+
+		// Expansion state is not properly persisted for these structures, so let's expanded it by default for now
+		if (PropertyRow)
+		{
+			PropertyRow->ShouldAutoExpand(true);
+		}
+	}
+	else
 	{
-		IDetailPropertyRow& Row = ChildBuilder.AddProperty(ChildHandle.ToSharedRef());
-		OnChildRowAdded(Row);
+		TArray<TSharedPtr<IPropertyHandle>> ChildProperties = StructProperty->AddChildStructure(NewStructProvider);
+		for (TSharedPtr<IPropertyHandle> ChildHandle : ChildProperties)
+		{
+			IDetailPropertyRow& Row = ChildBuilder.AddProperty(ChildHandle.ToSharedRef());
+			OnChildRowAdded(Row);
+		}
 	}
 
 	StructProvider = NewStructProvider;
diff --git a/Engine/Source/Editor/PropertyEditor/Private/CustomChildBuilder.cpp b/Engine/Source/Editor/PropertyEditor/Private/CustomChildBuilder.cpp
--- a/Engine/Source/Editor/PropertyEditor/Private/CustomChildBuilder.cpp
+++ b/Engine/Source/Editor/PropertyEditor/Private/CustomChildBuilder.cpp
@@ -70,11 +70,29 @@ IDetailPropertyRow* FCustomChildrenBuilder::AddExternalStructure(TSharedRef<FStr
 
 IDetailPropertyRow* FCustomChildrenBuilder::AddExternalStructureProperty(TSharedRef<FStructOnScope> ChildStructure, FName PropertyName, const FAddPropertyParams& Params)
 {
-	FDetailLayoutCustomization NewCustomization;
+	return AddExternalStructureProperty<>(ChildStructure, PropertyName, Params);
+}
 
-	TSharedRef<FDetailCategoryImpl> ParentCategoryRef = ParentCategory.Pin().ToSharedRef();
+IDetailPropertyRow* FCustomChildrenBuilder::AddExternalStructure(TSharedPtr<IStructureDataProvider> ChildStructure, FName UniqueIdName)
+{
+	return AddExternalStructureProperty(ChildStructure, NAME_None, FAddPropertyParams().UniqueId(UniqueIdName));
+}
 
-	FDetailPropertyRow::MakeExternalPropertyRowCustomization(ChildStructure, PropertyName, ParentCategoryRef, NewCustomization, Params);
+IDetailPropertyRow* FCustomChildrenBuilder::AddChildStructure(TSharedRef<IPropertyHandle> PropertyHandle, TSharedPtr<IStructureDataProvider> ChildStructure, FName UniqueIdName, const FText& DisplayNameOverride)
+{
+	return AddChildStructureProperty(PropertyHandle, ChildStructure, NAME_None, FAddPropertyParams().UniqueId(UniqueIdName), DisplayNameOverride);
+}
+
+IDetailPropertyRow* FCustomChildrenBuilder::AddExternalStructureProperty(TSharedPtr<IStructureDataProvider> ChildStructure, FName PropertyName, const FAddPropertyParams& Params)
+{
+	return AddExternalStructureProperty<>(ChildStructure, PropertyName, Params);
+}
+
+IDetailPropertyRow* FCustomChildrenBuilder::AddStructureProperty(const FAddPropertyParams& Params, TFunctionRef<void(FDetailLayoutCustomization&)> MakePropertyRowCustomization)
+{
+	FDetailLayoutCustomization NewCustomization;
+
+	MakePropertyRowCustomization(NewCustomization);
 
 	if (Params.ShouldHideRootObjectNode() && NewCustomization.HasPropertyNode() && NewCustomization.GetPropertyNode()->AsComplexNode())
 	{
@@ -96,6 +114,28 @@ IDetailPropertyRow* FCustomChildrenBuilder::AddExternalStructureProperty(TShared
 	return NewRow.Get();
 }
 
+template<class T>
+IDetailPropertyRow* FCustomChildrenBuilder::AddExternalStructureProperty(const T& ChildStructure, FName PropertyName, const FAddPropertyParams& Params)
+{
+	return AddStructureProperty(Params, [&](FDetailLayoutCustomization& NewCustomization)
+	{
+		FDetailPropertyRow::MakeExternalPropertyRowCustomization(
+			ChildStructure, PropertyName, ParentCategory.Pin().ToSharedRef(), NewCustomization, Params
+		);
+	});
+}
+
+IDetailPropertyRow* FCustomChildrenBuilder::AddChildStructureProperty(TSharedRef<IPropertyHandle> PropertyHandle,
+	TSharedPtr<IStructureDataProvider> ChildStructure, FName PropertyName, const FAddPropertyParams& Params, const FText& DisplayNameOverride)
+{
+	return AddStructureProperty(Params, [&](FDetailLayoutCustomization& NewCustomization)
+	{
+		FDetailPropertyRow::MakeChildPropertyRowCustomization(
+			PropertyHandle, ChildStructure, PropertyName, ParentCategory.Pin().ToSharedRef(), NewCustomization, Params, DisplayNameOverride
+		);
+	});
+}
+
 TArray<TSharedPtr<IPropertyHandle>> FCustomChildrenBuilder::AddAllExternalStructureProperties(TSharedRef<FStructOnScope> ChildStructure)
 {
 	const TSharedPtr<FDetailCategoryImpl> ParentCategoryPinned = ParentCategory.Pin();
diff --git a/Engine/Source/Editor/PropertyEditor/Private/CustomChildBuilder.h b/Engine/Source/Editor/PropertyEditor/Private/CustomChildBuilder.h
--- a/Engine/Source/Editor/PropertyEditor/Private/CustomChildBuilder.h
+++ b/Engine/Source/Editor/PropertyEditor/Private/CustomChildBuilder.h
@@ -29,7 +29,11 @@ public:
 	virtual IDetailPropertyRow* AddExternalObjects(const TArray<UObject*>& Objects, const FAddPropertyParams& Params) override;
 	virtual IDetailPropertyRow* AddExternalObjectProperty(const TArray<UObject*>& Objects, FName PropertyName, const FAddPropertyParams& Params) override;
 	virtual IDetailPropertyRow* AddExternalStructure(TSharedRef<FStructOnScope> ChildStructure, FName UniqueIdName) override;
+	virtual IDetailPropertyRow* AddExternalStructure(TSharedPtr<IStructureDataProvider> ChildStructure, FName UniqueIdName) override;
+	virtual IDetailPropertyRow* AddChildStructure(TSharedRef<IPropertyHandle> PropertyHandle, TSharedPtr<IStructureDataProvider> ChildStructure, FName UniqueIdName, const FText& DisplayNameOverride = FText()) override;
 	virtual IDetailPropertyRow* AddExternalStructureProperty(TSharedRef<FStructOnScope> ChildStructure, FName PropertyName, const FAddPropertyParams& Params) override;
+	virtual IDetailPropertyRow* AddExternalStructureProperty(TSharedPtr<IStructureDataProvider> ChildStructure, FName PropertyName, const FAddPropertyParams& Params) override;
+	virtual IDetailPropertyRow* AddChildStructureProperty(TSharedRef<IPropertyHandle> PropertyHandle, TSharedPtr<IStructureDataProvider> ChildStructure, FName PropertyName, const FAddPropertyParams& Params = FAddPropertyParams(), const FText& DisplayNameOverride = FText()) override;
 	virtual TArray<TSharedPtr<IPropertyHandle>> AddAllExternalStructureProperties(TSharedRef<FStructOnScope> ChildStructure) override;
 	virtual TArray<TSharedPtr<IPropertyHandle>> AddAllExternalStructureProperties(TSharedPtr<IStructureDataProvider> ChildStructure) override;
 	virtual TSharedRef<SWidget> GenerateStructValueWidget(TSharedRef<IPropertyHandle> StructPropertyHandle) override;
@@ -48,4 +52,7 @@ private:
 
 	/** User customized reset to default on children */
 	TOptional<FResetToDefaultOverride> CustomResetChildToDefault;
+
+	IDetailPropertyRow* AddStructureProperty(const FAddPropertyParams& Params, TFunctionRef<void(FDetailLayoutCustomization&)> MakePropertyRowCustomization);
+	template<class T> IDetailPropertyRow* AddExternalStructureProperty(const T& ChildStructure, FName PropertyName, const FAddPropertyParams& Params);
 };
diff --git a/Engine/Source/Editor/PropertyEditor/Private/DetailPropertyRow.cpp b/Engine/Source/Editor/PropertyEditor/Private/DetailPropertyRow.cpp
--- a/Engine/Source/Editor/PropertyEditor/Private/DetailPropertyRow.cpp
+++ b/Engine/Source/Editor/PropertyEditor/Private/DetailPropertyRow.cpp
@@ -11,6 +11,7 @@
 #include "ObjectPropertyNode.h"
 #include "PropertyCustomizationHelpers.h"
 #include "PropertyEditorHelpers.h"
+#include "PropertyHandleImpl.h"
 #include "StructurePropertyNode.h"
 #include "Modules/ModuleManager.h"
 #include "Widgets/Layout/SSpacer.h"
@@ -636,12 +637,15 @@ TSharedPtr<IPropertyTypeCustomization> FDetailPropertyRow::GetPropertyCustomizat
 	return CustomInterface;
 }
 
-void FDetailPropertyRow::MakeExternalPropertyRowCustomization(TSharedPtr<FStructOnScope> StructData, FName PropertyName, TSharedRef<FDetailCategoryImpl> ParentCategory, FDetailLayoutCustomization& OutCustomization, const FAddPropertyParams& Parameters)
+template<typename T>
+void MakeExternalStructPropertyRowCustomization(const T& Struct, const UStruct* StructClass,
+	FName PropertyName, TSharedRef<FDetailCategoryImpl> ParentCategory, FDetailLayoutCustomization& OutCustomization,
+	const FAddPropertyParams& Parameters, bool bAllowChildren)
 {
 	TSharedRef<FStructurePropertyNode> RootPropertyNode = MakeShared<FStructurePropertyNode>();
 
 	//SET
-	RootPropertyNode->SetStructure(StructData);
+	RootPropertyNode->SetStructure(Struct);
 
 	FPropertyNodeInitParams InitParams;
 	InitParams.ParentNode = nullptr;
@@ -650,7 +654,7 @@ void FDetailPropertyRow::MakeExternalPropertyRowCustomization(TSharedPtr<FStruct
 	InitParams.ArrayIndex = INDEX_NONE;
 	InitParams.bForceHiddenPropertyVisibility = Parameters.ShouldForcePropertyVisible() || FPropertySettings::Get().ShowHiddenProperties();
 	InitParams.bCreateCategoryNodes = PropertyName == NAME_None;
-	InitParams.bAllowChildren = false;
+	InitParams.bAllowChildren = bAllowChildren;
 
 	Parameters.OverrideAllowChildren(InitParams.bAllowChildren);
 	Parameters.OverrideCreateCategoryNodes(InitParams.bCreateCategoryNodes);
@@ -676,7 +680,7 @@ void FDetailPropertyRow::MakeExternalPropertyRowCustomization(TSharedPtr<FStruct
 		FPropertyEditorModule& PropertyEditorModule = FModuleManager::GetModuleChecked<FPropertyEditorModule>(PropertyEditorModuleName);
 
 		// Make a "fake" struct property to represent the entire struct
-		FStructProperty* StructProperty = PropertyEditorModule.RegisterStructOnScopeProperty(StructData.ToSharedRef());
+		FStructProperty* StructProperty = PropertyEditorModule.RegisterStructProperty(StructClass);
 
 		// Generate a node for the struct
 		TSharedPtr<FItemPropertyNode> ItemNode = MakeShared<FItemPropertyNode>();
@@ -699,45 +703,30 @@ void FDetailPropertyRow::MakeExternalPropertyRowCustomization(TSharedPtr<FStruct
 	}
 }
 
-void FDetailPropertyRow::MakeExternalPropertyRowCustomization(TSharedPtr<IStructureDataProvider> StructDataProvider, FName PropertyName, TSharedRef<FDetailCategoryImpl> ParentCategory, struct FDetailLayoutCustomization& OutCustomization, const FAddPropertyParams& Parameters)
+void FDetailPropertyRow::MakeExternalPropertyRowCustomization(TSharedPtr<FStructOnScope> StructData, FName PropertyName, TSharedRef<FDetailCategoryImpl> ParentCategory, FDetailLayoutCustomization& OutCustomization, const FAddPropertyParams& Parameters)
 {
-	TSharedRef<FStructurePropertyNode> RootPropertyNode = MakeShared<FStructurePropertyNode>();
-
-	//SET
-	RootPropertyNode->SetStructure(StructDataProvider);
-
-	FPropertyNodeInitParams InitParams;
-	InitParams.ParentNode = nullptr;
-	InitParams.Property = nullptr;
-	InitParams.ArrayOffset = 0;
-	InitParams.ArrayIndex = INDEX_NONE;
-	InitParams.bAllowChildren = true;
-	InitParams.bForceHiddenPropertyVisibility = Parameters.ShouldForcePropertyVisible() || FPropertySettings::Get().ShowHiddenProperties();
-	InitParams.bCreateCategoryNodes = PropertyName == NAME_None;
-
-	Parameters.OverrideAllowChildren(InitParams.bAllowChildren);
-	Parameters.OverrideCreateCategoryNodes(InitParams.bCreateCategoryNodes);
-
-	RootPropertyNode->InitNode(InitParams);
-
-	ParentCategory->GetParentLayoutImpl()->AddExternalRootPropertyNode(RootPropertyNode);
-
-	if (PropertyName != NAME_None)
-	{
-		TSharedPtr<FPropertyNode> PropertyNode = RootPropertyNode->GenerateSingleChild(PropertyName);
-		if (PropertyNode.IsValid())
-		{
-			PropertyNode->RebuildChildren();
+	MakeExternalStructPropertyRowCustomization<>(
+		StructData,
+		StructData->GetStruct(),
+		PropertyName,
+		ParentCategory,
+		OutCustomization,
+		Parameters,
+		/* bAllowChildren= */ false
+	);
+}
 
-			OutCustomization.PropertyRow = MakeShared<FDetailPropertyRow>(PropertyNode, ParentCategory, RootPropertyNode);
-			OutCustomization.PropertyRow->SetCustomExpansionId(Parameters.GetUniqueId());
-		}
-	}
-	else
-	{
-		OutCustomization.PropertyRow = MakeShared<FDetailPropertyRow>(RootPropertyNode, ParentCategory, RootPropertyNode);
-		OutCustomization.PropertyRow->SetCustomExpansionId(Parameters.GetUniqueId());
-	}
+void FDetailPropertyRow::MakeExternalPropertyRowCustomization(TSharedPtr<IStructureDataProvider> StructDataProvider, FName PropertyName, TSharedRef<FDetailCategoryImpl> ParentCategory, struct FDetailLayoutCustomization& OutCustomization, const FAddPropertyParams& Parameters)
+{
+	MakeExternalStructPropertyRowCustomization<>(
+		StructDataProvider,
+		StructDataProvider->GetBaseStructure(),
+		PropertyName,
+		ParentCategory,
+		OutCustomization,
+		Parameters,
+		/* bAllowChildren= */ true
+	);
 }
 
 void FDetailPropertyRow::MakeExternalPropertyRowCustomization(const TArray<UObject*>& InObjects, FName PropertyName, TSharedRef<FDetailCategoryImpl> ParentCategory, struct FDetailLayoutCustomization& OutCustomization, const FAddPropertyParams& Parameters)
@@ -796,6 +785,53 @@ void FDetailPropertyRow::MakeExternalPropertyRowCustomization(const TArray<UObje
 	ParentCategory->GetParentLayoutImpl()->AddExternalRootPropertyNode(RootPropertyNode);
 }
 
+void FDetailPropertyRow::MakeChildPropertyRowCustomization(TSharedRef<IPropertyHandle> PropertyHandle,
+	TSharedPtr<IStructureDataProvider> StructDataProvider, FName PropertyName, TSharedRef<FDetailCategoryImpl> ParentCategory,
+	FDetailLayoutCustomization& OutCustomization, const FAddPropertyParams& Parameters, const FText& DisplayNameOverride)
+{
+	TSharedRef<FPropertyHandleBase> PropertyHandleImpl = StaticCastSharedRef<FPropertyHandleBase>(PropertyHandle);
+	TSharedPtr<FStructurePropertyNode> RootPropertyNode = StaticCastSharedPtr<FStructurePropertyNode>(PropertyHandleImpl->GetPropertyNode());
+
+	if (PropertyName != NAME_None)
+	{
+		TSharedPtr<FPropertyNode> PropertyNode = RootPropertyNode->GenerateSingleChild(PropertyName);
+		if (PropertyNode.IsValid())
+		{
+			PropertyNode->RebuildChildren();
+
+			OutCustomization.PropertyRow = MakeShared<FDetailPropertyRow>(PropertyNode, ParentCategory);
+			OutCustomization.PropertyRow->SetCustomExpansionId(Parameters.GetUniqueId());
+		}
+	}
+	else
+	{
+		// Generate a node for the struct
+		TSharedRef<FStructurePropertyNode> StructPropertyNode = MakeShared<FStructurePropertyNode>();
+		StructPropertyNode->SetStructure(StructDataProvider);
+		StructPropertyNode->SetDisplayNameOverride(DisplayNameOverride);
+
+		// Make a "fake" struct property to represent the entire struct
+		static const FName PropertyEditorModuleName("PropertyEditor");
+		FPropertyEditorModule& PropertyEditorModule = FModuleManager::GetModuleChecked<FPropertyEditorModule>(PropertyEditorModuleName);
+		FStructProperty* StructProperty = PropertyEditorModule.RegisterStructProperty(StructDataProvider->GetBaseStructure());
+
+		FPropertyNodeInitParams ItemNodeInitParams;
+		ItemNodeInitParams.ParentNode = RootPropertyNode;
+		ItemNodeInitParams.Property = StructProperty;
+		ItemNodeInitParams.ArrayOffset = 0;
+		ItemNodeInitParams.ArrayIndex = INDEX_NONE;
+		ItemNodeInitParams.bAllowChildren = true;
+		ItemNodeInitParams.bForceHiddenPropertyVisibility = Parameters.ShouldForcePropertyVisible() || FPropertySettings::Get().ShowHiddenProperties();
+		ItemNodeInitParams.bCreateCategoryNodes = false;
+
+		StructPropertyNode->InitNode(ItemNodeInitParams);
+		RootPropertyNode->AddChildNode(StructPropertyNode);
+
+		OutCustomization.PropertyRow = MakeShared<FDetailPropertyRow>(StructPropertyNode, ParentCategory);
+		OutCustomization.PropertyRow->SetCustomExpansionId(Parameters.GetUniqueId());
+	}
+}
+
 EVisibility FDetailPropertyRow::GetPropertyVisibility() const
 {
 	if (IsOnlyVisibleWhenEditConditionMet() && !IsEditConditionMet())
diff --git a/Engine/Source/Editor/PropertyEditor/Private/DetailPropertyRow.h b/Engine/Source/Editor/PropertyEditor/Private/DetailPropertyRow.h
--- a/Engine/Source/Editor/PropertyEditor/Private/DetailPropertyRow.h
+++ b/Engine/Source/Editor/PropertyEditor/Private/DetailPropertyRow.h
@@ -129,6 +129,7 @@ public:
 	static void MakeExternalPropertyRowCustomization(TSharedPtr<FStructOnScope> StructData, FName PropertyName, TSharedRef<FDetailCategoryImpl> ParentCategory, struct FDetailLayoutCustomization& OutCustomization, const FAddPropertyParams& Parameters);
 	static void MakeExternalPropertyRowCustomization(TSharedPtr<IStructureDataProvider> StructDataProvider, FName PropertyName, TSharedRef<FDetailCategoryImpl> ParentCategory, struct FDetailLayoutCustomization& OutCustomization, const FAddPropertyParams& Parameters);
 	static void MakeExternalPropertyRowCustomization(const TArray<UObject*>& InObjects, FName PropertyName, TSharedRef<FDetailCategoryImpl> ParentCategory, struct FDetailLayoutCustomization& OutCustomization, const FAddPropertyParams& Parameters);
+	static void MakeChildPropertyRowCustomization(TSharedRef<IPropertyHandle> PropertyHandle, TSharedPtr<IStructureDataProvider> StructDataProvider, FName PropertyName, TSharedRef<FDetailCategoryImpl> ParentCategory, struct FDetailLayoutCustomization& OutCustomization, const FAddPropertyParams& Parameters, const FText& DisplayNameOverride = FText());
 
 private:
 	/**
diff --git a/Engine/Source/Editor/PropertyEditor/Private/Presentation/PropertyEditor/PropertyEditor.cpp b/Engine/Source/Editor/PropertyEditor/Private/Presentation/PropertyEditor/PropertyEditor.cpp
--- a/Engine/Source/Editor/PropertyEditor/Private/Presentation/PropertyEditor/PropertyEditor.cpp
+++ b/Engine/Source/Editor/PropertyEditor/Private/Presentation/PropertyEditor/PropertyEditor.cpp
@@ -58,6 +58,17 @@ FText FPropertyEditor::GetDisplayName() const
 		return ItemPropertyNode->GetDisplayName();
 	}
 
+	if (const FComplexPropertyNode* ComplexPropertyNode = PropertyNode->AsComplexNode())
+	{
+		const FText DisplayName = ComplexPropertyNode->GetDisplayName();
+
+		// Does this property define its own name?
+		if (!DisplayName.IsEmpty())
+		{
+			return DisplayName;
+		}
+	}
+
 	FString DisplayName;
 	PropertyNode->GetQualifiedName( DisplayName, true );
 	return FText::FromString(DisplayName);
diff --git a/Engine/Source/Editor/PropertyEditor/Private/PropertyEditorModule.cpp b/Engine/Source/Editor/PropertyEditor/Private/PropertyEditorModule.cpp
--- a/Engine/Source/Editor/PropertyEditor/Private/PropertyEditorModule.cpp
+++ b/Engine/Source/Editor/PropertyEditor/Private/PropertyEditorModule.cpp
@@ -33,6 +33,7 @@
 #include "Widgets/Layout/SBorder.h"
 #include "DetailsViewStyle.h"
 #include "ToolMenus.h"
+#include "IStructureDataProvider.h"
 
 IMPLEMENT_MODULE( FPropertyEditorModule, PropertyEditor );
 
@@ -403,9 +404,11 @@ TSharedRef< IPropertyTableCellPresenter > FPropertyEditorModule::CreateTextPrope
 	return MakeShareable( new FTextPropertyTableCellPresenter( PropertyEditor, InPropertyUtilities, InFont, InCell) );
 }
 
-FStructProperty* FPropertyEditorModule::RegisterStructOnScopeProperty(TSharedRef<FStructOnScope> StructOnScope)
+FStructProperty* FPropertyEditorModule::RegisterStructProperty(const UStruct* StructClass)
 {
-	const FName StructName = StructOnScope->GetStruct()->GetFName();
+	check(StructClass);
+
+	const FName StructName = StructClass->GetFName();
 	FStructProperty* StructProperty = RegisteredStructToProxyMap.FindRef(StructName);
 
 	if(!StructProperty)
@@ -418,10 +421,10 @@ FStructProperty* FPropertyEditorModule::RegisterStructOnScopeProperty(TSharedRef
 			StructOnScopePropertyOwner = NewObject<UStruct>(GetTransientPackage(), TEXT("StructOnScope"), RF_Transient);
 			StructOnScopePropertyOwner->AddToRoot();
 		}
-		UScriptStruct* InnerStruct = Cast<UScriptStruct>(const_cast<UStruct*>(StructOnScope->GetStruct()));
+		UScriptStruct* InnerStruct = CastChecked<UScriptStruct>(const_cast<UStruct*>(StructClass));
 		StructProperty = new FStructProperty(StructOnScopePropertyOwner, *MakeUniqueObjectName(StructOnScopePropertyOwner, UField::StaticClass(), InnerStruct->GetFName()).ToString(), RF_Transient);
 		StructProperty->Struct = InnerStruct;
-		StructProperty->ElementSize = StructOnScope->GetStruct()->GetStructureSize();
+		StructProperty->ElementSize = StructClass->GetStructureSize();
 		StructOnScopePropertyOwner->AddCppProperty(StructProperty);
 
 		RegisteredStructToProxyMap.Add(StructName, StructProperty);
diff --git a/Engine/Source/Editor/PropertyEditor/Private/PropertyNode.cpp b/Engine/Source/Editor/PropertyEditor/Private/PropertyNode.cpp
--- a/Engine/Source/Editor/PropertyEditor/Private/PropertyNode.cpp
+++ b/Engine/Source/Editor/PropertyEditor/Private/PropertyNode.cpp
@@ -4169,4 +4169,19 @@ bool FPropertyNode::GenerateRestrictionToolTip(const FString& Value, FText& OutT
 	return bRestricted;
 }
 
+void FComplexPropertyNode::SetDisplayNameOverride(const FText& InDisplayNameOverride)
+{
+	DisplayNameOverride = InDisplayNameOverride;
+}
+
+FText FComplexPropertyNode::GetDisplayName() const
+{
+	if (!DisplayNameOverride.IsEmpty())
+	{
+		return DisplayNameOverride;
+	}
+
+	return FPropertyNode::GetDisplayName();
+}
+
 #undef LOCTEXT_NAMESPACE
diff --git a/Engine/Source/Editor/PropertyEditor/Private/PropertyNode.h b/Engine/Source/Editor/PropertyEditor/Private/PropertyNode.h
--- a/Engine/Source/Editor/PropertyEditor/Private/PropertyNode.h
+++ b/Engine/Source/Editor/PropertyEditor/Private/PropertyNode.h
@@ -1348,6 +1348,9 @@ public:
 	virtual FStructurePropertyNode* AsStructureNode() { return nullptr; }
 	virtual const FStructurePropertyNode* AsStructureNode() const { return nullptr; }
 
+	virtual void SetDisplayNameOverride(const FText& InDisplayNameOverride) override;
+	virtual FText GetDisplayName() const override;
+
 	virtual UStruct* GetBaseStructure() = 0;
 	virtual const UStruct* GetBaseStructure() const = 0;
 
@@ -1369,4 +1372,8 @@ public:
 
 	/** Generates a single child from the provided property name.  Any existing children are destroyed */
 	virtual TSharedPtr<FPropertyNode> GenerateSingleChild(FName ChildPropertyName) = 0;
+
+private:
+	/** Display name to use instead of the fully qualified name */
+	FText DisplayNameOverride;
 };
diff --git a/Engine/Source/Editor/PropertyEditor/Public/IDetailChildrenBuilder.h b/Engine/Source/Editor/PropertyEditor/Public/IDetailChildrenBuilder.h
--- a/Engine/Source/Editor/PropertyEditor/Public/IDetailChildrenBuilder.h
+++ b/Engine/Source/Editor/PropertyEditor/Public/IDetailChildrenBuilder.h
@@ -91,9 +91,11 @@ public:
 	 *
 	 * @param ChildStructure	The structure to add
 	 * @param PropertyName		Optional name of a property inside the Child structure to add.  If this is empty, the entire structure will be added
-	 * @param UniqueIdName		Optional identifier that uniquely identifies this structure among other structures of the same type.  If this is empty, saving and restoring expansion state of this structure may not work
+	 * @param Params		    Optional identifier that uniquely identifies this structure among other structures of the same type.  If this is empty, saving and restoring expansion state of this structure may not work
 	 */
 	virtual IDetailPropertyRow* AddExternalStructureProperty(TSharedRef<FStructOnScope> ChildStructure, FName PropertyName, const FAddPropertyParams& Params = FAddPropertyParams()) = 0;
+	virtual IDetailPropertyRow* AddExternalStructureProperty(TSharedPtr<IStructureDataProvider> ChildStructure, FName PropertyName, const FAddPropertyParams& Params = FAddPropertyParams()) = 0;
+	virtual IDetailPropertyRow* AddChildStructureProperty(TSharedRef<IPropertyHandle> PropertyHandle, TSharedPtr<IStructureDataProvider> ChildStructure, FName PropertyName, const FAddPropertyParams& Params = FAddPropertyParams(), const FText& DisplayNameOverride = FText()) = 0;
 
 	/**
 	 * Adds a custom structure as a child
@@ -102,6 +104,8 @@ public:
 	 * @param UniqueIdName		Optional identifier that uniquely identifies this structure among other structures of the same type.  If this is empty, saving and restoring expansion state of this structure may not work
 	 */
 	virtual IDetailPropertyRow* AddExternalStructure(TSharedRef<FStructOnScope> ChildStructure, FName UniqueIdName = NAME_None) = 0;
+	virtual IDetailPropertyRow* AddExternalStructure(TSharedPtr<IStructureDataProvider> ChildStructure, FName UniqueIdName = NAME_None) = 0;
+	virtual IDetailPropertyRow* AddChildStructure(TSharedRef<IPropertyHandle> PropertyHandle, TSharedPtr<IStructureDataProvider> ChildStructure, FName UniqueIdName = NAME_None, const FText& DisplayNameOverride = FText()) = 0;
 
 	/**
 	 * Adds all the properties of an external structure as a children
diff --git a/Engine/Source/Editor/PropertyEditor/Public/PropertyEditorModule.h b/Engine/Source/Editor/PropertyEditor/Public/PropertyEditorModule.h
--- a/Engine/Source/Editor/PropertyEditor/Public/PropertyEditorModule.h
+++ b/Engine/Source/Editor/PropertyEditor/Public/PropertyEditorModule.h
@@ -402,12 +402,12 @@ public:
 		const FSlateFontInfo* InFontPtr = NULL, const TSharedPtr< IPropertyTableCell >& InCell = nullptr);
 
 	/**
-	 * Register a floating struct on scope so that the details panel may use it as a property
+	 * Register a floating struct so that the details panel may use it as a property
 	 *
-	 * @param StructOnScope		The struct to register
+	 * @param StructClass		The struct to register
 	 * @return The struct property that may may be associated with the details panel
  	 */
-	virtual FStructProperty* RegisterStructOnScopeProperty(TSharedRef<FStructOnScope> StructOnScope);
+	virtual FStructProperty* RegisterStructProperty(const UStruct* StructClass);
 
 	/**
 	 *
@@ -418,7 +418,7 @@ public:
 
 	FPropertyTypeLayoutCallback GetPropertyTypeCustomization(const FProperty* InProperty,const IPropertyHandle& PropertyHandle, const FCustomPropertyTypeLayoutMap& InstancedPropertyTypeLayoutMap);
 	FPropertyTypeLayoutCallback FindPropertyTypeLayoutCallback(FName PropertyTypeName, const IPropertyHandle& PropertyHandle, const FCustomPropertyTypeLayoutMap& InstancedPropertyTypeLayoutMapp);
-	bool IsCustomizedStruct(const UStruct* Struct, const FCustomPropertyTypeLayoutMap& InstancePropertyTypeLayoutMap) const;
+	PROPERTYEDITOR_API bool IsCustomizedStruct(const UStruct* Struct, const FCustomPropertyTypeLayoutMap& InstancePropertyTypeLayoutMap) const;
 
 	DECLARE_EVENT(PropertyEditorModule, FPropertyEditorOpenedEvent);
 	virtual FPropertyEditorOpenedEvent& OnPropertyEditorOpened() { return PropertyEditorOpened; }
