Allow GetOptions functions for FName and FString properties to specify a Display Name in addition to value

Opt-in by creating a UFUNCTION with the signature:
`TArray<FPropertyTextFString> GetOptionsWithStrings()`
or
`TArray<FPropertyDisplayNameFName> GetOptionsWithFNames()`

and in the UPROPERTY macros point to the function
```
UPROPERTY(EditAnywhere, meta=(GetOptions="GetOptionsWithStrings"))
FString/FName MyProperty;
```

or

```
UPROPERTY(EditAnywhere, meta=(GetOptions="GetOptionsWithFNames"))
FString/FName MyProperty2;
```

This should bring GetOptions closer to feature parity with enums.  Enums still have one advantage of allowing for tooltips to also be specified.

diff --git a/Engine/Source/Editor/PropertyEditor/Private/PropertyEditorUtils.cpp b/Engine/Source/Editor/PropertyEditor/Private/PropertyEditorUtils.cpp
--- a/Engine/Source/Editor/PropertyEditor/Private/PropertyEditorUtils.cpp
+++ b/Engine/Source/Editor/PropertyEditor/Private/PropertyEditorUtils.cpp
@@ -4,11 +4,46 @@
 
 #include "PropertyCustomizationHelpers.h"
 #include "PropertyPathHelpers.h"
+#include "Algo/ForEach.h"
+#include "UObject/PropertyText.h"
 
 namespace PropertyEditorUtils
 {
+	
 	void GetPropertyOptions(TArray<UObject*>& InOutContainers, FString& InOutPropertyPath,
 		TArray<TSharedPtr<FString>>& InOutOptions)
+	{
+		TArray<FString> OptionsStrings;
+		TArray<FText>* NullDisplayNames = nullptr;
+		GetPropertyOptions(InOutContainers, InOutPropertyPath, OptionsStrings, NullDisplayNames);
+		
+		Algo::Transform(OptionsStrings, InOutOptions, [](const FString& InString) { return MakeShared<FString>(InString); });
+	}
+
+	struct FOptionsData
+	{
+		FString ValueString;
+		FText DisplayName;
+
+		friend uint32 GetTypeHash(const FOptionsData& Data)
+		{
+			// TODO: Is there a faster way to hash a FText?
+			return HashCombine(GetTypeHash(Data.ValueString), GetTypeHash(Data.DisplayName.ToString()));
+		}
+
+		bool operator==(const FOptionsData& Other) const
+		{
+			return ValueString == Other.ValueString && DisplayName.IdenticalTo(Other.DisplayName, ETextIdenticalModeFlags::LexicalCompareInvariants);
+		}
+
+		bool operator!=(const FOptionsData& Other) const
+		{
+			return !(this->operator==(Other));
+		}
+	};
+
+	void GetPropertyOptions(TArray<UObject*>& InOutContainers, FString& InOutPropertyPath,
+		TArray<FString>& InOutOptions, TArray<FText>* OutDisplayNames)
 	{
 		// Check for external function references
 		if (InOutPropertyPath.Contains(TEXT(".")))
@@ -26,22 +61,59 @@ namespace PropertyEditorUtils
 
 		if (InOutContainers.Num() > 0)
 		{
-			TArray<FString> OptionIntersection;
-			TSet<FString> OptionIntersectionSet;
+			TArray<FOptionsData> OptionIntersection;
+			TSet<FOptionsData> OptionIntersectionSet;
 
+			// Loop carried containers
+			TArray<FOptionsData> LoopOptions;
+			TArray<FString> StringOptions;
+			TArray<FText> DisplayNameOptions;
+			TArray<FName> NameOptions;
+			TArray<FPropertyTextString> StringDisplayNameOptions;
+			TArray<FPropertyTextFName> FNameDisplayNameOptions;
+			
+			// Need to find the intersection between all OptionsData.
+			// If there is no OutDisplayNames or a bound function to GetOptions doesn't provide DisplayNames, then don't provide
+			// display names either.
+			bool bUseDisplayNames = true;
+			
 			for (UObject* Target : InOutContainers)
 			{
-				TArray<FString> StringOptions;
+				LoopOptions.Empty(LoopOptions.Num());
+				StringOptions.Empty(StringOptions.Num());
+				DisplayNameOptions.Empty(DisplayNameOptions.Num());
+				NameOptions.Empty(NameOptions.Num());
+				StringDisplayNameOptions.Empty(StringDisplayNameOptions.Num());
+				FNameDisplayNameOptions.Empty(FNameDisplayNameOptions.Num());
 				{
 					FEditorScriptExecutionGuard ScriptExecutionGuard;
 
+					// Test each signature of a function
 					FCachedPropertyPath Path(InOutPropertyPath);
-					if (!PropertyPathHelpers::GetPropertyValue(Target, Path, StringOptions))
+
+					// Handle function signature: "TArray<FString> GetOptions()"
+					if (PropertyPathHelpers::GetPropertyValue(Target, Path, StringOptions))
+					{
+						Algo::Transform(StringOptions, LoopOptions, [](const FString& InName) { return FOptionsData{ .ValueString = InName, .DisplayName = FText() }; });;
+						bUseDisplayNames = false;
+					}
+					// Handle function signature: "TArray<FText> GetOptions()"
+					else if (PropertyPathHelpers::GetPropertyValue(Target, Path, NameOptions))
 					{
-						TArray<FName> NameOptions;
-						if (PropertyPathHelpers::GetPropertyValue(Target, Path, NameOptions))
+						Algo::Transform(NameOptions, LoopOptions, [](const FName& InName) { return FOptionsData{ .ValueString = InName.ToString(), .DisplayName = FText() }; });;
+						bUseDisplayNames = false;
+					}
+					else if (OutDisplayNames && bUseDisplayNames)
+					{
+						// Handle function signature: "TArray<FPropertyTextString> GetOptions()"
+						if (PropertyPathHelpers::GetPropertyValue(Target, Path, StringDisplayNameOptions))
+						{
+							Algo::Transform(StringDisplayNameOptions, LoopOptions, [](const FPropertyTextString& Pair) { return FOptionsData{ .ValueString = Pair.ValueString, .DisplayName = Pair.DisplayName }; });;
+						}
+						// Handle function signature: "TArray<FPropertyTextFName> GetOptions()"
+						else if (PropertyPathHelpers::GetPropertyValue(Target, Path, FNameDisplayNameOptions))
 						{
-							Algo::Transform(NameOptions, StringOptions, [](const FName& InName) { return InName.ToString(); });
+							Algo::Transform(FNameDisplayNameOptions, LoopOptions, [](const FPropertyTextFName& Pair) { return FOptionsData{ .ValueString = Pair.ValueString.ToString(), .DisplayName = Pair.DisplayName }; });;
 						}
 					}
 				}
@@ -49,14 +121,14 @@ namespace PropertyEditorUtils
 				// If this is the first time there won't be any options.
 				if (OptionIntersection.Num() == 0)
 				{
-					OptionIntersection = StringOptions;
-					OptionIntersectionSet = TSet<FString>(StringOptions);
+					OptionIntersection = LoopOptions;
+					OptionIntersectionSet = TSet<FOptionsData>(LoopOptions);
 				}
 				else
 				{
-					TSet<FString> StringOptionsSet(StringOptions);
-					OptionIntersectionSet = StringOptionsSet.Intersect(OptionIntersectionSet);
-					OptionIntersection.RemoveAll([&OptionIntersectionSet](const FString& Option){ return !OptionIntersectionSet.Contains(Option); });
+					TSet<FOptionsData> LoopOptionsSet(LoopOptions);
+					OptionIntersectionSet = LoopOptionsSet.Intersect(OptionIntersectionSet);
+					OptionIntersection.RemoveAll([&OptionIntersectionSet](const FOptionsData& Option){ return !OptionIntersectionSet.Contains(Option); });
 				}
 
 				// If we're out of possible intersected options, we can stop.
@@ -66,7 +138,12 @@ namespace PropertyEditorUtils
 				}
 			}
 
-			Algo::Transform(OptionIntersection, InOutOptions, [](const FString& InString) { return MakeShared<FString>(InString); });
+			Algo::Transform(OptionIntersection, InOutOptions, [](const FOptionsData& Option) { return Option.ValueString; });
+			if (OutDisplayNames && bUseDisplayNames)
+			{
+				TArray<FText>& OutDisplayNamesRef = *OutDisplayNames;
+				Algo::Transform(OptionIntersection, OutDisplayNamesRef, [](const FOptionsData& Option) { return Option.DisplayName; });
+			}
 		}
 	}
 
diff --git a/Engine/Source/Editor/PropertyEditor/Private/PropertyHandleImpl.cpp b/Engine/Source/Editor/PropertyEditor/Private/PropertyHandleImpl.cpp
--- a/Engine/Source/Editor/PropertyEditor/Private/PropertyHandleImpl.cpp
+++ b/Engine/Source/Editor/PropertyEditor/Private/PropertyHandleImpl.cpp
@@ -3100,6 +3100,24 @@ FPropertyAccess::Result FPropertyHandleBase::GetPerObjectValue( const int32 Obje
 }
 
 bool FPropertyHandleBase::GeneratePossibleValues(TArray< TSharedPtr<FString> >& OutOptionStrings, TArray< FText >& OutToolTips, TArray<bool>& OutRestrictedItems)
+{
+	TArray<FText> DisplayNames;
+	TArray<FString> ValueStrings;
+	const bool bUsesDisplayNames = GeneratePossibleValues(ValueStrings, OutToolTips, OutRestrictedItems, &DisplayNames);
+
+	if (bUsesDisplayNames)
+	{
+		Algo::Transform(DisplayNames, OutOptionStrings, [](const FText& Str) { return MakeShared<FString>(Str.ToString()); });
+	}
+	else
+	{
+		Algo::Transform(ValueStrings, OutOptionStrings, [](const FString& Str) { return MakeShared<FString>(Str); });
+	}
+	
+	return bUsesDisplayNames;
+}
+
+bool FPropertyHandleBase::GeneratePossibleValues(TArray<FString>& OutOptionStrings, TArray< FText >& OutToolTips, TArray<bool>& OutRestrictedItems, TArray<FText>* OutDisplayNames)
 {
 	FProperty* Property = GetProperty();
 	if (Property == nullptr)
@@ -3173,17 +3191,23 @@ bool FPropertyHandleBase::GeneratePossibleValues(TArray< TSharedPtr<FString> >&
 				const bool bIsRestricted = GenerateRestrictionToolTip(EnumName, RestrictionTooltip) || RestrictedEnumValues.Contains(Enum->GetNameByIndex(EnumIndex));
 				OutRestrictedItems.Add(bIsRestricted);
 
-				if (EnumDisplayName.Len() == 0)
-				{
-					EnumDisplayName = MoveTemp(EnumName);
-				}
-				else
+				const bool bHasDisplayName = EnumDisplayName.Len() > 0;
+
+				if (OutDisplayNames)
 				{
-					bUsesAlternateDisplayValues = true;
+					if (bHasDisplayName)
+					{
+						OutDisplayNames->Add(FText::FromString(EnumDisplayName));
+						bUsesAlternateDisplayValues = true;
+					}
+					else
+					{
+						// Added to ensure matching index of DisplayNames to ValueStrings
+						OutDisplayNames->Add(FText::FromString(EnumName));
+					}
 				}
-
-				TSharedPtr< FString > EnumStr(new FString(EnumDisplayName));
-				OutOptionStrings.Add(EnumStr);
+				
+				OutOptionStrings.Add(EnumName);
 
 				FText EnumValueToolTip = bIsRestricted ? RestrictionTooltip : Enum->GetToolTipTextByIndex(EnumIndex);
 				OutToolTips.Add(MoveTemp(EnumValueToolTip));
@@ -3204,7 +3228,12 @@ bool FPropertyHandleBase::GeneratePossibleValues(TArray< TSharedPtr<FString> >&
 			TArray<UObject*> OuterObjects;
 			GetOuterObjects(OuterObjects);
 
-			PropertyEditorUtils::GetPropertyOptions(OuterObjects, GetOptionsFunctionName, OutOptionStrings);
+			PropertyEditorUtils::GetPropertyOptions(OuterObjects, GetOptionsFunctionName, OutOptionStrings, OutDisplayNames);
+
+			if (OutDisplayNames && !OutDisplayNames->IsEmpty())
+			{
+				bUsesAlternateDisplayValues = true;
+			}
 		}
 	}
 	else if( Property->IsA(FClassProperty::StaticClass()) || Property->IsA(FSoftClassProperty::StaticClass()) )		
@@ -3213,8 +3242,12 @@ bool FPropertyHandleBase::GeneratePossibleValues(TArray< TSharedPtr<FString> >&
 			? CastFieldChecked<FClassProperty>(Property)->MetaClass
 			: CastFieldChecked<FSoftClassProperty>(Property)->MetaClass;
 
-		TSharedPtr< FString > NoneStr( new FString( TEXT("None") ) );
+		FString NoneStr( TEXT("None") );
 		OutOptionStrings.Add( NoneStr );
+		if (OutDisplayNames)
+		{
+			OutDisplayNames->Add(FText::FromString(NoneStr));
+		}
 
 		const bool bAllowAbstract = Property->GetOwnerProperty()->HasMetaData(TEXT("AllowAbstract"));
 		const bool bBlueprintBaseOnly = Property->GetOwnerProperty()->HasMetaData(TEXT("BlueprintBaseOnly"));
@@ -3231,7 +3264,11 @@ bool FPropertyHandleBase::GeneratePossibleValues(TArray< TSharedPtr<FString> >&
 					&& (!InterfaceThatMustBeImplemented || It->ImplementsInterface(InterfaceThatMustBeImplemented))
 					&& (!bAllowOnlyPlaceable || !It->HasAnyClassFlags(CLASS_Abstract | CLASS_NotPlaceable)))
 				{
-					OutOptionStrings.Add(TSharedPtr< FString >(new FString(It->GetName())));
+					OutOptionStrings.Add(It->GetName());
+					if (OutDisplayNames)
+					{
+						OutDisplayNames->Add(FText::FromString(It->GetName()));
+					}
 				}
 			}
 		}
diff --git a/Engine/Source/Editor/PropertyEditor/Private/PropertyHandleImpl.h b/Engine/Source/Editor/PropertyEditor/Private/PropertyHandleImpl.h
--- a/Engine/Source/Editor/PropertyEditor/Private/PropertyHandleImpl.h
+++ b/Engine/Source/Editor/PropertyEditor/Private/PropertyHandleImpl.h
@@ -494,6 +494,7 @@ public:
 	virtual FPropertyAccess::Result SetPerObjectValue( const int32 ObjectIndex, const FString& ObjectValue, EPropertyValueSetFlags::Type Flags = EPropertyValueSetFlags::DefaultFlags ) override;
 	virtual FPropertyAccess::Result GetPerObjectValue( const int32 ObjectIndex, FString& OutObjectValue ) const override;
 	virtual bool GeneratePossibleValues(TArray< TSharedPtr<FString> >& OutOptionStrings, TArray< FText >& OutToolTips, TArray<bool>& OutRestrictedItems) override;
+	virtual bool GeneratePossibleValues(TArray<FString>& OutOptionStrings, TArray< FText >& OutToolTips, TArray<bool>& OutRestrictedItems, TArray<FText>* OutDisplayNames) override;
 	virtual FPropertyAccess::Result SetObjectValueFromSelection() override;
 	virtual void NotifyPreChange() override;
 	virtual void NotifyPostChange(EPropertyChangeType::Type ChangeType) override;
diff --git a/Engine/Source/Editor/PropertyEditor/Private/UserInterface/PropertyEditor/SPropertyEditorCombo.cpp b/Engine/Source/Editor/PropertyEditor/Private/UserInterface/PropertyEditor/SPropertyEditorCombo.cpp
--- a/Engine/Source/Editor/PropertyEditor/Private/UserInterface/PropertyEditor/SPropertyEditorCombo.cpp
+++ b/Engine/Source/Editor/PropertyEditor/Private/UserInterface/PropertyEditor/SPropertyEditorCombo.cpp
@@ -150,9 +150,11 @@ void SPropertyEditorCombo::GenerateComboBoxStrings( TArray< TSharedPtr<FString>
 		return;
 	}
 
+	TArray<FString> ValueStrings;
 	TArray<FText> BasicTooltips;
-
-	bUsesAlternateDisplayValues = ComboArgs.PropertyHandle->GeneratePossibleValues(OutComboBoxStrings, BasicTooltips, OutRestrictedItems);
+	TArray<FText> DisplayNames;
+	
+	bUsesAlternateDisplayValues = ComboArgs.PropertyHandle->GeneratePossibleValues(ValueStrings, BasicTooltips, OutRestrictedItems, &DisplayNames);
 
 	// If we regenerate the entries, let's make sure that the currently selected item has the same shared pointer as
 	// the newly generated item with the same value, so that the generation of elements won't immediately result in a
@@ -179,68 +181,29 @@ void SPropertyEditorCombo::GenerateComboBoxStrings( TArray< TSharedPtr<FString>
 	AlternateDisplayValueToInternalValue.Reset();
 	InternalValueToAlternateDisplayValue.Reset();
 	if (const FProperty* Property = ComboArgs.PropertyHandle->GetProperty();
-		bUsesAlternateDisplayValues && !Property->IsA(FStrProperty::StaticClass()))
+		bUsesAlternateDisplayValues)
 	{
-		// currently only enum properties can use alternate display values; this 
-		// might change, so assert here so that if support is expanded to other 
-		// property types without updating this block of code, we'll catch it quickly
-		const UEnum* Enum = nullptr;
-		if (const FByteProperty* ByteProperty = CastField<FByteProperty>(Property))
-		{
-			Enum = ByteProperty->Enum;
-		}
-		else if (const FEnumProperty* EnumProperty = CastField<FEnumProperty>(Property))
+		if (ensureMsgf(ValueStrings.Num() == DisplayNames.Num(), TEXT("Mismatched Value and DisplayNames Array")))
 		{
-			Enum = EnumProperty->GetEnum();
-		}
-		check(Enum != nullptr);
-
-		const TMap<FName, FText> EnumValueDisplayNameOverrides = PropertyEditorHelpers::GetEnumValueDisplayNamesFromPropertyOverride(Property, Enum);
-		auto FindEnumValueIndex = [&EnumValueDisplayNameOverrides, &Enum](const FString& ValueString) -> int32
-		{
-			for (const TTuple<FName, FText>& EnumValueDisplayNameOverridePair : EnumValueDisplayNameOverrides)
-			{
-				if (EnumValueDisplayNameOverridePair.Value.ToString() == ValueString)
-				{
-					return Enum->GetIndexByName(EnumValueDisplayNameOverridePair.Key);
-				}
-			}
-
-			for (int32 ValIndex = 0; ValIndex < Enum->NumEnums() - 1; ++ValIndex)
+			AlternateDisplayValueToInternalValue.Reserve(ValueStrings.Num());
+			InternalValueToAlternateDisplayValue.Reserve(ValueStrings.Num());
+			for (int32 Index = 0, End = ValueStrings.Num(); Index < End; ++Index)
 			{
-				const FString EnumName = Enum->GetNameStringByIndex(ValIndex);
-				const FString DisplayName = Enum->GetDisplayNameTextByIndex(ValIndex).ToString();
-
-				if (DisplayName.Len() > 0)
-				{
-					if (DisplayName == ValueString)
-					{
-						return ValIndex;
-					}
-				}
-
-				if (EnumName == ValueString)
-				{
-					return ValIndex;
-				}
+				AlternateDisplayValueToInternalValue.Emplace(DisplayNames[Index].ToString(), ValueStrings[Index]);
+				InternalValueToAlternateDisplayValue.Emplace(ValueStrings[Index], DisplayNames[Index].ToString());
 			}
-
-			return INDEX_NONE;
-		};
-
-		for (const TSharedPtr<FString>& ValueStringPtr : OutComboBoxStrings)
+			Algo::Transform(DisplayNames, OutComboBoxStrings, [](const FText& Str) { return MakeShared<FString>(Str.ToString()); });
+		}
+		else
 		{
-			const int32 EnumIndex = FindEnumValueIndex(*ValueStringPtr);
-			check(EnumIndex != INDEX_NONE);
-
-			const FString EnumValue = Enum->GetNameStringByIndex(EnumIndex);
-			if (EnumValue != *ValueStringPtr)
-			{
-				AlternateDisplayValueToInternalValue.Add(*ValueStringPtr, EnumValue);
-				InternalValueToAlternateDisplayValue.Add(EnumValue, *ValueStringPtr);
-			}
+			bUsesAlternateDisplayValues = false;
+			Algo::Transform(ValueStrings, OutComboBoxStrings, [](const FString& Str) { return MakeShared<FString>(Str); });
 		}
 	}
+	else
+	{
+		Algo::Transform(ValueStrings, OutComboBoxStrings, [](const FString& Str) { return MakeShared<FString>(Str); });
+	}
 
 	// For enums, look for rich tooltip information
 	if(ComboArgs.PropertyHandle.IsValid())
diff --git a/Engine/Source/Editor/PropertyEditor/Public/PropertyEditorUtils.h b/Engine/Source/Editor/PropertyEditor/Public/PropertyEditorUtils.h
--- a/Engine/Source/Editor/PropertyEditor/Public/PropertyEditorUtils.h
+++ b/Engine/Source/Editor/PropertyEditor/Public/PropertyEditorUtils.h
@@ -14,6 +14,19 @@ namespace PropertyEditorUtils
 	 */
 	PROPERTYEDITOR_API void GetPropertyOptions(TArray<UObject*>& InOutContainers, FString& InOutPropertyPath, TArray<TSharedPtr<FString>>& InOutOptions);
 
+	/**
+	 * Calculates the possible drop-down options for the specified property path
+	 * The resulting option 
+	 * @param	InOutContainers		The container objects to resolve the property path against
+	 * @param	InOutPropertyPath	The property path
+	 * @param	InOutOptions		The resulting options
+	 */
+	PROPERTYEDITOR_API inline void GetPropertyOptions(
+		TArray<UObject*>& InOutContainers,
+		FString& InOutPropertyPath,
+		TArray<FString>& OutOptions,
+		TArray<FText>* OutDisplayNames);
+
 	/**
 	 * Get all allowed and disallowed classes according to meta data
 	 * @param	ObjectList			The list of object that owns the property
diff --git a/Engine/Source/Editor/PropertyEditor/Public/PropertyHandle.h b/Engine/Source/Editor/PropertyEditor/Public/PropertyHandle.h
--- a/Engine/Source/Editor/PropertyEditor/Public/PropertyHandle.h
+++ b/Engine/Source/Editor/PropertyEditor/Public/PropertyHandle.h
@@ -574,6 +574,12 @@ public:
 	 */
 	virtual bool GeneratePossibleValues(TArray< TSharedPtr<FString> >& OutOptionStrings, TArray< FText >& OutToolTips, TArray<bool>& OutRestrictedItems) = 0;
 
+	/**
+	* Generates a list of possible enum/class options for the property
+	* DisplayNames is optional
+	*/
+	virtual bool GeneratePossibleValues(TArray<FString>& OutOptionStrings, TArray< FText >& OutToolTips, TArray<bool>& OutRestrictedItems, TArray<FText>* OutDisplayNames) = 0;
+
 	/**
 	 * Marks this property has hidden by customizaton (will not show up in the default place)
 	 */
diff --git a/Engine/Source/Runtime/CoreUObject/Public/UObject/PropertyText.h b/Engine/Source/Runtime/CoreUObject/Public/UObject/PropertyText.h
new file mode 100644
--- /dev/null
+++ b/Engine/Source/Runtime/CoreUObject/Public/UObject/PropertyText.h
@@ -0,0 +1,50 @@
+// Copyright Epic Games, Inc. All Rights Reserved.
+
+#pragma once
+#include "CoreMinimal.h"
+#include "UObject/NameTypes.h"
+#include "Internationalization/Text.h"
+#include "PropertyText.generated.h"
+
+// Structs that support UPROPERTY(..., meta = (GetOptions = <function>)) on FName and FString
+// properties which allows the display name to be set separately from the value for a given FName or FString
+//
+// Example usage:
+//
+// UCLASS()
+// class UMyClass : public UObject
+// {
+// 	GENERATED_BODY()
+// public:
+//     UPROPERTY(EditAnywhere, meta = (GetOptions = GetFooOptions))
+// 	FName Foo;
+//
+// 	UFUNCTION()
+// 	static TArray<FPropertyTextFName> GetFooOptions();
+// };
+
+
+/**
+ * Tuple returned by <function> in a UPROPERTY(..., meta = (GetOptions = <function>))
+ * allowing value to be specified as a FName
+ */
+USTRUCT()
+struct COREUOBJECT_API FPropertyTextFName
+{
+	GENERATED_BODY()
+	
+	FName ValueString;
+	FText DisplayName;
+};
+
+/**
+ * Tuple returned by <function> in a UPROPERTY(..., meta = (GetOptions = <function>))
+ * allowing value to be specified as a FString
+ */
+USTRUCT()
+struct COREUOBJECT_API FPropertyTextString
+{
+	GENERATED_BODY()
+	FString ValueString;
+	FText DisplayName;
+};
\ No newline at end of file
