Fix compiler error on non-unity for certain compiler configurations
* `inline` specifier used on function that is not defined

diff --git a/Engine/Source/Editor/PropertyEditor/Public/PropertyEditorUtils.h b/Engine/Source/Editor/PropertyEditor/Public/PropertyEditorUtils.h
--- a/Engine/Source/Editor/PropertyEditor/Public/PropertyEditorUtils.h
+++ b/Engine/Source/Editor/PropertyEditor/Public/PropertyEditorUtils.h
@@ -21,7 +21,7 @@ namespace PropertyEditorUtils
 	 * @param	InOutPropertyPath	The property path
 	 * @param	InOutOptions		The resulting options
 	 */
-	PROPERTYEDITOR_API inline void GetPropertyOptions(
+	PROPERTYEDITOR_API void GetPropertyOptions(
 		TArray<UObject*>& InOutContainers,
 		FString& InOutPropertyPath,
 		TArray<FString>& OutOptions,
