Fix to a crash in Property Editor when a property is null.

diff --git a/Engine/Source/Editor/PropertyEditor/Private/PropertyEditorModule.cpp b/Engine/Source/Editor/PropertyEditor/Private/PropertyEditorModule.cpp
--- a/Engine/Source/Editor/PropertyEditor/Private/PropertyEditorModule.cpp
+++ b/Engine/Source/Editor/PropertyEditor/Private/PropertyEditorModule.cpp
@@ -993,6 +993,11 @@ FPropertyTypeLayoutCallback FPropertyEditorModule::GetPropertyTypeCustomization(
 			UClass* InstanceBaseClass = nullptr;
 			PropertyHandle.EnumerateConstRawData([&InstanceBaseClass, ObjectProperty](const void* RawData, const int32 /*DataIndex*/, const int32 /*NumDatas*/)
 			{
+				if (!RawData)
+				{
+					return true;
+				}
+
 				if (const UObject* Object = ObjectProperty->GetObjectPropertyValue(RawData))
 				{
 					UClass* Class = Object->GetClass();
