Performance improvement for nested FInstancedStruct in details panel
* Pathological performance issues occurred when viewing/editing USTRUCT objects that nested FInstancedStruct (either directly or in containers).  The PropertyNode system's implementation when handling FInstancedStructs would recurse up the property tree to discover the address of the owning objects and then recurse again up with the starting address.  The root cause of the performance issue is that the recursion would branch on the way up in the following functions: GetInstancesNum, GetMemoryOfInstance, GetOwnerPackages, HasValidStructData (by way of IsValid). On the way up again, branches would occur in GetValueBaseAddress, via duplicate calls in FItemPropertyNode::GetValueBaseAddress to ParentNode->GetValueAddress.
* The problem has been addressed by having the StructPropertyNode handle StructProviders that perform indirection differently than other providers.  When indirection is present, walk up to the parent node of the struct which will be an ItemPropertyNode - the address of this is retrieved and passed back down to the StructProvider.  In the current codebase, the FInstancedStructProvider is the only one that does this type of indirection - it can make the assumption that the passed in address is an FInstancedStruct and thus reinterpret it to get the StructMemory and UStruct* pointers.
* Added an InternalGetReadAddressUncached which public GetReadAddress** functions will call.  This new function will not resolve the Struct in the case of a Struct property indirection unlike GetReadAddressUncached.  It is intended as an internal function only to be called by the public APIs.
* Added a safeguard canary value to FInstancedStruct to debug/catch cases where the wrong raw void* is reinterpreted as an FInstancedStruct.  FInstancedStruct::CastFromVoid helper should be used over a naked reinterpret_cast
* Note that it isn't guarantees all paths where slow recursive behaviour occurs has been addressed in this CL. Specifically, the InstancedStructDetails implementations to get the dispaly values still go through the data enumeration though it should be better now that GetInstancesNum and GetMemoryOfInstance are not recursively calling the Enumerate pathways.  See GetDisplayValueString, GetTooltipText, GetDisplayValueIcon, GenerateStructPicker - this is the primary expense in Slate that is probably causing additional performance loss.

diff --git a/Engine/Plugins/Experimental/Chooser/Source/ProxyTableEditor/Private/StructOutputDataCustomization.cpp b/Engine/Plugins/Experimental/Chooser/Source/ProxyTableEditor/Private/StructOutputDataCustomization.cpp
--- a/Engine/Plugins/Experimental/Chooser/Source/ProxyTableEditor/Private/StructOutputDataCustomization.cpp
+++ b/Engine/Plugins/Experimental/Chooser/Source/ProxyTableEditor/Private/StructOutputDataCustomization.cpp
@@ -128,7 +128,7 @@ void FStructOutputDataCustomization::CustomizeHeader(TSharedRef<IPropertyHandle>
          	 				
 							// reset the type of the FInstancedStruct storing the value
 							ValueHandle->GetValueData(ValuePtr);
-							FInstancedStruct* ValueStruct = reinterpret_cast<FInstancedStruct*>(ValuePtr);
+							FInstancedStruct* ValueStruct = FInstancedStruct::CastFromVoid(ValuePtr);
 
          	 				// begin Temporary data conversion codepath
          	 				if (StructPropertyBinding->StructType == FChooserPlayerSettings::StaticStruct() &&
diff --git a/Engine/Plugins/Experimental/StructUtils/Source/StructUtils/Public/InstancedStruct.h b/Engine/Plugins/Experimental/StructUtils/Source/StructUtils/Public/InstancedStruct.h
--- a/Engine/Plugins/Experimental/StructUtils/Source/StructUtils/Public/InstancedStruct.h
+++ b/Engine/Plugins/Experimental/StructUtils/Source/StructUtils/Public/InstancedStruct.h
@@ -12,6 +12,12 @@ struct FConstSharedStruct;
 template<typename BaseStructT> struct TConstSharedStruct;
 class UUserDefinedStruct;
 
+#if UE_BUILD_SHIPPING
+#define INSTANCED_STRUCT_CAST_DEBUG 0
+#else
+#define INSTANCED_STRUCT_CAST_DEBUG 1
+#endif
+
 /**
  * FInstancedStruct works similarly as instanced UObject* property but is USTRUCTs.
  * 
@@ -29,6 +35,14 @@ struct STRUCTUTILS_API FInstancedStruct
 	GENERATED_BODY()
 
 public:
+	/**
+	 * The property tree relies a lot of blind casting a void* to an FInstancedStruct in order to
+	 * dereference to get at the struct within.  Unforunately, this can be very error prone to debug
+	 * Therefore, some internal machinery has been added to help catch programming errors with the cast.
+	 * It is preferable to use this utility function over a static_cast or reinterpret_cast.
+	 * @return 
+	 */
+	static FInstancedStruct* CastFromVoid(void*);
 
 	FInstancedStruct();
 
@@ -257,10 +271,37 @@ protected:
 		StructMemory = InStructMemory;
 	}
 
+#if INSTANCED_STRUCT_CAST_DEBUG
+	static constexpr uint64 CanaryValue = 0x72745374736e4946;  // "FInstStr"
+
+	// Debug tool to catch bad casts from void*
+	uint64 BeginCanary = CanaryValue;
+#endif
+	
 	TObjectPtr<const UScriptStruct> ScriptStruct = nullptr;
 	uint8* StructMemory = nullptr;
 };
 
+inline FInstancedStruct* FInstancedStruct::CastFromVoid(void* Memory)
+{
+#if INSTANCED_STRUCT_CAST_DEBUG
+	if (Memory == nullptr)
+	{
+		return nullptr;
+	}
+
+	// Check the canary values
+	FInstancedStruct* InstancedStruct = static_cast<FInstancedStruct*>(Memory);
+	if (ensureMsgf(InstancedStruct->BeginCanary == FInstancedStruct::CanaryValue, TEXT("Bad cast from a void*")))
+	{
+		return InstancedStruct;
+	}
+	return nullptr;
+#else
+	return static_cast<FInstancedStruct*>(Memory);
+#endif
+}
+
 template<>
 struct TStructOpsTypeTraits<FInstancedStruct> : public TStructOpsTypeTraitsBase2<FInstancedStruct>
 {
diff --git a/Engine/Plugins/Experimental/StructUtils/Source/StructUtilsEditor/Private/Customizations/InstancedStructDetails.cpp b/Engine/Plugins/Experimental/StructUtils/Source/StructUtilsEditor/Private/Customizations/InstancedStructDetails.cpp
--- a/Engine/Plugins/Experimental/StructUtils/Source/StructUtilsEditor/Private/Customizations/InstancedStructDetails.cpp
+++ b/Engine/Plugins/Experimental/StructUtils/Source/StructUtilsEditor/Private/Customizations/InstancedStructDetails.cpp
@@ -46,22 +46,15 @@ public:
 	
 	virtual bool IsValid() const override
 	{
-		bool bHasValidData = false;
-		EnumerateInstances([&bHasValidData](const UScriptStruct* ScriptStruct, uint8* Memory, UPackage* Package)
-		{
-			if (ScriptStruct && Memory)
-			{
-				bHasValidData = true;
-				return false; // Stop
-			}
-			return true; // Continue
-		});
-
-		return bHasValidData;
+		ensureMsgf(false, TEXT("InstancedStructDetails does not support IsValid - get address of parent node and use IsIndirectedValid"));
+		return false;
 	}
 	
 	virtual const UStruct* GetBaseStructure() const override
 	{
+		// TODO: Determine a way to implement this function for callers in a way that doesn't require EnumerateInstances
+		// The problem with EnumerateInstances is that nested FInstanceStructs can have pathologically bad performance
+		
 		// Taken from UClass::FindCommonBase
 		auto FindCommonBaseStruct = [](const UScriptStruct* StructA, const UScriptStruct* StructB)
 		{
@@ -72,7 +65,7 @@ public:
 			}
 			return CommonBaseStruct;
 		};
-
+		
 		const UScriptStruct* CommonStruct = nullptr;
 		EnumerateInstances([&CommonStruct, &FindCommonBaseStruct](const UScriptStruct* ScriptStruct, uint8* Memory, UPackage* Package)
 		{
@@ -82,28 +75,13 @@ public:
 			}
 			return true; // Continue
 		});
-
+		
 		return CommonStruct;
 	}
 
 	virtual void GetInstances(TArray<TSharedPtr<FStructOnScope>>& OutInstances, const UStruct* ExpectedBaseStructure) const override
 	{
-		// The returned instances need to be compatible with base structure.
-		// This function returns empty instances in case they are not compatible, with the idea that we have as many instances as we have outer objects.
-		EnumerateInstances([&OutInstances, ExpectedBaseStructure](const UScriptStruct* ScriptStruct, uint8* Memory, UPackage* Package)
-		{
-			TSharedPtr<FStructOnScope> Result;
-			
-			if (ExpectedBaseStructure && ScriptStruct && ScriptStruct->IsChildOf(ExpectedBaseStructure))
-			{
-				Result = MakeShared<FStructOnScope>(ScriptStruct, Memory);
-				Result->SetPackage(Package);
-			}
-
-			OutInstances.Add(Result);
-
-			return true; // Continue
-		});
+		ensureMsgf(false, TEXT("InstancedStructDetails does not support GetInstances - get address of parent node and use GetValueBaseAddress & GetIndirectedStructType"));
 	}
 
 	virtual bool IsPropertyIndirection() const override
@@ -113,20 +91,51 @@ public:
 
 	virtual uint8* GetValueBaseAddress(uint8* ParentValueAddress, const UStruct* ExpectedBaseStructure) const override
 	{
-		if (!ParentValueAddress)
+		FInstancedStruct* InstancedStruct = FInstancedStruct::CastFromVoid(ParentValueAddress);
+		if (ExpectedBaseStructure && InstancedStruct)
 		{
-			return nullptr;
+			const UScriptStruct* Struct = InstancedStruct->GetScriptStruct();
+			if (Struct && Struct->IsChildOf(ExpectedBaseStructure))
+			{
+				return InstancedStruct->GetMutableMemory();
+			}
 		}
+		
+		return nullptr;
+	}
 
-		FInstancedStruct& InstancedStruct = *reinterpret_cast<FInstancedStruct*>(ParentValueAddress);
-		if (ExpectedBaseStructure && InstancedStruct.GetScriptStruct() && InstancedStruct.GetScriptStruct()->IsChildOf(ExpectedBaseStructure))
+	virtual const UStruct* GetIndirectedStructType(uint8* ParentValueAddress, const UStruct* ExpectedBaseStructure) const override
+	{
+		FInstancedStruct* InstancedStruct = FInstancedStruct::CastFromVoid(ParentValueAddress);
+		if (ExpectedBaseStructure && InstancedStruct)
 		{
-			return InstancedStruct.GetMutableMemory();
+			const UStruct* Struct = InstancedStruct->GetScriptStruct();
+			if (Struct && Struct->IsChildOf(ExpectedBaseStructure))
+			{
+				return Struct;
+			}
 		}
 		
 		return nullptr;
 	}
 	
+	virtual const UStruct* GetIndirectedStructType(uint8* ParentValueAddress) const override
+	{
+		FInstancedStruct* InstancedStruct = FInstancedStruct::CastFromVoid(ParentValueAddress);
+		return InstancedStruct ? InstancedStruct->GetScriptStruct() : nullptr;
+	}
+
+	virtual bool IsIndirectedValid(uint8* ParentValueAddress) const override
+	{
+		if (FInstancedStruct* InstancedStruct = FInstancedStruct::CastFromVoid(ParentValueAddress))
+		{
+			const UStruct* StructType = InstancedStruct->GetScriptStruct();
+			const uint8* ValueAddress = InstancedStruct->GetMutableMemory();
+			return StructType && ValueAddress;
+		}
+		return false;
+	}
+	
 protected:
 
 	void EnumerateInstances(TFunctionRef<bool(const UScriptStruct* ScriptStruct, uint8* Memory, UPackage* Package)> InFunc) const
diff --git a/Engine/Plugins/Runtime/StateTree/Source/StateTreeEditorModule/Private/StateTreePropertyHelpers.cpp b/Engine/Plugins/Runtime/StateTree/Source/StateTreeEditorModule/Private/StateTreePropertyHelpers.cpp
--- a/Engine/Plugins/Runtime/StateTree/Source/StateTreeEditorModule/Private/StateTreePropertyHelpers.cpp
+++ b/Engine/Plugins/Runtime/StateTree/Source/StateTreeEditorModule/Private/StateTreePropertyHelpers.cpp
@@ -45,8 +45,8 @@ void DispatchPostEditToNodes(UObject& Owner, FPropertyChangedChainEvent& InPrope
 		{
 			if (StructProperty->Struct == FInstancedStruct::StaticStruct())
 			{
-				FInstancedStruct& InstancedStruct = *reinterpret_cast<FInstancedStruct*>(CurrentAddress);
-				CurrentAddress = InstancedStruct.GetMutableMemory();
+				FInstancedStruct* InstancedStruct = FInstancedStruct::CastFromVoid(CurrentAddress);
+				CurrentAddress = InstancedStruct ? InstancedStruct->GetMutableMemory() : nullptr;
 			}
 			else if (StructProperty->Struct == FStateTreeEditorNode::StaticStruct())
 			{
diff --git a/Engine/Plugins/Runtime/StateTree/Source/StateTreeModule/Private/StateTreePropertyBindings.cpp b/Engine/Plugins/Runtime/StateTree/Source/StateTreeModule/Private/StateTreePropertyBindings.cpp
--- a/Engine/Plugins/Runtime/StateTree/Source/StateTreeModule/Private/StateTreePropertyBindings.cpp
+++ b/Engine/Plugins/Runtime/StateTree/Source/StateTreeModule/Private/StateTreePropertyBindings.cpp
@@ -870,12 +870,12 @@ uint8* FStateTreePropertyBindings::GetAddress(FStateTreeDataView InStructView, c
 		case EStateTreePropertyAccessType::StructInstance:
 		{
 			check(Indirection->InstanceStruct);
-			FInstancedStruct& InstancedStruct = *reinterpret_cast<FInstancedStruct*>(Address + Indirection->Offset);
-			const UScriptStruct* InstanceType = InstancedStruct.GetScriptStruct(); 
+			FInstancedStruct* InstancedStruct = FInstancedStruct::CastFromVoid(Address + Indirection->Offset);
+			const UScriptStruct* InstanceType = InstancedStruct ? InstancedStruct->GetScriptStruct() : nullptr; 
 			if (InstanceType != nullptr
 				&& InstanceType->IsChildOf(Indirection->InstanceStruct))
 			{
-				Address = InstancedStruct.GetMutableMemory();
+				Address = InstancedStruct->GetMutableMemory();
 			}
 			else
 			{
diff --git a/Engine/Source/Editor/PropertyEditor/Private/ItemPropertyNode.cpp b/Engine/Source/Editor/PropertyEditor/Private/ItemPropertyNode.cpp
--- a/Engine/Source/Editor/PropertyEditor/Private/ItemPropertyNode.cpp
+++ b/Engine/Source/Editor/PropertyEditor/Private/ItemPropertyNode.cpp
@@ -70,7 +70,7 @@ uint8* FItemPropertyNode::GetValueBaseAddress(uint8* StartAddress, bool bIsSpars
 		}
 		else
 		{
-			uint8* ValueAddress = ParentNode->GetValueAddress(StartAddress, bIsSparseData, bIsStruct);
+			uint8* ValueAddress = ValueBaseAddress;
 			if (ValueAddress != nullptr && ParentNode->GetProperty() != MyProperty)
 			{
 				// if this is not a fixed size array (in which the parent property and this property are the same), we need to offset from the property (otherwise, the parent already did that for us)
diff --git a/Engine/Source/Editor/PropertyEditor/Private/ObjectPropertyNode.cpp b/Engine/Source/Editor/PropertyEditor/Private/ObjectPropertyNode.cpp
--- a/Engine/Source/Editor/PropertyEditor/Private/ObjectPropertyNode.cpp
+++ b/Engine/Source/Editor/PropertyEditor/Private/ObjectPropertyNode.cpp
@@ -415,7 +415,7 @@ bool FObjectPropertyNode::GetReadAddressUncached(const FPropertyNode& InNode,
  * @param InItem		The property to get objects from.
  * @param OutAddresses	Storage array for all of the objects' addresses.
  */
-bool FObjectPropertyNode::GetReadAddressUncached(const FPropertyNode& InNode, FReadAddressListData& OutAddresses ) const
+bool FObjectPropertyNode::GetReadAddressUncached(const FPropertyNode& InPropertyNode, FReadAddressListData& OutAddresses ) const
 {
 	// Are any objects selected for property editing?
 	if( !GetNumObjects())
@@ -423,7 +423,37 @@ bool FObjectPropertyNode::GetReadAddressUncached(const FPropertyNode& InNode, FR
 		return false;
 	}
 
-	const FProperty* InItemProperty = InNode.GetProperty();
+	const FProperty* InItemProperty = InPropertyNode.GetProperty();
+	// Is there a InItemProperty bound to the InItemProperty window?
+	if( !InItemProperty )
+	{
+		return false;
+	}
+
+
+	// Write addresses to the output.
+	for (int32 ObjIndex = 0 ; ObjIndex < GetNumObjects() ; ++ObjIndex)
+	{
+		const UObject* TempObject = GetUObject(ObjIndex);
+		if (TempObject != nullptr)
+		{
+			OutAddresses.Add(TempObject, InPropertyNode.GetValueBaseAddressFromObject(TempObject));
+		}
+	}
+
+	// Everything checked out and we have usable addresses.
+	return true;
+}
+
+bool FObjectPropertyNode::InternalGetReadAddressUncached(const FPropertyNode& InPropertyNode, FUncachedPropertyNodeAddresses& OutAddresses) const
+{
+	// Are any objects selected for property editing?
+	if( !GetNumObjects())
+	{
+		return false;
+	}
+
+	const FProperty* InItemProperty = InPropertyNode.GetProperty();
 	// Is there a InItemProperty bound to the InItemProperty window?
 	if( !InItemProperty )
 	{
@@ -437,7 +467,9 @@ bool FObjectPropertyNode::GetReadAddressUncached(const FPropertyNode& InNode, FR
 		const UObject* TempObject = GetUObject(ObjIndex);
 		if (TempObject != nullptr)
 		{
-			OutAddresses.Add(TempObject, InNode.GetValueBaseAddressFromObject(TempObject));
+			uint8* Address = InPropertyNode.GetValueBaseAddressFromObject(TempObject);
+			constexpr bool bIsStruct = false;
+			OutAddresses.Emplace(FAddressPair(TempObject, Address, bIsStruct));
 		}
 	}
 
diff --git a/Engine/Source/Editor/PropertyEditor/Private/ObjectPropertyNode.h b/Engine/Source/Editor/PropertyEditor/Private/ObjectPropertyNode.h
--- a/Engine/Source/Editor/PropertyEditor/Private/ObjectPropertyNode.h
+++ b/Engine/Source/Editor/PropertyEditor/Private/ObjectPropertyNode.h
@@ -26,6 +26,7 @@ public:
 	virtual const FObjectPropertyNode* AsObjectNode() const override { return this; }
 	virtual bool GetReadAddressUncached(const FPropertyNode& InNode, bool InRequiresSingleSelection, FReadAddressListData* OutAddresses, bool bComparePropertyContents = true, bool bObjectForceCompare = false, bool bArrayPropertiesCanDifferInSize = false) const override;
 	virtual bool GetReadAddressUncached(const FPropertyNode& InNode, FReadAddressListData& OutAddresses) const override;
+	virtual bool InternalGetReadAddressUncached(const FPropertyNode& InPropertyNode, FUncachedPropertyNodeAddresses& OutAddresses) const override;
 
 	/**
 	 * Returns the UObject at index "n" of the Objects Array
diff --git a/Engine/Source/Editor/PropertyEditor/Private/PropertyNode.cpp b/Engine/Source/Editor/PropertyEditor/Private/PropertyNode.cpp
--- a/Engine/Source/Editor/PropertyEditor/Private/PropertyNode.cpp
+++ b/Engine/Source/Editor/PropertyEditor/Private/PropertyNode.cpp
@@ -645,6 +645,16 @@ bool FPropertyNode::DoesChildPropertyRequireValidation(FProperty* InChildProp)
 	return InChildProp != nullptr && (CastField<FObjectProperty>(InChildProp) != nullptr || CastField<FStructProperty>(InChildProp) != nullptr);
 }
 
+bool FPropertyNode::InternalGetReadAddressUncached(const FPropertyNode& InPropertyNode, FUncachedPropertyNodeAddresses& OutAddresses) const
+{
+	const TSharedPtr<FPropertyNode> ParentNode = ParentNodeWeakPtr.Pin();
+	if (ParentNode.IsValid())
+	{
+		return ParentNode->InternalGetReadAddressUncached( InPropertyNode, OutAddresses );
+	}
+	return false;
+}
+
 void FPropertyNode::MarkChildrenAsRebuilt()
 {
 	bChildrenRebuilt = false;
@@ -1617,7 +1627,18 @@ bool FPropertyNode::GetReadAddressUncached( const FPropertyNode& InPropertyNode,
 	const TSharedPtr<FPropertyNode> ParentNode = ParentNodeWeakPtr.Pin();
 	if (ParentNode.IsValid())
 	{
-		return ParentNode->GetReadAddressUncached( InPropertyNode, OutAddresses );
+		FUncachedPropertyNodeAddresses Addresses;
+		const bool bSuccess = ParentNode->InternalGetReadAddressUncached(InPropertyNode, Addresses);
+
+		for (int32 Index = 0, End = Addresses.Num(); Index < End; ++Index)
+		{
+			OutAddresses.Add(
+				Addresses[Index].Object.Get(),
+				Addresses[Index].ReadAddress,
+				Addresses[Index].bIsStruct);
+		}
+		
+		return bSuccess;
 	}
 	return false;
 }
@@ -1644,6 +1665,7 @@ bool FPropertyNode::GetReadAddress(bool InRequiresSingleSelection,
 	CachedReadAddresses.Reset();
 
 	bool bAllValuesTheSame = GetReadAddressUncached( *this, InRequiresSingleSelection, &CachedReadAddresses, bComparePropertyContents, bObjectForceCompare, bArrayPropertiesCanDifferInSize );
+
 	OutAddresses.ReadAddressListData = &CachedReadAddresses;
 	CachedReadAddresses.bAllValuesTheSame = bAllValuesTheSame;
 	CachedReadAddresses.bRequiresCache = false;
diff --git a/Engine/Source/Editor/PropertyEditor/Private/PropertyNode.h b/Engine/Source/Editor/PropertyEditor/Private/PropertyNode.h
--- a/Engine/Source/Editor/PropertyEditor/Private/PropertyNode.h
+++ b/Engine/Source/Editor/PropertyEditor/Private/PropertyNode.h
@@ -128,7 +128,7 @@ public:
 	}
 	void Add(const UObject* Object, uint8* Address, bool bIsStruct = false)
 	{
-		ReadAddresses.Add(FAddressPair(Object, Address, bIsStruct));
+		ReadAddresses.Emplace(FAddressPair(Object, Address, bIsStruct));
 	}
 
 	int32 Num() const
@@ -362,6 +362,8 @@ private:
  */
 class FPropertyNode : public TSharedFromThis<FPropertyNode>
 {
+protected:
+	typedef TArray<FAddressPair, TInlineAllocator<1>> FUncachedPropertyNodeAddresses;
 public:
 
 	FPropertyNode();
@@ -1118,6 +1120,17 @@ public:
 	}
 
 	TSharedRef<FEditPropertyChain> BuildPropertyChain( FProperty* PropertyAboutToChange ) const;
+
+	/**
+	 * Returns addresses of given property node
+	 * If the property node is a struct indirection, will not dereference the indirection to the struct.  This behaviour
+	 * differs from the GetReadAddressUncached function.
+	 * This function is intended to only be called internally.
+	 * @param InPropertyNode 
+	 * @param OutAddresses 
+	 * @return 
+	 */
+	virtual bool InternalGetReadAddressUncached(const FPropertyNode& InPropertyNode, FUncachedPropertyNodeAddresses& OutAddresses) const;
 	
 protected:
 	TSharedRef<FEditPropertyChain> BuildPropertyChain( FProperty* PropertyAboutToChange, const TSet<UObject*>& InAffectedArchetypeInstances ) const;
diff --git a/Engine/Source/Editor/PropertyEditor/Private/StructurePropertyNode.cpp b/Engine/Source/Editor/PropertyEditor/Private/StructurePropertyNode.cpp
--- a/Engine/Source/Editor/PropertyEditor/Private/StructurePropertyNode.cpp
+++ b/Engine/Source/Editor/PropertyEditor/Private/StructurePropertyNode.cpp
@@ -3,6 +3,7 @@
 
 #include "StructurePropertyNode.h"
 #include "ItemPropertyNode.h"
+#include "ObjectPropertyNode.h"
 #include "PropertyEditorHelpers.h"
 
 void FStructurePropertyNode::InitChildNodes()
@@ -15,7 +16,44 @@ void FStructurePropertyNode::InternalInitChildNodes(FName SinglePropertyName)
 	const bool bShouldShowHiddenProperties = !!HasNodeFlags(EPropertyNodeFlags::ShouldShowHiddenProperties);
 	const bool bShouldShowDisableEditOnInstance = !!HasNodeFlags(EPropertyNodeFlags::ShouldShowDisableEditOnInstance);
 
-	const UStruct* Struct = GetBaseStructure();
+	const UStruct* Struct = nullptr;
+	if (StructProvider)
+	{
+		if (!StructProvider->IsPropertyIndirection())
+		{
+			Struct = GetBaseStructure();
+		}
+		else
+		{
+			FPropertyNode* Parent = GetParentNode();
+
+			FUncachedPropertyNodeAddresses ParentAddresses;
+			ParentAddresses.Reset();
+			const bool bIsValid = Parent->InternalGetReadAddressUncached(*Parent, ParentAddresses);
+			if (bIsValid)
+			{
+				auto FindCommonBaseStruct = [](const UStruct* StructA, const UStruct* StructB)
+				{
+					const UStruct* CommonBaseStruct = StructA;
+					while (CommonBaseStruct && StructB && !StructB->IsChildOf(CommonBaseStruct))
+					{
+						CommonBaseStruct = Cast<UStruct>(CommonBaseStruct->GetSuperStruct());
+					}
+					return CommonBaseStruct;
+				};
+
+				const UStruct* CommonStruct = nullptr;
+				for (int32 Index = 0, End = ParentAddresses.Num(); Index < End; ++Index)
+				{
+					uint8* ParentValueAddress = ParentAddresses[Index].ReadAddress;
+				
+					const UStruct* TempStruct = StructProvider->GetIndirectedStructType(ParentValueAddress);
+					CommonStruct = FindCommonBaseStruct(TempStruct, CommonStruct);
+				}
+				Struct = CommonStruct;
+			}
+		}
+	}
 
 	TArray<FProperty*> StructMembers;
 
@@ -66,17 +104,243 @@ void FStructurePropertyNode::InitBeforeNodeFlags()
 	// The struct will be cached on each call to InternalInitChildNodes() as well.
 	// We'll cache it here too, so that a FStructurePropertyNode which is initialized
 	// with "InitParams.bAllowChildren = false" has it properly set up.
-	WeakCachedBaseStruct = GetBaseStructure();
+	const UStruct* Struct = nullptr;
+	if (StructProvider)
+	{
+		if (!StructProvider->IsPropertyIndirection())
+		{
+			Struct = GetBaseStructure();
+		}
+		else
+		{
+			FPropertyNode* Parent = GetParentNode();
+
+			FUncachedPropertyNodeAddresses ParentAddresses;
+			ParentAddresses.Reset();
+			const bool bIsValid = Parent->InternalGetReadAddressUncached(*Parent, ParentAddresses);
+			if (bIsValid)
+			{
+				auto FindCommonBaseStruct = [](const UStruct* StructA, const UStruct* StructB)
+				{
+					const UStruct* CommonBaseStruct = StructA;
+					while (CommonBaseStruct && StructB && !StructB->IsChildOf(CommonBaseStruct))
+					{
+						CommonBaseStruct = Cast<UStruct>(CommonBaseStruct->GetSuperStruct());
+					}
+					return CommonBaseStruct;
+				};
+
+				const UStruct* CommonStruct = nullptr;
+				for (int32 Index = 0, End = ParentAddresses.Num(); Index < End; ++Index)
+				{
+					uint8* ParentValueAddress = ParentAddresses[Index].ReadAddress;
+				
+					const UStruct* TempStruct = StructProvider->GetIndirectedStructType(ParentValueAddress);
+					CommonStruct = FindCommonBaseStruct(TempStruct, CommonStruct);
+				}
+				Struct = CommonStruct;
+			}
+		}
+		WeakCachedBaseStruct = Struct;
+	}
+	
+
 }
 
-bool FStructurePropertyNode::GetReadAddressUncached(const FPropertyNode& InPropertyNode, FReadAddressListData& OutAddresses) const
+bool FStructurePropertyNode::HasValidStructData() const
 {
-	if (!HasValidStructData())
+	if (!StructProvider.IsValid())
 	{
 		return false;
 	}
-	check(StructProvider.IsValid());
+	if (!StructProvider->IsPropertyIndirection())
+	{
+		return StructProvider->IsValid();
+	}
+	else
+	{
+		const UStruct* ExpectedBaseStruct = WeakCachedBaseStruct.Get();
+		if (!ExpectedBaseStruct)
+		{
+			return false;
+		}
+		
+		const FPropertyNode* Parent = this->GetParentNode();
+		FUncachedPropertyNodeAddresses ParentAddresses;
+		const bool bIsValid = InternalGetReadAddressUncached(*Parent, ParentAddresses);
+
+		for (int32 AddressIndex = 0, AddressListEnd = ParentAddresses.Num(); AddressIndex < AddressListEnd; ++AddressIndex)
+		{
+			uint8* ParentAddress = ParentAddresses[AddressIndex].ReadAddress;
+			StructProvider->GetValueBaseAddress(ParentAddress, ExpectedBaseStruct);
+			StructProvider->GetIndirectedStructType(ParentAddress, ExpectedBaseStruct);
+			if (!StructProvider->IsIndirectedValid(ParentAddress))
+			{
+				return false;
+			}
+		}
+		return true;
+	}
+}
+
+void FStructurePropertyNode::GetAllStructureData(TArray<TSharedPtr<FStructOnScope>>& OutStructs) const
+{
+	if (StructProvider)
+	{
+		if (!StructProvider->IsPropertyIndirection())
+		{
+			StructProvider->GetInstances(OutStructs, WeakCachedBaseStruct.Get());
+		}
+		else
+		{
+			const FPropertyNode* ParentNode = this->GetParentNode();
+			if (ParentNode)
+			{
+				// NOTE: Recursively walks up the chain until there is a node that allows directly getting the memory address
+				// of the FInstancedStruct!!!
+				// TODO: It is unfortunate that this create temporary arrays on recurse
+				FUncachedPropertyNodeAddresses ParentAddresses;
+				const bool bIsValid = InternalGetReadAddressUncached(*ParentNode, ParentAddresses);
+
+				const UStruct* ExpectedBaseStruct = WeakCachedBaseStruct.Get();
+				if (bIsValid && ExpectedBaseStruct)
+				{
+					TArray<UPackage*> Packages;
+					GetOwnerPackages(Packages);
+					OutStructs.Reserve(ParentAddresses.Num());
+					for (int32 Index = 0, End = ParentAddresses.Num(); Index < End; ++Index)
+					{
+						uint8* ParentAddress = ParentAddresses[Index].ReadAddress;
+						uint8* ValueAddress = StructProvider->GetValueBaseAddress(ParentAddress, ExpectedBaseStruct);
+						const UStruct* Struct = StructProvider->GetIndirectedStructType(ParentAddress, ExpectedBaseStruct);
+						if (ValueAddress && Struct)
+						{
+							OutStructs.Emplace(MakeShared<FStructOnScope>(Struct, ValueAddress));
+							OutStructs.Last()->SetPackage(Packages[Index]);
+						}
+					}
+				}
+			}
+		}
+	}
+}
+
+void FStructurePropertyNode::GetOwnerPackages(TArray<UPackage*>& OutPackages) const
+{
+	if (StructProvider)
+	{
+		// Walk up until we find the objects that contain this struct property to get the packages
+		const FPropertyNode* Parent = this;
+		
+		while(Parent)
+		{
+			const FComplexPropertyNode* ComplexParent = Parent->FindComplexParent();
+			if (!ensureMsgf(ComplexParent != nullptr, TEXT("Expected to find a complex parent")))
+			{
+				return;
+			}
+			EPropertyType PropertyType = ComplexParent->GetPropertyType();
+				
+			switch (PropertyType)
+			{
+			case EPT_Object:
+				{
+					const FObjectPropertyNode* ObjectNode = ComplexParent->AsObjectNode();
+					for (int32 ObjectIndex = 0, ObjectCount = ObjectNode->GetNumObjects(); ObjectIndex < ObjectCount; ++ObjectIndex)
+					{
+						const UPackage* Package = ObjectNode->GetUPackage(ObjectIndex);
+						OutPackages.Add(const_cast<UPackage*>(Package));
+					}
+					Parent = nullptr;
+					break;
+				}
+			case EPT_StandaloneStructure:
+				{
+					// Iterate up until we find objects
+					Parent = ComplexParent->GetParentNode();
+					break;
+				}
+			}
+		}
+	}
+}
+
+bool FStructurePropertyNode::GetReadAddressUncached(const FPropertyNode& InPropertyNode, FReadAddressListData& OutAddresses) const
+{
+	bool bHasData = false;
+	if (StructProvider)
+	{
+		if (!StructProvider->IsPropertyIndirection())
+		{
+			if (!HasValidStructData())
+			{
+				return false;
+			}
+			check(StructProvider.IsValid());
+
+			const FProperty* InItemProperty = InPropertyNode.GetProperty();
+			if (!InItemProperty)
+			{
+				return false;
+			}
+
+			UStruct* OwnerStruct = InItemProperty->GetOwnerStruct();
+			if (!OwnerStruct || OwnerStruct->IsStructTrashed())
+			{
+				// Verify that the property is not part of an invalid trash class
+				return false;
+			}
+
+			TArray<TSharedPtr<FStructOnScope>> Instances;
+			StructProvider->GetInstances(Instances, WeakCachedBaseStruct.Get());
+
+			for (TSharedPtr<FStructOnScope>& Instance : Instances)
+			{
+				uint8* ReadAddress = Instance.IsValid() ? Instance->GetStructMemory() : nullptr;
+				if (ReadAddress)
+				{
+					OutAddresses.Add(nullptr, InPropertyNode.GetValueBaseAddress(ReadAddress, InPropertyNode.HasNodeFlags(EPropertyNodeFlags::IsSparseClassData) != 0, /*bIsStruct=*/true), /*bIsStruct=*/true);
+					bHasData = true;
+				}
+			}
+		}
+		else
+		{
+			// Need the parent of this node to be passed to the struct provider.  The struct provider will know what type it is and what to do with it
+			const FComplexPropertyNode* CurrentNode = this;
+			const FPropertyNode* ParentNode = CurrentNode->GetParentNode();
+			if (ParentNode)
+			{
+				// NOTE: Recursively walks up the chain until there is a node that allows directly getting the memory address
+				// of the FInstancedStruct!!!
+				// TODO: It is unfortunate that this create temporary arrays on recurse
+				FUncachedPropertyNodeAddresses ParentAddresses;
+				const bool bIsValid = InternalGetReadAddressUncached(*ParentNode, ParentAddresses);
+			
+				// Pass those addresses down to the StructProvider... it will know what to do with them
+				const UStruct* ExpectedBaseStruct = WeakCachedBaseStruct.Get();
+				if (bIsValid && ExpectedBaseStruct)
+				{
+					for (int32 Index = 0, End = ParentAddresses.Num(); Index < End; ++Index)
+					{
+						const UObject* Object = nullptr;
+						uint8* ParentValueAddress = ParentAddresses[Index].ReadAddress;
+						const bool bIsStruct = true;
+					
+						OutAddresses.Add(Object, StructProvider->GetValueBaseAddress(ParentValueAddress, ExpectedBaseStruct), bIsStruct);
+					}
+					bHasData = true;
+				}
+			}
+		}
+	}
+	
+	return bHasData;
+}
 
+bool FStructurePropertyNode::InternalGetReadAddressUncached(const FPropertyNode& InPropertyNode, FUncachedPropertyNodeAddresses& OutAddresses) const
+{
+	bool bHasData = false;
 	const FProperty* InItemProperty = InPropertyNode.GetProperty();
 	if (!InItemProperty)
 	{
@@ -89,20 +353,61 @@ bool FStructurePropertyNode::GetReadAddressUncached(const FPropertyNode& InPrope
 		// Verify that the property is not part of an invalid trash class
 		return false;
 	}
+	
+	if (!StructProvider)
+	{
+		return false;
+	}
+	
+	if (!StructProvider->IsPropertyIndirection())
+	{
+		TArray<TSharedPtr<FStructOnScope>> Instances;
+		StructProvider->GetInstances(Instances, WeakCachedBaseStruct.Get());
 
-	TArray<TSharedPtr<FStructOnScope>> Instances;
-	StructProvider->GetInstances(Instances, WeakCachedBaseStruct.Get());
-	bool bHasData = false;
-
-	for (TSharedPtr<FStructOnScope>& Instance : Instances)
+		for (TSharedPtr<FStructOnScope>& Instance : Instances)
+		{
+			uint8* ReadAddress = Instance.IsValid() ? Instance->GetStructMemory() : nullptr;
+			if (ReadAddress)
+			{
+				const UObject* Object = nullptr;
+				const bool bIsSparse = InPropertyNode.HasNodeFlags(EPropertyNodeFlags::IsSparseClassData) != 0;
+				constexpr bool bIsStruct = true;
+				uint8* Address = InPropertyNode.GetValueBaseAddress(ReadAddress, bIsSparse, bIsStruct);
+				OutAddresses.Emplace(FAddressPair(Object, Address, bIsStruct));
+				bHasData = true;
+			}
+		}
+	}
+	else
 	{
-		uint8* ReadAddress = Instance.IsValid() ? Instance->GetStructMemory() : nullptr;
-		if (ReadAddress)
+		// Need the parent of this node to be passed to the struct provider.  The struct provider will know what type it is and what to do with it
+		const FComplexPropertyNode* CurrentNode = this;
+		const FPropertyNode* ParentNode = CurrentNode->GetParentNode();
+		if (ParentNode)
 		{
-			OutAddresses.Add(nullptr, InPropertyNode.GetValueBaseAddress(ReadAddress, InPropertyNode.HasNodeFlags(EPropertyNodeFlags::IsSparseClassData) != 0, /*bIsStruct=*/true), /*bIsStruct=*/true);
-			bHasData = true;
+			// NOTE: Recursively walks up the chain until there is a node that allows directly getting the memory address
+			// of the FInstancedStruct!!!
+			// TODO: It is unfortunate that this create temporary arrays on recurse
+			FUncachedPropertyNodeAddresses ParentAddresses;
+			const bool bIsValid = ParentNode->InternalGetReadAddressUncached(InPropertyNode, ParentAddresses);
+			
+			// Pass those addresses down to the StructProvider... it will know what to do with them
+			const UStruct* ExpectedBaseStruct = WeakCachedBaseStruct.Get();
+			if (bIsValid && ExpectedBaseStruct)
+			{
+				for (int32 Index = 0, End = ParentAddresses.Num(); Index < End; ++Index)
+				{
+					const UObject* Object = nullptr;
+					uint8* ParentValueAddress = ParentAddresses[Index].ReadAddress;
+					const bool bIsStruct = true;
+
+					OutAddresses.Emplace(FAddressPair(Object, ParentValueAddress, bIsStruct));
+				}
+				bHasData = true;
+			}
 		}
 	}
+	
 	return bHasData;
 }
 
@@ -134,45 +439,94 @@ bool FStructurePropertyNode::GetReadAddressUncached(const FPropertyNode& InPrope
 
 	bool bAllTheSame = true;
 
-	TArray<TSharedPtr<FStructOnScope>> Instances;
-	StructProvider->GetInstances(Instances, WeakCachedBaseStruct.Get());
-	
-	if (Instances.IsEmpty())
+	struct InstanceData
 	{
-		return false;
-	}
+		uint8* Address;
+		const UStruct* Struct;
+	};
+	TArray<InstanceData> TempInstances;
 
+	{
+		if (!StructProvider->IsPropertyIndirection())
+		{
+			TArray<TSharedPtr<FStructOnScope>> Instances;
+			StructProvider->GetInstances(Instances, WeakCachedBaseStruct.Get());
+	
+			if (Instances.IsEmpty())
+			{
+				return false;
+			}
+
+			// TODO Convert Instances to TempInstances
+			for (auto& Instance : Instances)
+			{
+				if (Instance)
+				{
+					TempInstances.Emplace(
+						InstanceData{
+							.Address = Instance->GetStructMemory(),
+							.Struct = Instance->GetStruct()
+					});
+				}
+			}
+		}
+		else
+		{
+			const FPropertyNode* ParentNode = this->GetParentNode();
+			FUncachedPropertyNodeAddresses ParentAddresses;
+			const bool bIsValid = InternalGetReadAddressUncached(*ParentNode, ParentAddresses);
+
+			const UStruct* ExpectedBaseStruct = WeakCachedBaseStruct.Get();
+			if (bIsValid && ExpectedBaseStruct)
+			{
+				for (int32 Index = 0, End = ParentAddresses.Num(); Index < End; ++Index)
+				{
+					uint8* ParentValueAddress = ParentAddresses[Index].ReadAddress;
+					
+					const UStruct* StructType = StructProvider->GetIndirectedStructType(ParentValueAddress, ExpectedBaseStruct);
+					uint8* Address = StructProvider->GetValueBaseAddress(ParentValueAddress, ExpectedBaseStruct);
+					TempInstances.Emplace(
+						InstanceData{
+							.Address = Address,
+							.Struct = StructType
+					});
+				}
+			}
+			else
+			{
+				return false;
+			}
+		}
+	}
+	
 	if (bComparePropertyContents || bObjectForceCompare)
 	{
 		const bool bIsSparse = InPropertyNode.HasNodeFlags(EPropertyNodeFlags::IsSparseClassData) != 0;
 		const uint8* BaseAddress = nullptr;
 		const UStruct* BaseStruct = nullptr;
 
-		for (TSharedPtr<FStructOnScope>& Instance : Instances)
+		for (auto& Instance : TempInstances)
 		{
-			if (Instance.IsValid())
+			if (const UStruct* Struct = Instance.Struct)
 			{
-				if (const UStruct* Struct = Instance->GetStruct())
+				if (const uint8* ReadAddress = InPropertyNode.GetValueBaseAddress(Instance.Address, bIsSparse, /*bIsStruct=*/true))
 				{
-					if (const uint8* ReadAddress = InPropertyNode.GetValueBaseAddress(Instance->GetStructMemory(), bIsSparse, /*bIsStruct=*/true))
+					if (!BaseAddress)
 					{
-						if (!BaseAddress)
+						BaseAddress = ReadAddress;
+						BaseStruct = Struct;
+					}
+					else
+					{
+						if (BaseStruct != Struct)
 						{
-							BaseAddress = ReadAddress;
-							BaseStruct = Struct;
+							bAllTheSame = false;
+							break;
 						}
-						else
+						if (!InItemProperty->Identical(BaseAddress, ReadAddress))
 						{
-							if (BaseStruct != Struct)
-							{
-								bAllTheSame = false;
-								break;
-							}
-							if (!InItemProperty->Identical(BaseAddress, ReadAddress))
-							{
-								bAllTheSame = false;
-								break;
-							}
+							bAllTheSame = false;
+							break;
 						}
 					}
 				}
@@ -188,10 +542,10 @@ bool FStructurePropertyNode::GetReadAddressUncached(const FPropertyNode& InPrope
 	else
 	{
 		// Check that all are valid or invalid.
-		const UStruct* BaseStruct = Instances[0].IsValid() ? Instances[0]->GetStruct() : nullptr;
-		for (int32 Index = 1; Index < Instances.Num(); Index++)
+		const UStruct* BaseStruct = TempInstances[0].Struct;
+		for (int32 Index = 1; Index < TempInstances.Num(); Index++)
 		{
-			const UStruct* Struct = Instances[Index].IsValid() ? Instances[Index]->GetStruct() : nullptr;
+			const UStruct* Struct = TempInstances[Index].Struct;
 			if (BaseStruct != Struct)
 			{
 				bAllTheSame = false;
@@ -208,9 +562,9 @@ bool FStructurePropertyNode::GetReadAddressUncached(const FPropertyNode& InPrope
 
 	if (bAllTheSame && OutAddresses)
 	{
-		for (TSharedPtr<FStructOnScope>& Instance : Instances)
+		for (const auto& Instance : TempInstances)
 		{
-			uint8* ReadAddress = Instance.IsValid() ? Instance->GetStructMemory() : nullptr;
+			uint8* ReadAddress = Instance.Address;
 			if (ReadAddress)
 			{
 				OutAddresses->Add(nullptr, InPropertyNode.GetValueBaseAddress(ReadAddress, InPropertyNode.HasNodeFlags(EPropertyNodeFlags::IsSparseClassData) != 0, /*bIsStruct=*/true), /*bIsStruct=*/true);
@@ -232,6 +586,19 @@ uint8* FStructurePropertyNode::GetValueBaseAddress(uint8* StartAddress, bool bIs
 
 	if (StructProvider)
 	{
+		if (!StructProvider->IsPropertyIndirection())
+		{
+			// The struct is really standalone, in which case we always return the standalone struct data.
+			// In that case we can only support one instance, since we cannot discern them.
+			// Note: Multiple standalone structure instances are supported when bIsStruct is true (e.g. when the structure property is root node).
+			TArray<TSharedPtr<FStructOnScope>> Instances;
+			StructProvider->GetInstances(Instances, WeakCachedBaseStruct.Get());
+			ensureMsgf(Instances.Num() <= 1, TEXT("Expecting max one instance on standalone structure provider."));
+			if (Instances.Num() == 1 && Instances[0].IsValid())
+			{
+				return Instances[0]->GetStructMemory();
+			}
+		}
 		// Assume that this code gets called with an object or object sparse data.
 		//
 		// The passed object might not be the one that contains the values provided by struct provider.
@@ -240,7 +607,7 @@ uint8* FStructurePropertyNode::GetValueBaseAddress(uint8* StartAddress, bool bIs
 
 		// If the structure provider is set up as indirection, then it knows how to translate parent node's value address to
 		// new value address even on data that is not the same as in the structure provider.
-		if (StructProvider->IsPropertyIndirection())
+		else
 		{
 			const TSharedPtr<FPropertyNode> ParentNode = ParentNodeWeakPtr.Pin();
 			if (!ensureMsgf(ParentNode, TEXT("Expecting valid parent node when indirection structure provider is called with Object data.")))
@@ -252,19 +619,78 @@ uint8* FStructurePropertyNode::GetValueBaseAddress(uint8* StartAddress, bool bIs
 			uint8* ValueAddress = StructProvider->GetValueBaseAddress(ParentValueAddress, WeakCachedBaseStruct.Get());
 			return ValueAddress;
 		}
+	}
+	
+	return nullptr;
+}
 
-		// The struct is really standalone, in which case we always return the standalone struct data.
-		// In that case we can only support one instance, since we cannot discern them.
-		// Note: Multiple standalone structure instances are supported when bIsStruct is true (e.g. when the structure property is root node).
-		TArray<TSharedPtr<FStructOnScope>> Instances;
-		StructProvider->GetInstances(Instances, WeakCachedBaseStruct.Get());
-		ensureMsgf(Instances.Num() <= 1, TEXT("Expecting max one instance on standalone structure provider."));
-		if (Instances.Num() == 1 && Instances[0].IsValid())
+int32 FStructurePropertyNode::GetInstancesNum() const
+{
+	// Can't get instance count directly from standalone structures, need to walk to the next parent that is an object
+	// to get number of instances from that
+	const FComplexPropertyNode* CurrentNode = this;	
+	
+	while(CurrentNode)
+	{
+		if (const FObjectPropertyNode* ObjectNode = CurrentNode->AsObjectNode())
+		{
+			// Found owning UObject
+			const int32 ObjectCount = ObjectNode->GetInstancesNum();
+			return ObjectCount;
+		}
+		else if (const FStructurePropertyNode* StructNode = CurrentNode->AsStructureNode())
+		{
+			TSharedPtr<IStructureDataProvider> TempStructProvider = StructNode->GetStructProvider();
+			if (TempStructProvider->IsPropertyIndirection())
+			{
+				// Keep walking up
+				const FPropertyNode* ParentNode = CurrentNode->GetParentNode();
+				CurrentNode = ParentNode->FindComplexParent();
+			}
+			else
+			{
+				TArray<TSharedPtr<FStructOnScope>> Instances;
+				TempStructProvider->GetInstances(Instances, StructNode->WeakCachedBaseStruct.Get());
+				return Instances.Num();
+			}
+		}
+		else
 		{
-			return Instances[0]->GetStructMemory();
+			return 0;
+		}
+	}
+	return 0;
+}
+
+uint8* FStructurePropertyNode::GetMemoryOfInstance(int32 Index) const
+{
+	if (StructProvider)
+	{
+		if (!StructProvider->IsPropertyIndirection())
+		{
+			FReadAddressListData ReadAddresses;
+			bool bSuccessful = GetReadAddressUncached(*this, ReadAddresses);
+			if (bSuccessful && ReadAddresses.IsValidIndex(Index))
+			{
+				return ReadAddresses.GetAddress(Index);
+			}
+		}
+		else
+		{
+			const UStruct* ExpectedBaseStruct = WeakCachedBaseStruct.Get();
+			const FPropertyNode* Parent = GetParentNode();
+			FUncachedPropertyNodeAddresses ParentAddresses;
+			bool bSuccessful = InternalGetReadAddressUncached(*Parent, ParentAddresses);
+
+			if (bSuccessful && ParentAddresses.IsValidIndex(Index))
+			{
+				uint8* ParentAddress = ParentAddresses[Index].ReadAddress;
+				// Dereference to get the address of the struct
+				uint8* ValueAddress = StructProvider->GetValueBaseAddress(ParentAddress, ExpectedBaseStruct);
+				return ValueAddress;
+			}
 		}
 	}
-	
 	return nullptr;
 }
 
diff --git a/Engine/Source/Editor/PropertyEditor/Private/StructurePropertyNode.h b/Engine/Source/Editor/PropertyEditor/Private/StructurePropertyNode.h
--- a/Engine/Source/Editor/PropertyEditor/Private/StructurePropertyNode.h
+++ b/Engine/Source/Editor/PropertyEditor/Private/StructurePropertyNode.h
@@ -43,10 +43,7 @@ public:
 		StructProvider = InStructProvider;
 	}
 
-	bool HasValidStructData() const
-	{
-		return StructProvider.IsValid() && StructProvider->IsValid();
-	}
+	bool HasValidStructData() const;
 
 	// Returns just the first structure. Please use GetStructProvider() or GetAllStructureData() when dealing with multiple struct instances.
 	TSharedPtr<FStructOnScope> GetStructData() const
@@ -64,13 +61,7 @@ public:
 		return nullptr;
 	}
 
-	void GetAllStructureData(TArray<TSharedPtr<FStructOnScope>>& OutStructs) const
-	{
-		if (StructProvider)
-		{
-			StructProvider->GetInstances(OutStructs, WeakCachedBaseStruct.Get());
-		}
-	}
+	void GetAllStructureData(TArray<TSharedPtr<FStructOnScope>>& OutStructs) const;
 
 	TSharedPtr<IStructureDataProvider> GetStructProvider() const
 	{
@@ -86,20 +77,7 @@ public:
 		bool bObjectForceCompare,
 		bool bArrayPropertiesCanDifferInSize) const override;
 
-	void GetOwnerPackages(TArray<UPackage*>& OutPackages) const
-	{
-		if (StructProvider)
-		{
-			TArray<TSharedPtr<FStructOnScope>> Instances;
-			StructProvider->GetInstances(Instances, WeakCachedBaseStruct.Get());
-
-			for (TSharedPtr<FStructOnScope>& Instance : Instances)
-			{
-				// Returning null for invalid instances, to match instance count.
-				OutPackages.Add(Instance.IsValid() ? Instance->GetPackage() : nullptr);
-			}
-		}
-	}
+	void GetOwnerPackages(TArray<UPackage*>& OutPackages) const;
 
 	/** FComplexPropertyNode Interface */
 	virtual const UStruct* GetBaseStructure() const override
@@ -157,31 +135,8 @@ public:
 		}
 		return RetVal;
 	}
-	virtual int32 GetInstancesNum() const override
-	{
-		if (StructProvider)
-		{
-			TArray<TSharedPtr<FStructOnScope>> Instances;
-			StructProvider->GetInstances(Instances, WeakCachedBaseStruct.Get());
-			
-			return Instances.Num();
-		}
-		return 0;
-		
-	}
-	virtual uint8* GetMemoryOfInstance(int32 Index) const override
-	{
-		if (StructProvider)
-		{
-			TArray<TSharedPtr<FStructOnScope>> Instances;
-			StructProvider->GetInstances(Instances, WeakCachedBaseStruct.Get());
-			if (Instances.IsValidIndex(Index) && Instances[Index].IsValid())
-			{
-				return Instances[Index]->GetStructMemory();
-			}
-		}
-		return nullptr;
-	}
+	virtual int32 GetInstancesNum() const override;
+	virtual uint8* GetMemoryOfInstance(int32 Index) const override;
 	virtual uint8* GetValuePtrOfInstance(int32 Index, const FProperty* InProperty, const FPropertyNode* InParentNode) const override
 	{ 
 		if (InProperty == nullptr || InParentNode == nullptr)
@@ -254,7 +209,8 @@ protected:
 
 		return bAddedAnything;
 	}
-
+	
+	virtual bool InternalGetReadAddressUncached(const FPropertyNode& InPropertyNode, FUncachedPropertyNodeAddresses& OutAddresses) const override;
 private:
 	TSharedPtr<IStructureDataProvider> StructProvider;
 
diff --git a/Engine/Source/Editor/PropertyEditor/Public/IStructureDataProvider.h b/Engine/Source/Editor/PropertyEditor/Public/IStructureDataProvider.h
--- a/Engine/Source/Editor/PropertyEditor/Public/IStructureDataProvider.h
+++ b/Engine/Source/Editor/PropertyEditor/Public/IStructureDataProvider.h
@@ -84,6 +84,20 @@ public:
 	 */
 	virtual uint8* GetValueBaseAddress(uint8* ParentValueAddress, const UStruct* ExpectedBaseStructure) const { return ParentValueAddress; }
 
+	/**
+	 * Gets the struct type of the indirected struct held by the property holding this StructProvider
+	 * Null if dereferencing the indirection would result in a nullptr
+	 */
+	virtual const UStruct* GetIndirectedStructType(uint8* ParentValueAddress, const UStruct* ExpectedBaseStructure) const { return nullptr; }
+
+	/**
+	 * Gets the struct type of the indirected struct held by the property holding this StructProvider
+	 * Used when determining the value of the ExpectedBaseStructure to be cached
+	 */
+	virtual const UStruct* GetIndirectedStructType(uint8* ParentValueAddress) const { return nullptr; }
+
+	virtual bool IsIndirectedValid(uint8* ParentValueAddress) const { return false; }
+
 	UE_DEPRECATED(5.4, "Please use and override GetInstances() with ExpectedType instead.")
 	virtual void GetInstances(TArray<TSharedPtr<FStructOnScope>>& OutInstances) const final
 	{
