Localization:
Work around to allow property matrix to edit FText.

Currently a bulk edit of FText will result in the localization cache becoming invalidated as certain strings lose their localizability. This is a hack work around to fix that issue until we have a proper fix in place.

diff --git a/Engine/Source/Editor/PropertyEditor/Private/UserInterface/PropertyEditor/SPropertyEditorText.cpp b/Engine/Source/Editor/PropertyEditor/Private/UserInterface/PropertyEditor/SPropertyEditorText.cpp
--- a/Engine/Source/Editor/PropertyEditor/Private/UserInterface/PropertyEditor/SPropertyEditorText.cpp
+++ b/Engine/Source/Editor/PropertyEditor/Private/UserInterface/PropertyEditor/SPropertyEditorText.cpp
@@ -190,7 +190,37 @@ void SPropertyEditorText::OnTextCommitted( const FText& NewText, ETextCommit::Ty
 	if( (PropertyHandle->GetValueAsFormattedText( CurrentText ) != FPropertyAccess::MultipleValues || NewText.ToString() != FPropertyEditor::MultipleValuesDisplayName)
 		&& !NewText.ToString().Equals(CurrentText.ToString(), ESearchCase::CaseSensitive))
 	{
-		PropertyHandle->SetValueFromFormattedString( NewText.ToString() );
+		if (CastField<FTextProperty>(PropertyHandle->GetProperty()))
+		{
+			PropertyHandle->NotifyPreChange();
+
+			// We should preserve the localizable state of the existing text values when applying a new source string
+			PropertyHandle->EnumerateRawData([&NewText](void* RawData, const int32 /*DataIndex*/, const int32 /*NumDatas*/)
+			{
+				FText* TextData = static_cast<FText*>(RawData);
+				if (NewText.IsEmpty())
+				{
+					*TextData = NewText;
+				}
+				else if (TextData->IsCultureInvariant())
+				{
+					*TextData = FText::AsCultureInvariant(NewText.ToString());
+				}
+				else
+				{
+					// TODO: This could attempt to use stable keys, however FText properties should really be edited via STextPropertyEditableTextBox
+					*TextData = FText::ChangeKey(TEXT(""), FGuid::NewGuid().ToString(), NewText);
+				}
+				return true;
+			});
+
+			PropertyHandle->NotifyPostChange(EPropertyChangeType::ValueSet);
+			PropertyHandle->NotifyFinishedChangingProperties();
+		}
+		else
+		{
+			PropertyHandle->SetValueFromFormattedString(NewText.ToString());
+		}
 	}
 }
 
