Null check FObjectPropertyNode::GetUObject as other callers do

diff --git a/Engine/Source/Editor/PropertyEditor/Private/EditConditionContext.cpp b/Engine/Source/Editor/PropertyEditor/Private/EditConditionContext.cpp
--- a/Engine/Source/Editor/PropertyEditor/Private/EditConditionContext.cpp
+++ b/Engine/Source/Editor/PropertyEditor/Private/EditConditionContext.cpp
@@ -224,7 +224,7 @@ TOptional<bool> FEditConditionContext::GetBoolValue(const FString& PropertyName,
 				UObject* FunctionTarget = ObjectNode->GetUObject(ObjectIndex);
 				FCachedPropertyPath FunctionPath(PropertyName);
 				bool bResult = true;
-				if (PropertyPathHelpers::GetPropertyValue(FunctionTarget, FunctionPath, bResult))
+				if (FunctionTarget && PropertyPathHelpers::GetPropertyValue(FunctionTarget, FunctionPath, bResult))
 				{
 					if (!bResult)
 					{
