[BugFix] Various bugs related to showing/hiding properties in the fixup tool.
jira UE-201373 caused a few regressions in the fixup tool so this change is making an effort to use the same paradigms as that CL while resolving those regressions

- fixed Subproperties of structs being visible when FDetailsViewArgs::bShowLooseProperties is set to false

- fixed unset properties being visible in the left panel of fixup tool
- fixed remapping not hidding the previous property that was being remapped from in the fixup tool's left panel

diff --git a/Engine/Plugins/Experimental/InstanceDataObjectFixupTool/Source/Private/InstanceDataObjectFixupPanel.cpp b/Engine/Plugins/Experimental/InstanceDataObjectFixupTool/Source/Private/InstanceDataObjectFixupPanel.cpp
--- a/Engine/Plugins/Experimental/InstanceDataObjectFixupTool/Source/Private/InstanceDataObjectFixupPanel.cpp
+++ b/Engine/Plugins/Experimental/InstanceDataObjectFixupTool/Source/Private/InstanceDataObjectFixupPanel.cpp
@@ -9,6 +9,7 @@
 #include "InstanceDataObjectFixupDetailCustomization.h"
 #include "Modules/ModuleManager.h"
 #include "Editor.h"
+#include "Editor/PropertyEditor/Private/PropertyNode.h"
 #include "UObject/PropertyBagRepository.h"
 
 #include "UObject/OverriddenPropertySet.h"
@@ -299,6 +300,14 @@ TSharedPtr<IDetailsView>& FInstanceDataObjectFixupPanel::GenerateDetailsView(boo
 	DetailsViewArgs.ScrollbarAlignment = bScrollbarOnLeft ? HAlign_Left : HAlign_Right;
 	DetailsViewArgs.DetailsNameWidgetOverrideCustomization = MakeShared<FInstanceDataObjectNameWidgetOverride>(SharedThis(this));
 	DetailsViewArgs.bShowLooseProperties = !HasViewFlag(EViewFlags::HideLooseProperties);
+
+	if (HasViewFlag(EViewFlags::IncludeOnlySetBySerialization))
+	{
+		DetailsViewArgs.ShouldForceHideProperty.BindLambda([this](const TSharedRef<FPropertyNode>& PropertyNode)->bool
+		{
+			return !IsInRedirectedPropertyTree(*FPropertyNode::CreatePropertyPath(PropertyNode));
+		});
+	}
 	
 	FPropertyEditorModule& PropertyEditorModule = FModuleManager::GetModuleChecked<FPropertyEditorModule>("PropertyEditor");
 	DetailsView = PropertyEditorModule.CreateDetailView(DetailsViewArgs);
diff --git a/Engine/Source/Editor/PropertyEditor/Private/DetailFilter.h b/Engine/Source/Editor/PropertyEditor/Private/DetailFilter.h
--- a/Engine/Source/Editor/PropertyEditor/Private/DetailFilter.h
+++ b/Engine/Source/Editor/PropertyEditor/Private/DetailFilter.h
@@ -54,4 +54,8 @@ struct FDetailFilter
 	TSet<FPropertyPath> PropertyAllowList;
 	/** The set of selected sections to show. If empty, all sections are shown.*/
 	TSet<FName> VisibleSections;
+
+	/** Delegate that determines whether a property should be forced hidden - evaluated on panel refresh */
+	DECLARE_DELEGATE_RetVal_OneParam(bool, FShouldForceHideProperty, const TSharedRef<class FPropertyNode>&);
+	FShouldForceHideProperty ShouldForceHideProperty;
 };
diff --git a/Engine/Source/Editor/PropertyEditor/Private/DetailItemNode.cpp b/Engine/Source/Editor/PropertyEditor/Private/DetailItemNode.cpp
--- a/Engine/Source/Editor/PropertyEditor/Private/DetailItemNode.cpp
+++ b/Engine/Source/Editor/PropertyEditor/Private/DetailItemNode.cpp
@@ -22,6 +22,7 @@ FDetailItemNode::FDetailItemNode(const FDetailLayoutCustomization& InCustomizati
 	, ParentGroup(InParentGroup)
 	, IsParentEnabled( InIsParentEnabled )
 	, CachedItemVisibility( EVisibility::Visible )
+	, bForceHidden( false )
 	, bShouldBeVisibleDueToFiltering( false )
 	, bShouldBeVisibleDueToChildFiltering( false )
 	, bTickable( false )
@@ -525,7 +526,7 @@ bool FDetailItemNode::ShouldBeExpanded() const
 ENodeVisibility FDetailItemNode::GetVisibility() const
 {
 	ENodeVisibility Visibility = CachedItemVisibility == EVisibility::Collapsed ? ENodeVisibility::ForcedHidden : ENodeVisibility::Visible;
-	if(Customization.IsHidden())
+	if(Customization.IsHidden() || bForceHidden)
 	{
 		Visibility = ENodeVisibility::ForcedHidden;
 	}
@@ -665,20 +666,6 @@ static bool PassesAllFilters( FDetailItemNode* ItemNode, const FDetailLayoutCust
 	bool bPassesAllFilters = true;
 	
 	TSharedPtr<FPropertyNode> PropertyNodePin = InCustomization.GetPropertyNode();
-	
-	if( PropertyNodePin.IsValid())
-	{
-		if (!InFilter.bShowLooseProperties)
-		{
-			if (FProperty* Property = PropertyNodePin->GetProperty())
-			{
-				if (Property->GetBoolMetaData(NAME_IsLooseMetadata))
-				{
-					return false;
-				}
-			}
-		}
-	}
 
 	if( InFilter.FilterStrings.Num() > 0 || 
 		InFilter.bShowOnlyModified == true || 
@@ -975,6 +962,29 @@ void FDetailItemNode::FilterNode(const FDetailFilter& InFilter)
 		bShouldBeVisibleDueToFiltering = PassesAllFilters(this, Customization, InFilter, ParentGroup.Pin()->GetGroupName().ToString());
 	}
 
+	// set bForceHidden if this node is loose and loose properties are hidden
+	if( TSharedPtr<FPropertyNode> PropertyNodePin = Customization.GetPropertyNode())
+	{
+		if (LIKELY(!InFilter.bShowLooseProperties))
+		{
+			if (FProperty* Property = PropertyNodePin->GetProperty())
+			{
+				if (Property->GetBoolMetaData(NAME_IsLooseMetadata))
+				{
+					bForceHidden = true;
+				}
+			}
+		}
+		if (!bForceHidden && InFilter.ShouldForceHideProperty.IsBound())
+		{
+			if (InFilter.ShouldForceHideProperty.Execute(PropertyNodePin.ToSharedRef()))
+			{
+				bForceHidden = true;
+			}
+		}
+	}
+
+
 	bShouldBeVisibleDueToChildFiltering = false;
 
 	// Filter each child
@@ -987,7 +997,10 @@ void FDetailItemNode::FilterNode(const FDetailFilter& InFilter)
 		// filtered incorrectly because they have no means of discovering if their parents were filtered.
 		if ( bShouldBeVisibleDueToFiltering )
 		{
-			Child->FilterNode(FDetailFilter());
+			FDetailFilter ChildFilter;
+			ChildFilter.bShowLooseProperties = InFilter.bShowLooseProperties; // bShowLooseProperties is inherited from parent regardless
+			ChildFilter.ShouldForceHideProperty = InFilter.ShouldForceHideProperty; // ShouldForceHideProperty is inherited from parent regardless
+			Child->FilterNode(ChildFilter);
 
 			// The child should be visible, but maybe something else has it hidden, check if it's
 			// visible just for safety reasons.
diff --git a/Engine/Source/Editor/PropertyEditor/Private/DetailItemNode.h b/Engine/Source/Editor/PropertyEditor/Private/DetailItemNode.h
--- a/Engine/Source/Editor/PropertyEditor/Private/DetailItemNode.h
+++ b/Engine/Source/Editor/PropertyEditor/Private/DetailItemNode.h
@@ -123,6 +123,8 @@ private:
 	TAttribute<bool> IsParentEnabled;
 	/** Cached visibility of this node */
 	EVisibility CachedItemVisibility;
+	/** If true, this node will be hidden regardless of if it's parent or children would've otherwise overridden the filter result */
+	bool bForceHidden;
 	/** True if this node passes filtering */
 	bool bShouldBeVisibleDueToFiltering;
 	/** True if this node is visible because its children are filtered successfully */
diff --git a/Engine/Source/Editor/PropertyEditor/Private/SDetailsView.cpp b/Engine/Source/Editor/PropertyEditor/Private/SDetailsView.cpp
--- a/Engine/Source/Editor/PropertyEditor/Private/SDetailsView.cpp
+++ b/Engine/Source/Editor/PropertyEditor/Private/SDetailsView.cpp
@@ -76,6 +76,7 @@ void SDetailsView::Construct(const FArguments& InArgs, const FDetailsViewArgs& I
 	}
 
 	CurrentFilter.bShowLooseProperties = InDetailsViewArgs.bShowLooseProperties;
+	CurrentFilter.ShouldForceHideProperty = InDetailsViewArgs.ShouldForceHideProperty;
 
 	ColumnSizeData.SetValueColumnWidth(DetailsViewArgs.ColumnWidth);
 	ColumnSizeData.SetRightColumnMinWidth(DetailsViewArgs.RightColumnMinWidth);
diff --git a/Engine/Source/Editor/PropertyEditor/Public/DetailsViewArgs.h b/Engine/Source/Editor/PropertyEditor/Public/DetailsViewArgs.h
--- a/Engine/Source/Editor/PropertyEditor/Public/DetailsViewArgs.h
+++ b/Engine/Source/Editor/PropertyEditor/Public/DetailsViewArgs.h
@@ -14,6 +14,7 @@ class FUICommandList;
 class IClassViewerFilter;
 class FDetailsViewStyleKey;
 class FDetailsViewStyleKeys;
+class FPropertyNode;
 
 enum class EEditDefaultsOnlyNodeVisibility : uint8
 {
@@ -117,6 +118,10 @@ struct FDetailsViewArgs
 	/** If true, loose properties should be shown. */
 	bool bShowLooseProperties : 1;
 
+	/** Delegate that determines whether a property should be forced hidden - evaluated on panel refresh */
+	DECLARE_DELEGATE_RetVal_OneParam(bool, FShouldForceHideProperty, const TSharedRef<FPropertyNode>&);
+	FShouldForceHideProperty ShouldForceHideProperty;
+
 public:
 
 	FDetailsViewArgs()
diff --git a/Engine/Source/Runtime/CoreUObject/Private/UObject/InstanceDataObjectUtils.cpp b/Engine/Source/Runtime/CoreUObject/Private/UObject/InstanceDataObjectUtils.cpp
--- a/Engine/Source/Runtime/CoreUObject/Private/UObject/InstanceDataObjectUtils.cpp
+++ b/Engine/Source/Runtime/CoreUObject/Private/UObject/InstanceDataObjectUtils.cpp
@@ -677,8 +677,11 @@ namespace UE
 		if (ensureMsgf(ResolvePropertyPathName(Object, Path, ResolvedData, ResolvedProperty, ResolvedOwnerData, ResolvedOwnerProperty),
 			TEXT("Failed to resolve property path name %s"), *WriteToString<256>(Path)))
 		{
-			// only mark properties set if they're in structs/classes
-			if (ResolvedOwnerProperty)
+			if (Path.GetSegmentCount() == 1) // for paths of length == 1, the owner is the Object
+			{
+				MarkPropertySetBySerialization(Object->GetClass(), Object, ResolvedData);
+			}
+			else if (ResolvedOwnerProperty) // only mark properties set if they're in structs/classes
 			{
 				MarkPropertySetBySerialization(ResolvedOwnerProperty->Struct, ResolvedOwnerData, ResolvedData);
 			}
