Support Map key filtering in the details panel

diff --git a/Engine/Source/Editor/PropertyEditor/Private/DetailItemNode.cpp b/Engine/Source/Editor/PropertyEditor/Private/DetailItemNode.cpp
--- a/Engine/Source/Editor/PropertyEditor/Private/DetailItemNode.cpp
+++ b/Engine/Source/Editor/PropertyEditor/Private/DetailItemNode.cpp
@@ -655,6 +655,47 @@ static bool PassesAllFilters( FDetailItemNode* ItemNode, const FDetailLayoutCust
 
 			return FString();
 		}
+
+		static FString GetPropertyNodeKeyFilterString(const FDetailLayoutCustomization& InCustomization, TSharedPtr<FPropertyNode> PropertyNode)
+		{
+			if (PropertyNode.IsValid())
+			{
+				// Is it a container (array, map, set?) - if so, ignore it, we don't care about these, only their inner nodes.
+				if (CastField<FArrayProperty>(PropertyNode->GetProperty()) || CastField<FMapProperty>(PropertyNode->GetProperty()) || CastField<FSetProperty>(PropertyNode->GetProperty()) || CastField<FOptionalProperty>(PropertyNode->GetProperty()))
+				{
+					return FString();
+				}
+
+				// Need to know if parent is a Map though...
+				const FProperty* Property = PropertyNode->GetProperty();
+				FPropertyNode* Parent = PropertyNode->GetParentNode();
+				if (Parent && Property)
+				{
+					const FMapProperty* OuterMapProp = Property->GetOwner<FMapProperty>();
+					if (OuterMapProp)
+					{
+						const FProperty* KeyProperty = OuterMapProp->GetKeyProperty();
+
+						uint8* MapValueAddress = nullptr;
+						FPropertyAccess::Result Result = Parent->GetSingleReadAddress(MapValueAddress);
+						if (Result != FPropertyAccess::Success)
+						{
+							return FString();
+						}
+						FScriptMapHelper MapHelper(OuterMapProp, MapValueAddress);
+						FScriptMapHelper::FIterator Iterator = MapHelper.CreateIterator(PropertyNode->GetArrayIndex());
+						const uint8* PairPtr = MapHelper.GetKeyPtr(Iterator);
+					
+						FString OutString;
+					
+						KeyProperty->ExportText_Direct(OutString, PairPtr, PairPtr, nullptr, PPF_SimpleObjectText);
+						return OutString;
+					}
+				}
+			}
+			return FString();
+		}
+		
 	};
 	
 	auto IsCustomResetToDefaultVisible = [ItemNode, &InCustomization]()
@@ -678,6 +719,9 @@ static bool PassesAllFilters( FDetailItemNode* ItemNode, const FDetailLayoutCust
 		const bool bPassesCategoryFilter = !bSearchFilterIsEmpty && InFilter.bShowAllChildrenIfCategoryMatches ? Local::StringPassesFilter(InFilter, InCategoryName) : false;
 		const bool bPassesValueFilter = !bSearchFilterIsEmpty && Local::StringPassesFilter(InFilter, Local::GetPropertyNodeValueFilterString(InCustomization, PropertyNodePin));
 
+		const FString KeyValue = Local::GetPropertyNodeKeyFilterString(InCustomization, PropertyNodePin);
+		const bool bPassesKeyFilter = !bSearchFilterIsEmpty && Local::StringPassesFilter(InFilter, KeyValue);
+
 		bPassesAllFilters = false;
 		if( PropertyNodePin.IsValid() && !PropertyNodePin->AsCategoryNode())
 		{
@@ -685,7 +729,7 @@ static bool PassesAllFilters( FDetailItemNode* ItemNode, const FDetailLayoutCust
 			const bool bIsSeenDueToFiltering = PropertyNodePin->HasNodeFlags(EPropertyNodeFlags::IsSeenDueToFiltering) != 0;
 			const bool bIsParentSeenDueToFiltering = PropertyNodePin->HasNodeFlags(EPropertyNodeFlags::IsParentSeenDueToFiltering) != 0;
 
-			const bool bPassesTextFilter = bPassesCategoryFilter || bPassesValueFilter || Local::StringPassesFilter(InFilter, InCustomization.GetFilterTextString().ToString());
+			const bool bPassesTextFilter = bPassesCategoryFilter || bPassesValueFilter || bPassesKeyFilter || Local::StringPassesFilter(InFilter, InCustomization.GetFilterTextString().ToString());
 			const bool bPassesSearchFilter = bPassesTextFilter || bSearchFilterIsEmpty || ( bIsNotBeingFiltered || bIsSeenDueToFiltering || bIsParentSeenDueToFiltering );
 
 			bool bPassesModifiedFilter = true;
