Fix performance issue with nested InstancedStructs
* Avoid duplicate work in GetValueBaseAddress for FItemPropertyNode, this function is recursive, the duplicate work branches the recursion and can be expensive for deeply nested property node trees.
* Reimplement calls for GetOwnerPackages and GetInstacesNum for StructurePropertyNode without using GetInstances() - these functions are called from InstancedStructProvider's EnumerateInstances already and cause exponential walk up the property tree

diff --git a/Engine/Source/Editor/PropertyEditor/Private/ItemPropertyNode.cpp b/Engine/Source/Editor/PropertyEditor/Private/ItemPropertyNode.cpp
--- a/Engine/Source/Editor/PropertyEditor/Private/ItemPropertyNode.cpp
+++ b/Engine/Source/Editor/PropertyEditor/Private/ItemPropertyNode.cpp
@@ -70,7 +70,7 @@ uint8* FItemPropertyNode::GetValueBaseAddress(uint8* StartAddress, bool bIsSpars
 		}
 		else
 		{
-			uint8* ValueAddress = ParentNode->GetValueAddress(StartAddress, bIsSparseData, bIsStruct);
+			uint8* ValueAddress = ValueBaseAddress;
 			if (ValueAddress != nullptr && ParentNode->GetProperty() != MyProperty)
 			{
 				// if this is not a fixed size array (in which the parent property and this property are the same), we need to offset from the property (otherwise, the parent already did that for us)
diff --git a/Engine/Source/Editor/PropertyEditor/Private/StructurePropertyNode.h b/Engine/Source/Editor/PropertyEditor/Private/StructurePropertyNode.h
--- a/Engine/Source/Editor/PropertyEditor/Private/StructurePropertyNode.h
+++ b/Engine/Source/Editor/PropertyEditor/Private/StructurePropertyNode.h
@@ -6,6 +6,7 @@
 #include "PropertyNode.h"
 #include "UObject/StructOnScope.h"
 #include "IStructureDataProvider.h"
+#include "ObjectPropertyNode.h"
 
 //-----------------------------------------------------------------------------
 //	FStructPropertyNode - Used for the root and various sub-nodes
@@ -90,13 +91,59 @@ public:
 	{
 		if (StructProvider)
 		{
-			TArray<TSharedPtr<FStructOnScope>> Instances;
-			StructProvider->GetInstances(Instances, WeakCachedBaseStruct.Get());
-
-			for (TSharedPtr<FStructOnScope>& Instance : Instances)
+			// Walk up until we find the objects that contain this struct property to get the packages
+			const FPropertyNode* Parent = this;
+		
+			while(Parent)
 			{
-				// Returning null for invalid instances, to match instance count.
-				OutPackages.Add(Instance.IsValid() ? Instance->GetPackage() : nullptr);
+				const FComplexPropertyNode* ComplexParent = Parent->FindComplexParent();
+				if (!ensureMsgf(ComplexParent != nullptr, TEXT("Expected to find a complex parent")))
+				{
+					return;
+				}
+
+				if (const FObjectPropertyNode* ObjectNode = ComplexParent->AsObjectNode())
+				{
+					for (int32 ObjectIndex = 0, ObjectCount = ObjectNode->GetNumObjects(); ObjectIndex < ObjectCount; ++ObjectIndex)
+					{
+						const UPackage* Package = ObjectNode->GetUPackage(ObjectIndex);
+						OutPackages.Add(const_cast<UPackage*>(Package));
+					}
+					break;
+				}
+
+				if (const FStructurePropertyNode* StructNode = ComplexParent->AsStructureNode())
+				{
+					if (StructNode->StructProvider && StructNode->StructProvider->IsPropertyIndirection())
+					{
+						// Skip this and keep walking up to the next parent. This is assumed to be a struct that is pointed
+						// to be a parent indirection.
+						// Note: in InstancedStruct case, will re-enter GetOwnerPackages during EnumerateInstances if calling GetInstances here.
+						// We want to avoid that otherwise it can cause poor performance as recursion branches into multiple recursions on each level.
+						Parent = ComplexParent->GetParentNode();
+						if (Parent == nullptr)
+						{
+							// If no owning property, then there is no package.
+							// Returning null to match instance count.
+							OutPackages.Add(nullptr);
+							break;
+						}
+					}
+					else
+					{
+						if (StructNode->StructProvider)
+						{
+							TArray<TSharedPtr<FStructOnScope>> Instances;
+							StructNode->StructProvider->GetInstances(Instances, StructNode->WeakCachedBaseStruct.Get());
+							for (int32 Index = 0, ObjectCount = Instances.Num(); Index < ObjectCount; ++Index)
+							{
+								UPackage* Package = Instances[Index]->GetPackage();
+								OutPackages.Add(Package);
+							}
+						}
+						break;
+					}
+				}
 			}
 		}
 	}
@@ -159,12 +206,38 @@ public:
 	}
 	virtual int32 GetInstancesNum() const override
 	{
-		if (StructProvider)
+		// Can't get instance count directly from standalone structures, need to walk to the next parent that is an object
+		// to get number of instances from that
+		// Note: This uses iteration over StructProvider->GetInstances and counting returned instances since InstancedStructProvider
+		// calls GetInstancesNum from within it's implementation of GetInstances, leading to exponential recursion up the
+		// property tree.
+		const FComplexPropertyNode* CurrentNode = this;	
+	
+		while(CurrentNode)
 		{
-			TArray<TSharedPtr<FStructOnScope>> Instances;
-			StructProvider->GetInstances(Instances, WeakCachedBaseStruct.Get());
-			
-			return Instances.Num();
+			if (const FObjectPropertyNode* ObjectNode = CurrentNode->AsObjectNode())
+			{
+				// Found owning UObject
+				const int32 ObjectCount = ObjectNode->GetInstancesNum();
+				return ObjectCount;
+			}
+			else if (const FStructurePropertyNode* StructNode = CurrentNode->AsStructureNode())
+			{
+				TSharedPtr<IStructureDataProvider> TempStructProvider = StructNode->GetStructProvider();
+				if (TempStructProvider->IsPropertyIndirection())
+				{
+					// If the struct provider is marked as property indirection, it is assumed that it handles indirection between it's
+					// parent property, and some data inside that property (e.g. FInstancedStruct).
+					const FPropertyNode* ParentNode = CurrentNode->GetParentNode();
+					CurrentNode = ParentNode->FindComplexParent();
+				}
+				else
+				{
+					TArray<TSharedPtr<FStructOnScope>> Instances;
+					TempStructProvider->GetInstances(Instances, StructNode->WeakCachedBaseStruct.Get());
+					return Instances.Num();
+				}
+			}
 		}
 		return 0;
 		
