PropertyEditor: Fix issue where property change events on array members in array parent elements were indistinguishable between parent & child.

(Fix from [REDACTED_HANDLE].)

diff --git a/Engine/Source/Editor/PropertyEditor/Private/PropertyNode.cpp b/Engine/Source/Editor/PropertyEditor/Private/PropertyNode.cpp
--- a/Engine/Source/Editor/PropertyEditor/Private/PropertyNode.cpp
+++ b/Engine/Source/Editor/PropertyEditor/Private/PropertyNode.cpp
@@ -3198,6 +3198,8 @@ void FPropertyNode::NotifyPostChange( FPropertyChangedEvent& InPropertyChangedEv
 					FPropertyChangedEvent ChangedEvent = InPropertyChangedEvent;
 					if (CurProperty != InPropertyChangedEvent.Property)
 					{
+						// Parent object node property. Reset property and leave the event type as unspecified since we dont pass in the exact leaf property.
+						ChangedEvent.ChangeType = EPropertyChangeType::Unspecified;
 						ChangedEvent.Property = CurProperty;
 						ChangedEvent.MemberProperty = CurProperty;
 					}
