Property Editor - Fixed crash when using the DetailView on UStructs with a property that requires a tickable detail item node (e.g. one that uses EditConditionHides).

#prefligth 663e67b85f57243415b61b6a
#tests Tried the crash situation after the change and I couldn't reproduce it.
#rnx

diff --git a/Engine/Source/Editor/PropertyEditor/Private/DetailCategoryBuilderImpl.cpp b/Engine/Source/Editor/PropertyEditor/Private/DetailCategoryBuilderImpl.cpp
--- a/Engine/Source/Editor/PropertyEditor/Private/DetailCategoryBuilderImpl.cpp
+++ b/Engine/Source/Editor/PropertyEditor/Private/DetailCategoryBuilderImpl.cpp
@@ -246,6 +246,17 @@ FDetailCategoryImpl::FDetailCategoryImpl(FName InCategoryName, TSharedRef<FDetai
 
 FDetailCategoryImpl::~FDetailCategoryImpl()
 {
+	// The children won't be able to pin the weak pointer of their parent (this object) to remove themselves from the list of tickables
+	// because at this point the parent's shared pointer reference counter has already reached 0 (the execution has reached this destructor).
+	// Therefore the tickable nodes need to be manually removed from that list to prevent these dead node pointers remaining in the layout object.
+	for (TSharedRef<FDetailTreeNode>& SimpleChildNode : SimpleChildNodes)
+	{
+		RemoveTickableNode(*SimpleChildNode);
+	}
+	for (TSharedRef<FDetailTreeNode>& AdvancedChildNode : AdvancedChildNodes)
+	{
+		RemoveTickableNode(*AdvancedChildNode);
+	}
 }
 
 
diff --git a/Engine/Source/Editor/PropertyEditor/Private/DetailItemNode.cpp b/Engine/Source/Editor/PropertyEditor/Private/DetailItemNode.cpp
--- a/Engine/Source/Editor/PropertyEditor/Private/DetailItemNode.cpp
+++ b/Engine/Source/Editor/PropertyEditor/Private/DetailItemNode.cpp
@@ -46,7 +46,10 @@ void FDetailItemNode::Initialize()
 	{
 		// The node needs to be ticked because it has widgets that can dynamically come and go
 		bTickable = true;
-		ParentCategory.Pin()->AddTickableNode( *this );
+		if (const TSharedPtr<FDetailCategoryImpl> ParentCategoryPtr = ParentCategory.Pin())
+		{
+			ParentCategoryPtr->AddTickableNode(*this);
+		}
 	}
 
 	if( Customization.HasPropertyNode() )
