UE-170945 - Add a Play/Stop button to USoundBase when viewing default value

diff --git a/Engine/Source/Editor/PropertyEditor/Private/UserInterface/PropertyEditor/SPropertyEditorAsset.cpp b/Engine/Source/Editor/PropertyEditor/Private/UserInterface/PropertyEditor/SPropertyEditorAsset.cpp
--- a/Engine/Source/Editor/PropertyEditor/Private/UserInterface/PropertyEditor/SPropertyEditorAsset.cpp
+++ b/Engine/Source/Editor/PropertyEditor/Private/UserInterface/PropertyEditor/SPropertyEditorAsset.cpp
@@ -3,6 +3,7 @@
 #include "UserInterface/PropertyEditor/SPropertyEditorAsset.h"
 #include "Engine/Texture.h"
 #include "Engine/SkeletalMesh.h"
+#include "Components/AudioComponent.h"
 #include "Components/StaticMeshComponent.h"
 #include "Engine/StaticMesh.h"
 #include "Engine/LevelScriptActor.h"
@@ -21,6 +22,7 @@
 #include "IAssetTypeActions.h"
 #include "AssetToolsModule.h"
 #include "SAssetDropTarget.h"
+#include "Widgets/Input/SButton.h"
 #include "AssetRegistry/AssetRegistryModule.h"
 #include "Engine/Selection.h"
 #include "ObjectPropertyNode.h"
@@ -40,6 +42,7 @@
 #include "PropertyEditorConstants.h"
 #include "PropertyEditorUtils.h"
 #include "Misc/EditorPathHelper.h"
+#include "Sound/SoundBase.h"
 
 #define LOCTEXT_NAMESPACE "PropertyEditor"
 
@@ -705,6 +708,91 @@ void SPropertyEditorAsset::Construct(const FArguments& InArgs, const TSharedPtr<
 		];
 	}
 
+	auto OnGetDisplayBrushLambda = [this]() -> const FSlateBrush*
+		{
+			FObjectOrAssetData Value;
+			GetValue(Value);
+			if (USoundBase* SoundBase = Cast<USoundBase>(Value.Object))
+			{
+				bool bIsPlaying = false;
+				const UAudioComponent* PreviewComp = GEditor->GetPreviewAudioComponent();
+				if (PreviewComp && PreviewComp->Sound && PreviewComp->IsPlaying())
+				{
+					if (PreviewComp->Sound->GetFName() == SoundBase->GetFName())
+					{
+						if (PreviewComp->Sound->GetOutermost()->GetFName() == SoundBase->GetPackage()->GetFName())
+						{
+							bIsPlaying = true;
+						}
+					}
+				}
+
+				if (bIsPlaying)
+				{
+					return FAppStyle::GetBrush("MediaAsset.AssetActions.Stop.Small");
+				}
+			}
+			return FAppStyle::GetBrush("MediaAsset.AssetActions.Play.Small");
+		};
+
+	auto OnGetVisibilityLambda = [this]() -> EVisibility
+		{
+			FObjectOrAssetData Value;
+			GetValue(Value);
+			if (USoundBase* SoundBase = Cast<USoundBase>(Value.Object))
+			{
+				return EVisibility::Visible;
+			}
+			return EVisibility::Collapsed;
+		};
+
+	auto OnClickedLambda = [this]() -> FReply
+		{
+			FObjectOrAssetData Value;
+			GetValue(Value);
+			if (USoundBase* SoundBase = Cast<USoundBase>(Value.Object))
+			{
+				bool bIsPlaying = false;
+				const UAudioComponent* PreviewComp = GEditor->GetPreviewAudioComponent();
+				if (PreviewComp && PreviewComp->Sound && PreviewComp->IsPlaying())
+				{
+					if (PreviewComp->Sound->GetFName() == SoundBase->GetFName())
+					{
+						if (PreviewComp->Sound->GetOutermost()->GetFName() == SoundBase->GetPackage()->GetFName())
+						{
+							bIsPlaying = true;
+						}
+					}
+				}
+
+				if (bIsPlaying)
+				{
+					GEditor->ResetPreviewAudioComponent();
+				}
+				else
+				{
+					GEditor->PlayPreviewSound(SoundBase);
+				}
+			}
+			return FReply::Handled();
+		};
+
+	ButtonBox->AddSlot()
+	[
+		SNew(SButton)
+			.HAlign(HAlign_Left)
+			.ButtonStyle(FAppStyle::Get(), "HoverHintOnly")
+			.Cursor(EMouseCursor::Default)
+			.ForegroundColor(FSlateColor::UseForeground())
+			.IsFocusable(false)
+			.Visibility_Lambda(OnGetVisibilityLambda)
+			[
+				SNew(SImage)
+					.Image_Lambda(OnGetDisplayBrushLambda)
+			]
+			.OnClicked_Lambda(OnClickedLambda)
+	];
+
 	if( bIsActor )
 	{
 		TSharedRef<SWidget> ActorPicker = PropertyCustomizationHelpers::MakeInteractiveActorPicker( FOnGetAllowedClasses::CreateSP(this, &SPropertyEditorAsset::OnGetAllowedClasses), FOnShouldFilterActor::CreateSP(this, &SPropertyEditorAsset::IsFilteredActor), FOnActorSelected::CreateSP(this, &SPropertyEditorAsset::OnActorSelected));
