Introduced IDetailPropertyChildrenCustomizationHandler - a way to customize properties children in a DetailsView (contributed)

diff --git a/Engine/Source/Editor/PropertyEditor/Private/DetailPropertyRow.cpp b/Engine/Source/Editor/PropertyEditor/Private/DetailPropertyRow.cpp
--- a/Engine/Source/Editor/PropertyEditor/Private/DetailPropertyRow.cpp
+++ b/Engine/Source/Editor/PropertyEditor/Private/DetailPropertyRow.cpp
@@ -15,6 +15,7 @@
 #include "StructurePropertyNode.h"
 #include "Modules/ModuleManager.h"
 #include "Widgets/Layout/SSpacer.h"
+#include "IDetailPropertyChildrenCustomizationHandler.h"
 
 #include "UObject/PropertyOptional.h"
 
@@ -390,7 +391,13 @@ void FDetailPropertyRow::OnItemNodeInitialized( TSharedRef<FDetailCategoryImpl>
 		}
 	}
 
-	if( bShowCustomPropertyChildren && CustomTypeInterface.IsValid() )
+	IDetailPropertyChildrenCustomizationHandler* CustomizationHandler = InParentCategory->GetDetailsView()->GetChildrenCustomizationHandler().Get();
+	if (CustomizationHandler && CustomizationHandler->ShouldCustomizeChildren(PropertyHandle.ToSharedRef()))
+	{
+		PropertyTypeLayoutBuilder = MakeShared<FCustomChildrenBuilder>(InParentCategory, InParentGroup);
+		CustomizationHandler->CustomizeChildren(*PropertyTypeLayoutBuilder, PropertyHandle);
+	}
+	else if( bShowCustomPropertyChildren && CustomTypeInterface.IsValid() )
 	{
 		PropertyTypeLayoutBuilder = MakeShared<FCustomChildrenBuilder>(InParentCategory, InParentGroup);
 
diff --git a/Engine/Source/Editor/PropertyEditor/Private/SDetailsViewBase.cpp b/Engine/Source/Editor/PropertyEditor/Private/SDetailsViewBase.cpp
--- a/Engine/Source/Editor/PropertyEditor/Private/SDetailsViewBase.cpp
+++ b/Engine/Source/Editor/PropertyEditor/Private/SDetailsViewBase.cpp
@@ -757,6 +757,11 @@ void SDetailsViewBase::SetExtensionHandler(TSharedPtr<class IDetailPropertyExten
 	ExtensionHandler = InExtensionHandler;
 }
 
+void SDetailsViewBase::SetChildrenCustomizationHandler(TSharedPtr<IDetailPropertyChildrenCustomizationHandler> InChildrenHandler)
+{
+	ChildrenCustomizationHandler = InChildrenHandler;
+}
+
 void SDetailsViewBase::SetGenericLayoutDetailsDelegate(FOnGetDetailCustomizationInstance OnGetGenericDetails)
 {
 	GenericLayoutDelegate = OnGetGenericDetails;
diff --git a/Engine/Source/Editor/PropertyEditor/Private/SDetailsViewBase.h b/Engine/Source/Editor/PropertyEditor/Private/SDetailsViewBase.h
--- a/Engine/Source/Editor/PropertyEditor/Private/SDetailsViewBase.h
+++ b/Engine/Source/Editor/PropertyEditor/Private/SDetailsViewBase.h
@@ -36,6 +36,7 @@ class IClassViewerFilter;
 class IDetailCustomization;
 class IDetailKeyframeHandler;
 class IDetailPropertyExtensionHandler;
+class IDetailPropertyChildrenCustomizationHandler;
 class IPropertyGenerationUtilities;
 class IPropertyUtilities;
 class SDetailNameArea;
@@ -109,6 +110,8 @@ public:
 	virtual TSharedPtr<IDetailKeyframeHandler> GetKeyframeHandler() const override { return KeyframeHandler; }
 	virtual void SetExtensionHandler(TSharedPtr<IDetailPropertyExtensionHandler> InExtensionHandler) override;
 	virtual TSharedPtr<IDetailPropertyExtensionHandler> GetExtensionHandler() const override { return ExtensionHandler; }
+	virtual void SetChildrenCustomizationHandler(TSharedPtr<IDetailPropertyChildrenCustomizationHandler> InChildrenHandler) override;
+	virtual TSharedPtr<IDetailPropertyChildrenCustomizationHandler> GetChildrenCustomizationHandler() const override { return ChildrenCustomizationHandler; }
 	virtual bool IsPropertyVisible(const struct FPropertyAndParent& PropertyAndParent) const override;
 	virtual bool IsPropertyReadOnly(const struct FPropertyAndParent& PropertyAndParent) const override;
 	virtual bool IsCustomRowVisible(FName InRowName, FName InParentName) const override;
@@ -481,6 +484,8 @@ protected:
 	TSharedPtr<IDetailKeyframeHandler> KeyframeHandler;
 	/** Property extension handler returns additional UI to apply after the customization is applied to the property. */
 	TSharedPtr<IDetailPropertyExtensionHandler> ExtensionHandler;
+	/** Property children customization handler allows to customize property's children. */
+	TSharedPtr<IDetailPropertyChildrenCustomizationHandler> ChildrenCustomizationHandler;
 	/** The tree node that is currently highlighted, may be none. */
 	TWeakPtr<FDetailTreeNode> CurrentlyHighlightedNode;
 	/** The list of nodes whose widgets should be animating. */
diff --git a/Engine/Source/Editor/PropertyEditor/Public/IDetailPropertyChildrenCustomizationHandler.h b/Engine/Source/Editor/PropertyEditor/Public/IDetailPropertyChildrenCustomizationHandler.h
new file mode 100644
--- /dev/null
+++ b/Engine/Source/Editor/PropertyEditor/Public/IDetailPropertyChildrenCustomizationHandler.h
@@ -0,0 +1,21 @@
+// Copyright Epic Games, Inc. All Rights Reserved.
+
+#pragma once
+
+#include "Templates/SharedPointer.h"
+
+class IDetailChildrenBuilder;
+class IPropertyHandle;
+
+/* Customizes properties children in a detail view. */
+class IDetailPropertyChildrenCustomizationHandler
+{
+public:
+	virtual ~IDetailPropertyChildrenCustomizationHandler() = default;
+
+	/* Returns true if the handler should override provided property's children. */
+	virtual bool ShouldCustomizeChildren(TSharedRef<IPropertyHandle> InPropertyHandle)  = 0;
+
+	/* Gives a chance to customize provided property's children. */
+	virtual void CustomizeChildren(IDetailChildrenBuilder& ChildrenBuilder, TSharedPtr<IPropertyHandle> PropertyHandle) {}
+};
diff --git a/Engine/Source/Editor/PropertyEditor/Public/IDetailsView.h b/Engine/Source/Editor/PropertyEditor/Public/IDetailsView.h
--- a/Engine/Source/Editor/PropertyEditor/Public/IDetailsView.h
+++ b/Engine/Source/Editor/PropertyEditor/Public/IDetailsView.h
@@ -18,6 +18,7 @@ class FUICommandList;
 class IClassViewerFilter;
 class IDetailKeyframeHandler;
 class IDetailPropertyExtensionHandler;
+class IDetailPropertyChildrenCustomizationHandler;
 class IDetailRootObjectCustomization;
 class IPropertyTypeIdentifier;
 class FDetailsViewObjectFilter;
@@ -168,6 +169,9 @@ public:
 	virtual void SetExtensionHandler(TSharedPtr<IDetailPropertyExtensionHandler> InExtensionHandler) = 0;
 	virtual TSharedPtr<IDetailPropertyExtensionHandler> GetExtensionHandler() const = 0;
 
+	virtual void SetChildrenCustomizationHandler(TSharedPtr<IDetailPropertyChildrenCustomizationHandler> InExtensionHandler) = 0;
+	virtual TSharedPtr<IDetailPropertyChildrenCustomizationHandler> GetChildrenCustomizationHandler() const = 0;
+
 	/**
 	 * @return true if property editing is enabled (based on the FIsPropertyEditingEnabled delegate)
 	 */ 
