Details Panel: Added missing null check when getting IDetailPropertyChildrenCustomizationHandler.

diff --git a/Engine/Source/Editor/PropertyEditor/Private/DetailPropertyRow.cpp b/Engine/Source/Editor/PropertyEditor/Private/DetailPropertyRow.cpp
--- a/Engine/Source/Editor/PropertyEditor/Private/DetailPropertyRow.cpp
+++ b/Engine/Source/Editor/PropertyEditor/Private/DetailPropertyRow.cpp
@@ -391,7 +391,7 @@ void FDetailPropertyRow::OnItemNodeInitialized( TSharedRef<FDetailCategoryImpl>
 		}
 	}
 
-	IDetailPropertyChildrenCustomizationHandler* CustomizationHandler = InParentCategory->GetDetailsView()->GetChildrenCustomizationHandler().Get();
+	IDetailPropertyChildrenCustomizationHandler* CustomizationHandler = InParentCategory->GetDetailsView() ? InParentCategory->GetDetailsView()->GetChildrenCustomizationHandler().Get() : nullptr;
 	if (CustomizationHandler && CustomizationHandler->ShouldCustomizeChildren(PropertyHandle.ToSharedRef()))
 	{
 		PropertyTypeLayoutBuilder = MakeShared<FCustomChildrenBuilder>(InParentCategory, InParentGroup);
