Fixed a bunch of constness issues with FProperty*, for an upcoming change to TPropertyValueIterator.

diff --git a/Engine/Plugins/Enterprise/VariantManager/Source/VariantManager/Private/VariantManager.cpp b/Engine/Plugins/Enterprise/VariantManager/Source/VariantManager/Private/VariantManager.cpp
--- a/Engine/Plugins/Enterprise/VariantManager/Source/VariantManager/Private/VariantManager.cpp
+++ b/Engine/Plugins/Enterprise/VariantManager/Source/VariantManager/Private/VariantManager.cpp
@@ -516,7 +516,7 @@ TArray<UPropertyValue*> FVariantManager::CreatePropertyCaptures(const TArray<TSh
 			}
 			default: // Generic
 			{
-				FProperty* LeafProp = PropAndDisplay->Prop.GetLeafMostProperty().Property.Get();
+				const FProperty* LeafProp = PropAndDisplay->Prop.GetLeafMostProperty().Property.Get();
 
 				if (LeafProp && LeafProp->GetClass()->IsChildOf(FSoftObjectProperty::StaticClass()))
 				{
diff --git a/Engine/Plugins/Experimental/InstanceDataObjectFixupTool/Source/Private/InstanceDataObjectFixupDetailCustomization.cpp b/Engine/Plugins/Experimental/InstanceDataObjectFixupTool/Source/Private/InstanceDataObjectFixupDetailCustomization.cpp
--- a/Engine/Plugins/Experimental/InstanceDataObjectFixupTool/Source/Private/InstanceDataObjectFixupDetailCustomization.cpp
+++ b/Engine/Plugins/Experimental/InstanceDataObjectFixupTool/Source/Private/InstanceDataObjectFixupDetailCustomization.cpp
@@ -231,8 +231,8 @@ TSharedRef<SWidget> FInstanceDataObjectNameWidgetOverride::GeneratePropertyRedir
 	{
 		for (const FPropertyPath& Option : RedirectOptions)
 		{
-			FProperty* ThisProperty = Path.GetLeafMostProperty().Property.Get();
-			FProperty* OptionProperty = Option.GetLeafMostProperty().Property.Get();
+			const FProperty* ThisProperty = Path.GetLeafMostProperty().Property.Get();
+			const FProperty* OptionProperty = Option.GetLeafMostProperty().Property.Get();
 			if (ThisProperty->GetFName() == OptionProperty->GetFName())
 			{
 				if (OptionProperty->SameType(ThisProperty))
@@ -253,8 +253,8 @@ TSharedRef<SWidget> FInstanceDataObjectNameWidgetOverride::GeneratePropertyRedir
 	{
 		for (const FPropertyPath& Option : RedirectOptions)
 		{
-			FProperty* ThisProperty = Path.GetLeafMostProperty().Property.Get();
-			FProperty* OptionProperty = Option.GetLeafMostProperty().Property.Get();
+			const FProperty* ThisProperty = Path.GetLeafMostProperty().Property.Get();
+			const FProperty* OptionProperty = Option.GetLeafMostProperty().Property.Get();
 			if (ThisProperty->GetFName() == OptionProperty->GetFName())
 			{
 				continue; // handled in the "move" category
@@ -276,8 +276,8 @@ TSharedRef<SWidget> FInstanceDataObjectNameWidgetOverride::GeneratePropertyRedir
 	{
 		for (const FPropertyPath& Option : RedirectOptions)
 		{
-			FProperty* ThisProperty = Path.GetLeafMostProperty().Property.Get();
-			FProperty* OptionProperty = Option.GetLeafMostProperty().Property.Get();
+			const FProperty* ThisProperty = Path.GetLeafMostProperty().Property.Get();
+			const FProperty* OptionProperty = Option.GetLeafMostProperty().Property.Get();
 			if (OptionProperty->SameType(ThisProperty))
 			{
 				continue; // same type handled above
@@ -310,8 +310,8 @@ TSharedRef<SWidget> FInstanceDataObjectNameWidgetOverride::GeneratePropertyRedir
 	{
 		for (const FPropertyPath& Option : RedirectOptions)
 		{
-			FProperty* ThisProperty = Path.GetLeafMostProperty().Property.Get();
-			FProperty* OptionProperty = Option.GetLeafMostProperty().Property.Get();
+			const FProperty* ThisProperty = Path.GetLeafMostProperty().Property.Get();
+			const FProperty* OptionProperty = Option.GetLeafMostProperty().Property.Get();
 			if (OptionProperty->SameType(ThisProperty))
 			{
 				continue; // same type handled above
diff --git a/Engine/Plugins/Experimental/InstanceDataObjectFixupTool/Source/Private/InstanceDataObjectFixupPanel.cpp b/Engine/Plugins/Experimental/InstanceDataObjectFixupTool/Source/Private/InstanceDataObjectFixupPanel.cpp
--- a/Engine/Plugins/Experimental/InstanceDataObjectFixupTool/Source/Private/InstanceDataObjectFixupPanel.cpp
+++ b/Engine/Plugins/Experimental/InstanceDataObjectFixupTool/Source/Private/InstanceDataObjectFixupPanel.cpp
@@ -687,15 +687,15 @@ void FInstanceDataObjectFixupPanel::RedirectPropertyHelper(const FPropertyPath&
 		if (DestinationProperty)
 		{
 			if (DestinationProperty->HasAnyPropertyFlags(CPF_Transient) != Info->bWasTransient)
-            {
-            	// toggle transient flag if needed
-            	DestinationProperty->PropertyFlags ^= CPF_Transient;
-            }
-            if (!Info->bWasHidden)
-            {
-            	DestinationProperty->RemoveMetaData(TEXT("Hidden"));
-            }
-            DestinationProperty->RemoveMetaData(TEXT("Redirected"));
+			{
+				// toggle transient flag if needed
+				DestinationProperty->PropertyFlags ^= CPF_Transient;
+			}
+			if (!Info->bWasHidden)
+			{
+				DestinationProperty->RemoveMetaData(TEXT("Hidden"));
+			}
+			DestinationProperty->RemoveMetaData(TEXT("Redirected"));
 		}
 		
 		
@@ -743,7 +743,7 @@ void FInstanceDataObjectFixupPanel::RedirectPropertyHelper(const FPropertyPath&
 
 void FInstanceDataObjectFixupPanel::RedirectProperty(const FPropertyPath& From, const FPropertyPath& To)
 {
-	FProperty* SourceProperty = From.GetLeafMostProperty().Property.Get();
+	const FProperty* SourceProperty = From.GetLeafMostProperty().Property.Get();
 	check(SourceProperty);
 	FProperty* DestinationProperty = To.IsValid() ? To.GetLeafMostProperty().Property.Get() : nullptr;
 	
diff --git a/Engine/Plugins/Media/MediaCompositing/Source/MediaCompositingEditor/Private/Sequencer/MediaPlayerPropertyTrackEditor.cpp b/Engine/Plugins/Media/MediaCompositing/Source/MediaCompositingEditor/Private/Sequencer/MediaPlayerPropertyTrackEditor.cpp
--- a/Engine/Plugins/Media/MediaCompositing/Source/MediaCompositingEditor/Private/Sequencer/MediaPlayerPropertyTrackEditor.cpp
+++ b/Engine/Plugins/Media/MediaCompositing/Source/MediaCompositingEditor/Private/Sequencer/MediaPlayerPropertyTrackEditor.cpp
@@ -102,7 +102,7 @@ void FMediaPlayerPropertyTrackEditor::OnAnimatedPropertyChanged(const FPropertyC
 		return;
 	}
 
-	FProperty* Property = PropertyChangedParams.PropertyPath.GetLeafMostProperty().Property.Get();
+	const FProperty* Property = PropertyChangedParams.PropertyPath.GetLeafMostProperty().Property.Get();
 	if (!Property)
 	{
 		return;
diff --git a/Engine/Source/Editor/ConfigEditor/Private/PropertyVisualization/ConfigPropertyColumn.cpp b/Engine/Source/Editor/ConfigEditor/Private/PropertyVisualization/ConfigPropertyColumn.cpp
--- a/Engine/Source/Editor/ConfigEditor/Private/PropertyVisualization/ConfigPropertyColumn.cpp
+++ b/Engine/Source/Editor/ConfigEditor/Private/PropertyVisualization/ConfigPropertyColumn.cpp
@@ -35,7 +35,7 @@ bool FConfigPropertyCustomColumn::Supports(const TSharedRef< IPropertyTableColum
 		if (PropertyPath.IsValid() && PropertyPath->GetNumProperties() > 0)
 		{
 			const FPropertyInfo& PropertyInfo = PropertyPath->GetRootProperty();
-			FProperty* Property = PropertyInfo.Property.Get();
+			const FProperty* Property = PropertyInfo.Property.Get();
 			IsSupported = Property->GetFName() == TEXT("ExternalProperty");
 		}
 	}
diff --git a/Engine/Source/Editor/ConfigEditor/Private/PropertyVisualization/ConfigPropertyConfigFileStateColumn.cpp b/Engine/Source/Editor/ConfigEditor/Private/PropertyVisualization/ConfigPropertyConfigFileStateColumn.cpp
--- a/Engine/Source/Editor/ConfigEditor/Private/PropertyVisualization/ConfigPropertyConfigFileStateColumn.cpp
+++ b/Engine/Source/Editor/ConfigEditor/Private/PropertyVisualization/ConfigPropertyConfigFileStateColumn.cpp
@@ -96,7 +96,7 @@ bool FConfigPropertyConfigFileStateCustomColumn::Supports(const TSharedRef< IPro
 		if (PropertyPath.IsValid() && PropertyPath->GetNumProperties() > 0)
 		{
 			const FPropertyInfo& PropertyInfo = PropertyPath->GetRootProperty();
-			FProperty* Property = PropertyInfo.Property.Get();
+			const FProperty* Property = PropertyInfo.Property.Get();
 			if (SupportedProperty == Property)
 			{
 				IsSupported = true;
diff --git a/Engine/Source/Editor/DeviceProfileEditor/Private/DeviceProfileConsoleVariableColumn.cpp b/Engine/Source/Editor/DeviceProfileEditor/Private/DeviceProfileConsoleVariableColumn.cpp
--- a/Engine/Source/Editor/DeviceProfileEditor/Private/DeviceProfileConsoleVariableColumn.cpp
+++ b/Engine/Source/Editor/DeviceProfileEditor/Private/DeviceProfileConsoleVariableColumn.cpp
@@ -156,7 +156,7 @@ bool FDeviceProfileConsoleVariableColumn::Supports(const TSharedRef< IPropertyTa
 		if( PropertyPath.IsValid() && PropertyPath->GetNumProperties() > 0 )
 		{
 			const FPropertyInfo& PropertyInfo = PropertyPath->GetRootProperty();
-			FProperty* Property = PropertyInfo.Property.Get();
+			const FProperty* Property = PropertyInfo.Property.Get();
 			if (Property->GetName() == TEXT("CVars") && Property->IsA(FArrayProperty::StaticClass()))
 			{
 				return true;
diff --git a/Engine/Source/Editor/DeviceProfileEditor/Private/DeviceProfileTextureLODSettingsColumn.cpp b/Engine/Source/Editor/DeviceProfileEditor/Private/DeviceProfileTextureLODSettingsColumn.cpp
--- a/Engine/Source/Editor/DeviceProfileEditor/Private/DeviceProfileTextureLODSettingsColumn.cpp
+++ b/Engine/Source/Editor/DeviceProfileEditor/Private/DeviceProfileTextureLODSettingsColumn.cpp
@@ -156,7 +156,7 @@ bool FDeviceProfileTextureLODSettingsColumn::Supports(const TSharedRef< IPropert
 		if( PropertyPath.IsValid() && PropertyPath->GetNumProperties() > 0 )
 		{
 			const FPropertyInfo& PropertyInfo = PropertyPath->GetRootProperty();
-			FProperty* Property = PropertyInfo.Property.Get();
+			const FProperty* Property = PropertyInfo.Property.Get();
 			if (Property->GetName() == TEXT("TextureLODGroups") && Property->IsA(FArrayProperty::StaticClass()))
 			{
 				return true;
diff --git a/Engine/Source/Editor/MovieSceneTools/Private/TrackEditors/ObjectPropertyTrackEditor.cpp b/Engine/Source/Editor/MovieSceneTools/Private/TrackEditors/ObjectPropertyTrackEditor.cpp
--- a/Engine/Source/Editor/MovieSceneTools/Private/TrackEditors/ObjectPropertyTrackEditor.cpp
+++ b/Engine/Source/Editor/MovieSceneTools/Private/TrackEditors/ObjectPropertyTrackEditor.cpp
@@ -45,14 +45,14 @@ void FObjectPropertyTrackEditor::InitializeNewTrack(UMovieSceneObjectPropertyTra
 {
 	FPropertyTrackEditor::InitializeNewTrack(NewTrack, PropertyChangedParams);
 
-	FObjectPropertyBase* KeyedProperty = CastField<FObjectPropertyBase>(PropertyChangedParams.PropertyPath.GetLeafMostProperty().Property.Get());
+	const FObjectPropertyBase* KeyedProperty = CastField<FObjectPropertyBase>(PropertyChangedParams.PropertyPath.GetLeafMostProperty().Property.Get());
 	if (KeyedProperty)
 	{
 		NewTrack->PropertyClass = KeyedProperty->PropertyClass;
 
 		if (KeyedProperty->HasAllPropertyFlags(CPF_UObjectWrapper))
 		{
-			FClassProperty* ClassProperty = CastField<FClassProperty>(KeyedProperty);
+			const FClassProperty* ClassProperty = CastField<FClassProperty>(KeyedProperty);
 			if (ClassProperty)
 			{
 				NewTrack->PropertyClass = ClassProperty->MetaClass;
diff --git a/Engine/Source/Editor/MovieSceneTools/Private/TrackEditors/PropertyTrackEditors/BytePropertyTrackEditor.cpp b/Engine/Source/Editor/MovieSceneTools/Private/TrackEditors/PropertyTrackEditors/BytePropertyTrackEditor.cpp
--- a/Engine/Source/Editor/MovieSceneTools/Private/TrackEditors/PropertyTrackEditors/BytePropertyTrackEditor.cpp
+++ b/Engine/Source/Editor/MovieSceneTools/Private/TrackEditors/PropertyTrackEditors/BytePropertyTrackEditor.cpp
@@ -40,16 +40,16 @@ UEnum* GetEnumForByteTrack(TSharedPtr<ISequencer> Sequencer, const FGuid& OwnerO
 			continue;
 		}
 
-		FProperty* Property = PropertyChangedParams.PropertyPath.GetLeafMostProperty().Property.Get();
+		const FProperty* Property = PropertyChangedParams.PropertyPath.GetLeafMostProperty().Property.Get();
 		if (Property != nullptr)
 		{
 			UEnum* Enum = nullptr;
 
-			if (FEnumProperty* EnumProperty = CastField<FEnumProperty>(Property))
+			if (const FEnumProperty* EnumProperty = CastField<FEnumProperty>(Property))
 			{
 				Enum = EnumProperty->GetEnum();
 			}
-			else if (FByteProperty* ByteProperty = CastField<FByteProperty>(Property))
+			else if (const FByteProperty* ByteProperty = CastField<FByteProperty>(Property))
 			{
 				Enum = ByteProperty->Enum;
 			}
diff --git a/Engine/Source/Editor/MovieSceneTools/Private/TrackEditors/PropertyTrackEditors/ColorPropertyTrackEditor.cpp b/Engine/Source/Editor/MovieSceneTools/Private/TrackEditors/PropertyTrackEditors/ColorPropertyTrackEditor.cpp
--- a/Engine/Source/Editor/MovieSceneTools/Private/TrackEditors/PropertyTrackEditors/ColorPropertyTrackEditor.cpp
+++ b/Engine/Source/Editor/MovieSceneTools/Private/TrackEditors/PropertyTrackEditors/ColorPropertyTrackEditor.cpp
@@ -29,7 +29,7 @@ TSharedRef<ISequencerSection> FColorPropertyTrackEditor::MakeSectionInterface(UM
 
 void FColorPropertyTrackEditor::GenerateKeysFromPropertyChanged( const FPropertyChangedParams& PropertyChangedParams, UMovieSceneSection* SectionToKey, FGeneratedTrackKeys& OutGeneratedKeys )
 {
-	FProperty* Property = PropertyChangedParams.PropertyPath.GetLeafMostProperty().Property.Get();
+	const FProperty* Property = PropertyChangedParams.PropertyPath.GetLeafMostProperty().Property.Get();
 	if (!Property)
 	{
 		return;
diff --git a/Engine/Source/Editor/MovieSceneTools/Public/PropertyTrackEditor.h b/Engine/Source/Editor/MovieSceneTools/Public/PropertyTrackEditor.h
--- a/Engine/Source/Editor/MovieSceneTools/Public/PropertyTrackEditor.h
+++ b/Engine/Source/Editor/MovieSceneTools/Public/PropertyTrackEditor.h
@@ -110,7 +110,7 @@ protected:
 				const FPropertyInfo& Info = PropertyChangedParams.PropertyPath.GetPropertyInfo(PropertyIndex);
 				const FArrayProperty* ParentArrayProperty = PropertyIndex > 0 ? CastField<FArrayProperty>(PropertyChangedParams.PropertyPath.GetPropertyInfo(PropertyIndex - 1).Property.Get()) : nullptr;
 
-				FProperty* ArrayInnerProperty = Info.Property.Get();
+				const FProperty* ArrayInnerProperty = Info.Property.Get();
 				if (ArrayInnerProperty && Info.ArrayIndex != INDEX_NONE)
 				{
 					DisplayText = FText::Format(NSLOCTEXT("PropertyTrackEditor", "DisplayTextArrayFormat", "{0} ({1}[{2}])"),
@@ -208,7 +208,7 @@ private:
 	{
 		FKeyPropertyResult KeyPropertyResult;
 
-		FProperty* Property = PropertyChangedParams.PropertyPath.GetLeafMostProperty().Property.Get();
+		const FProperty* Property = PropertyChangedParams.PropertyPath.GetLeafMostProperty().Property.Get();
 		if (!Property)
 		{
 			return KeyPropertyResult;
diff --git a/Engine/Source/Editor/PropertyEditor/Private/Presentation/PropertyTable/PropertyTable.cpp b/Engine/Source/Editor/PropertyEditor/Private/Presentation/PropertyTable/PropertyTable.cpp
--- a/Engine/Source/Editor/PropertyEditor/Private/Presentation/PropertyTable/PropertyTable.cpp
+++ b/Engine/Source/Editor/PropertyEditor/Private/Presentation/PropertyTable/PropertyTable.cpp
@@ -1275,7 +1275,7 @@ void FPropertyTable::UpdateColumns()
 {
 	if( Orientation == EPropertyTableOrientation::AlignPropertiesInColumns)
 	{
-		TMultiMap< FProperty*, TSharedRef< IPropertyTableColumn > > ColumnsMap;
+		TMultiMap< const FProperty*, TSharedRef< IPropertyTableColumn > > ColumnsMap;
 		for (int ColumnIdx = 0; ColumnIdx < Columns.Num(); ++ColumnIdx)
 		{
 			TSharedRef< IDataSource > DataSource = Columns[ColumnIdx]->GetDataSource();
diff --git a/Engine/Source/Editor/PropertyEditor/Private/Presentation/PropertyTable/PropertyTableColumn.cpp b/Engine/Source/Editor/PropertyEditor/Private/Presentation/PropertyTable/PropertyTableColumn.cpp
--- a/Engine/Source/Editor/PropertyEditor/Private/Presentation/PropertyTable/PropertyTableColumn.cpp
+++ b/Engine/Source/Editor/PropertyEditor/Private/Presentation/PropertyTable/PropertyTableColumn.cpp
@@ -565,7 +565,7 @@ TSharedRef< class IPropertyTable > FPropertyTableColumn::GetTable() const
 bool FPropertyTableColumn::CanSortBy() const
 {
 	TWeakObjectPtr< UObject > Object = DataSource->AsUObject();
-	FProperty* Property = nullptr;
+	const FProperty* Property = nullptr;
 
 	TSharedPtr< FPropertyPath > Path = DataSource->AsPropertyPath();
 	if ( Path.IsValid() )
@@ -576,11 +576,11 @@ bool FPropertyTableColumn::CanSortBy() const
 	return ( Property != nullptr );
 }
 
-TSharedPtr<FCompareRowByColumnBase> FPropertyTableColumn::GetPropertySorter(FProperty* Property, EColumnSortMode::Type SortMode)
+TSharedPtr<FCompareRowByColumnBase> FPropertyTableColumn::GetPropertySorter(const FProperty* Property, EColumnSortMode::Type SortMode)
 {
 	if (Property->IsA(FEnumProperty::StaticClass()))
 	{
-		FEnumProperty* EnumProperty = CastField<FEnumProperty>(Property);
+		const FEnumProperty* EnumProperty = CastField<FEnumProperty>(Property);
 
 		if (SortMode == EColumnSortMode::Ascending)
 		{
@@ -593,7 +593,7 @@ TSharedPtr<FCompareRowByColumnBase> FPropertyTableColumn::GetPropertySorter(FPro
 	}
 	else if (Property->IsA(FByteProperty::StaticClass()))
 	{
-		FByteProperty* ByteProperty = CastField<FByteProperty>(Property);
+		const FByteProperty* ByteProperty = CastField<FByteProperty>(Property);
 
 		if (SortMode == EColumnSortMode::Ascending)
 		{
@@ -606,7 +606,7 @@ TSharedPtr<FCompareRowByColumnBase> FPropertyTableColumn::GetPropertySorter(FPro
 	}
 	else if (Property->IsA(FIntProperty::StaticClass()))
 	{
-		FIntProperty* IntProperty = CastField<FIntProperty>(Property);
+		const FIntProperty* IntProperty = CastField<FIntProperty>(Property);
 
 		if (SortMode == EColumnSortMode::Ascending)
 		{
@@ -619,7 +619,7 @@ TSharedPtr<FCompareRowByColumnBase> FPropertyTableColumn::GetPropertySorter(FPro
 	}
 	else if (Property->IsA(FBoolProperty::StaticClass()))
 	{
-		FBoolProperty* BoolProperty = CastField<FBoolProperty>(Property);
+		const FBoolProperty* BoolProperty = CastField<FBoolProperty>(Property);
 
 		if (SortMode == EColumnSortMode::Ascending)
 		{
@@ -632,7 +632,7 @@ TSharedPtr<FCompareRowByColumnBase> FPropertyTableColumn::GetPropertySorter(FPro
 	}
 	else if (Property->IsA(FFloatProperty::StaticClass()))
 	{
-		FFloatProperty* FloatProperty(CastField< FFloatProperty >(Property));
+		const FFloatProperty* FloatProperty(CastField< FFloatProperty >(Property));
 
 		if (SortMode == EColumnSortMode::Ascending)
 		{
@@ -645,7 +645,7 @@ TSharedPtr<FCompareRowByColumnBase> FPropertyTableColumn::GetPropertySorter(FPro
 	}
 	else if (Property->IsA(FDoubleProperty::StaticClass()))
 	{
-		FDoubleProperty* DoubleProperty(CastField< FDoubleProperty >(Property));
+		const FDoubleProperty* DoubleProperty(CastField< FDoubleProperty >(Property));
 
 		if (SortMode == EColumnSortMode::Ascending)
 		{
@@ -658,7 +658,7 @@ TSharedPtr<FCompareRowByColumnBase> FPropertyTableColumn::GetPropertySorter(FPro
 	}
 	else if (Property->IsA(FNameProperty::StaticClass()))
 	{
-		FNameProperty* NameProperty = CastField<FNameProperty>(Property);
+		const FNameProperty* NameProperty = CastField<FNameProperty>(Property);
 
 		if (SortMode == EColumnSortMode::Ascending)
 		{
@@ -671,7 +671,7 @@ TSharedPtr<FCompareRowByColumnBase> FPropertyTableColumn::GetPropertySorter(FPro
 	}
 	else if (Property->IsA(FStrProperty::StaticClass()))
 	{
-		FStrProperty* StrProperty = CastField<FStrProperty>(Property);
+		const FStrProperty* StrProperty = CastField<FStrProperty>(Property);
 
 		if (SortMode == EColumnSortMode::Ascending)
 		{
@@ -684,7 +684,7 @@ TSharedPtr<FCompareRowByColumnBase> FPropertyTableColumn::GetPropertySorter(FPro
 	}
 	else if (Property->IsA(FObjectPropertyBase::StaticClass()) && !Property->HasAnyPropertyFlags(CPF_InstancedReference))
 	{
-		FObjectPropertyBase* ObjectProperty = CastField<FObjectPropertyBase>(Property);
+		const FObjectPropertyBase* ObjectProperty = CastField<FObjectPropertyBase>(Property);
 
 		if (SortMode == EColumnSortMode::Ascending)
 		{
@@ -697,7 +697,7 @@ TSharedPtr<FCompareRowByColumnBase> FPropertyTableColumn::GetPropertySorter(FPro
 	}
 	else if (IsSupportedStructProperty(Property))
 	{
-		FStructProperty* StructProperty = CastField<FStructProperty>(Property);
+		const FStructProperty* StructProperty = CastField<FStructProperty>(Property);
 
 		if (SortMode == EColumnSortMode::Ascending)
 		{
@@ -733,7 +733,7 @@ void FPropertyTableColumn::Sort( TArray< TSharedRef< class IPropertyTableRow > >
 	}
 
 	UObject* PrimaryObject = DataSource->AsUObject().Get();
-	FProperty* PrimaryProperty = nullptr;
+	const FProperty* PrimaryProperty = nullptr;
 	TSharedPtr< FPropertyPath > PrimaryPath = DataSource->AsPropertyPath();
 	if (PrimaryPath.IsValid())
 	{
@@ -741,7 +741,7 @@ void FPropertyTableColumn::Sort( TArray< TSharedRef< class IPropertyTableRow > >
 	}
 
 	UObject* SecondaryObject = nullptr;
-	FProperty* SecondaryProperty = nullptr;
+	const FProperty* SecondaryProperty = nullptr;
 	if(SecondarySortColumn.IsValid())
 	{
 		SecondaryObject = SecondarySortColumn->GetDataSource()->AsUObject().Get();
diff --git a/Engine/Source/Editor/PropertyEditor/Private/Presentation/PropertyTable/PropertyTableColumn.h b/Engine/Source/Editor/PropertyEditor/Private/Presentation/PropertyTable/PropertyTableColumn.h
--- a/Engine/Source/Editor/PropertyEditor/Private/Presentation/PropertyTable/PropertyTableColumn.h
+++ b/Engine/Source/Editor/PropertyEditor/Private/Presentation/PropertyTable/PropertyTableColumn.h
@@ -59,7 +59,7 @@ public:
 
 	virtual void Tick() override;
 
-	virtual TSharedPtr<struct FCompareRowByColumnBase> GetPropertySorter(FProperty* Property, EColumnSortMode::Type SortMode) override;
+	virtual TSharedPtr<struct FCompareRowByColumnBase> GetPropertySorter(const FProperty* Property, EColumnSortMode::Type SortMode) override;
 
 	DECLARE_DERIVED_EVENT( FPropertyTableColumn, IPropertyTableColumn::FFrozenStateChanged, FFrozenStateChanged );
 	FFrozenStateChanged* OnFrozenStateChanged() override { return &FrozenStateChanged; }
diff --git a/Engine/Source/Editor/PropertyEditor/Private/Presentation/PropertyTable/PropertyTableObjectNameColumn.h b/Engine/Source/Editor/PropertyEditor/Private/Presentation/PropertyTable/PropertyTableObjectNameColumn.h
--- a/Engine/Source/Editor/PropertyEditor/Private/Presentation/PropertyTable/PropertyTableObjectNameColumn.h
+++ b/Engine/Source/Editor/PropertyEditor/Private/Presentation/PropertyTable/PropertyTableObjectNameColumn.h
@@ -133,7 +133,7 @@ public:
 		}
 	}
 
-	virtual TSharedPtr<struct FCompareRowByColumnBase> GetPropertySorter(FProperty* Property, EColumnSortMode::Type SortMode) override
+	virtual TSharedPtr<struct FCompareRowByColumnBase> GetPropertySorter(const FProperty* Property, EColumnSortMode::Type SortMode) override
 	{
 		// Does not sort properties
 		return nullptr;
diff --git a/Engine/Source/Editor/PropertyEditor/Private/Presentation/PropertyTable/PropertyTablePropertyNameColumn.cpp b/Engine/Source/Editor/PropertyEditor/Private/Presentation/PropertyTable/PropertyTablePropertyNameColumn.cpp
--- a/Engine/Source/Editor/PropertyEditor/Private/Presentation/PropertyTable/PropertyTablePropertyNameColumn.cpp
+++ b/Engine/Source/Editor/PropertyEditor/Private/Presentation/PropertyTable/PropertyTablePropertyNameColumn.cpp
@@ -87,7 +87,7 @@ void FPropertyTablePropertyNameColumn::Sort( TArray<TSharedRef<class IPropertyTa
 }
 
 
-TSharedPtr<struct FCompareRowByColumnBase> FPropertyTablePropertyNameColumn::GetPropertySorter(FProperty* Property, EColumnSortMode::Type SortMode)
+TSharedPtr<struct FCompareRowByColumnBase> FPropertyTablePropertyNameColumn::GetPropertySorter(const FProperty* Property, EColumnSortMode::Type SortMode)
 {
 	// Does not sort properties
 	return nullptr;
@@ -98,7 +98,7 @@ FString FPropertyTablePropertyNameColumn::GetPropertyNameAsString( const TShared
 	FString PropertyName;
 	if( Row->GetDataSource()->AsPropertyPath().IsValid() )
 	{
-		const TWeakFieldPtr< FProperty > Property = Row->GetDataSource()->AsPropertyPath()->GetLeafMostProperty().Property;
+		const TWeakFieldPtr< const FProperty > Property = Row->GetDataSource()->AsPropertyPath()->GetLeafMostProperty().Property;
 		PropertyName = UEditorEngine::GetFriendlyName( Property.Get() );
 	}
 	return PropertyName;
diff --git a/Engine/Source/Editor/PropertyEditor/Private/Presentation/PropertyTable/PropertyTablePropertyNameColumn.h b/Engine/Source/Editor/PropertyEditor/Private/Presentation/PropertyTable/PropertyTablePropertyNameColumn.h
--- a/Engine/Source/Editor/PropertyEditor/Private/Presentation/PropertyTable/PropertyTablePropertyNameColumn.h
+++ b/Engine/Source/Editor/PropertyEditor/Private/Presentation/PropertyTable/PropertyTablePropertyNameColumn.h
@@ -59,7 +59,7 @@ public:
 
 	virtual void Sort( TArray< TSharedRef< class IPropertyTableRow > >& Rows, const EColumnSortMode::Type PrimarySortMode, const TSharedPtr<IPropertyTableColumn>& SecondarySortColumn, const EColumnSortMode::Type SecondarySortMode ) override;
 
-	virtual TSharedPtr<struct FCompareRowByColumnBase> GetPropertySorter(FProperty* Property, EColumnSortMode::Type SortMode) override;
+	virtual TSharedPtr<struct FCompareRowByColumnBase> GetPropertySorter(const FProperty* Property, EColumnSortMode::Type SortMode) override;
 
 	virtual void Tick() override {}
 
diff --git a/Engine/Source/Editor/PropertyEditor/Private/Presentation/PropertyTable/PropertyTableRowHeaderColumn.h b/Engine/Source/Editor/PropertyEditor/Private/Presentation/PropertyTable/PropertyTableRowHeaderColumn.h
--- a/Engine/Source/Editor/PropertyEditor/Private/Presentation/PropertyTable/PropertyTableRowHeaderColumn.h
+++ b/Engine/Source/Editor/PropertyEditor/Private/Presentation/PropertyTable/PropertyTableRowHeaderColumn.h
@@ -76,7 +76,7 @@ public:
 
 	virtual void Sort( TArray< TSharedRef< class IPropertyTableRow > >& Rows, const EColumnSortMode::Type PrimarySortMode, const TSharedPtr<IPropertyTableColumn>& SecondarySortColumn, const EColumnSortMode::Type SecondarySortMode ) override {  }
 
-	virtual TSharedPtr<struct FCompareRowByColumnBase> GetPropertySorter(FProperty* Property, EColumnSortMode::Type SortMode) override { return nullptr;  }
+	virtual TSharedPtr<struct FCompareRowByColumnBase> GetPropertySorter(const FProperty* Property, EColumnSortMode::Type SortMode) override { return nullptr;  }
 
 	virtual void Tick() override {}
 
diff --git a/Engine/Source/Editor/PropertyEditor/Private/UserInterface/PropertyTable/ColumnWidgetFactory.cpp b/Engine/Source/Editor/PropertyEditor/Private/UserInterface/PropertyTable/ColumnWidgetFactory.cpp
--- a/Engine/Source/Editor/PropertyEditor/Private/UserInterface/PropertyTable/ColumnWidgetFactory.cpp
+++ b/Engine/Source/Editor/PropertyEditor/Private/UserInterface/PropertyTable/ColumnWidgetFactory.cpp
@@ -50,7 +50,7 @@ TSharedRef< SColumnHeader > ColumnWidgetFactory::CreateColumnHeaderWidget( const
 	const TSharedPtr< FPropertyPath > Path = ColumnDataSource->AsPropertyPath();
 	if( Path.IsValid() )
 	{
-		TWeakFieldPtr< FProperty > Property = Path->GetLeafMostProperty().Property;
+		TWeakFieldPtr< const FProperty > Property = Path->GetLeafMostProperty().Property;
 		check( Property.IsValid() );
 
 		if ( Property->IsA( FBoolProperty::StaticClass() ) )
diff --git a/Engine/Source/Editor/PropertyEditor/Private/UserInterface/PropertyTable/SColumnHeader.h b/Engine/Source/Editor/PropertyEditor/Private/UserInterface/PropertyTable/SColumnHeader.h
--- a/Engine/Source/Editor/PropertyEditor/Private/UserInterface/PropertyTable/SColumnHeader.h
+++ b/Engine/Source/Editor/PropertyEditor/Private/UserInterface/PropertyTable/SColumnHeader.h
@@ -72,7 +72,7 @@ protected:
 		}
 		else if ( PropertyPath.IsValid() && PropertyPath->GetNumProperties() > 0 )
 		{
-			FProperty* Property = PropertyPath->GetLeafMostProperty().Property.Get();
+			const FProperty* Property = PropertyPath->GetLeafMostProperty().Property.Get();
 			const FText& DisplayNameText = Column->GetDisplayName();
 			const FString DocumentationLink = PropertyEditorHelpers::GetDocumentationLink( Property );
 			const FString DocumentationExcerptName = PropertyEditorHelpers::GetDocumentationExcerptName( Property );
diff --git a/Engine/Source/Editor/PropertyEditor/Private/UserInterface/PropertyTable/SObjectColumnHeader.h b/Engine/Source/Editor/PropertyEditor/Private/UserInterface/PropertyTable/SObjectColumnHeader.h
--- a/Engine/Source/Editor/PropertyEditor/Private/UserInterface/PropertyTable/SObjectColumnHeader.h
+++ b/Engine/Source/Editor/PropertyEditor/Private/UserInterface/PropertyTable/SObjectColumnHeader.h
@@ -52,7 +52,7 @@ class SObjectColumnHeader : public SColumnHeader
 		{
 			TSharedRef< FPropertyEditor > PropertyEditor = FPropertyEditor::Create( Cell->GetNode().ToSharedRef(), Utilities.ToSharedRef() );
 
-			TWeakFieldPtr< FProperty > Property = PropertyTableRow->GetDataSource()->AsPropertyPath()->GetLeafMostProperty().Property.Get();
+			TWeakFieldPtr< const FProperty > Property = PropertyTableRow->GetDataSource()->AsPropertyPath()->GetLeafMostProperty().Property.Get();
 
 			if( Property->IsA( FBoolProperty::StaticClass() ) )
 			{
diff --git a/Engine/Source/Editor/PropertyEditor/Private/UserInterface/PropertyTable/SPropertyNameColumnHeader.h b/Engine/Source/Editor/PropertyEditor/Private/UserInterface/PropertyTable/SPropertyNameColumnHeader.h
--- a/Engine/Source/Editor/PropertyEditor/Private/UserInterface/PropertyTable/SPropertyNameColumnHeader.h
+++ b/Engine/Source/Editor/PropertyEditor/Private/UserInterface/PropertyTable/SPropertyNameColumnHeader.h
@@ -54,7 +54,7 @@ class SPropertyNameColumnHeader : public SColumnHeader
 			
 		if( PropertyTableRow->GetDataSource()->AsPropertyPath().IsValid() )
 		{	
-			const TWeakFieldPtr< FProperty > Property = PropertyTableRow->GetDataSource()->AsPropertyPath()->GetLeafMostProperty().Property;
+			const TWeakFieldPtr< const FProperty > Property = PropertyTableRow->GetDataSource()->AsPropertyPath()->GetLeafMostProperty().Property;
 			PropertyName = UEditorEngine::GetFriendlyName( Property.Get() );
 		}
 
diff --git a/Engine/Source/Editor/PropertyEditor/Private/UserInterface/PropertyTable/SPropertyTable.h b/Engine/Source/Editor/PropertyEditor/Private/UserInterface/PropertyTable/SPropertyTable.h
--- a/Engine/Source/Editor/PropertyEditor/Private/UserInterface/PropertyTable/SPropertyTable.h
+++ b/Engine/Source/Editor/PropertyEditor/Private/UserInterface/PropertyTable/SPropertyTable.h
@@ -737,7 +737,7 @@ private:
 					for( auto ExtensionIter = TypeToProperties.CreateConstKeyIterator( *TypeIter ); ExtensionIter; ++ExtensionIter )
 					{
 						const FPropertyInfo& PropInfo = ExtensionIter.Value();
-						const TWeakFieldPtr< FProperty > Property = PropInfo.Property;
+						const TWeakFieldPtr< const FProperty > Property = PropInfo.Property;
 						const FText PropName = FText::FromString( UEditorEngine::GetFriendlyName(Property.Get()) );
 						MenuBuilder.AddMenuEntry( PropName, PropName, FSlateIcon(), FUIAction( FExecuteAction::CreateSP( this, &SPropertyTable::SetRootPath, RootPath->ExtendPath( PropInfo ) ) ) );
 					}
diff --git a/Engine/Source/Editor/PropertyEditor/Public/IPropertyTableColumn.h b/Engine/Source/Editor/PropertyEditor/Public/IPropertyTableColumn.h
--- a/Engine/Source/Editor/PropertyEditor/Public/IPropertyTableColumn.h
+++ b/Engine/Source/Editor/PropertyEditor/Public/IPropertyTableColumn.h
@@ -49,7 +49,7 @@ public:
 	virtual bool CanSortBy() const = 0;
 	virtual void Sort( TArray< TSharedRef< class IPropertyTableRow > >& Rows, const EColumnSortMode::Type PrimarySortMode, const TSharedPtr<IPropertyTableColumn>& SecondarySortColumn, const EColumnSortMode::Type SecondarySortMode ) = 0;
 
-	virtual TSharedPtr<struct FCompareRowByColumnBase> GetPropertySorter(FProperty* Property, EColumnSortMode::Type SortMode) = 0;
+	virtual TSharedPtr<struct FCompareRowByColumnBase> GetPropertySorter(const FProperty* Property, EColumnSortMode::Type SortMode) = 0;
 
 	/**
 	 * Tick the property column each frame
diff --git a/Engine/Source/Editor/Sequencer/Private/KeyPropertyParams.cpp b/Engine/Source/Editor/Sequencer/Private/KeyPropertyParams.cpp
--- a/Engine/Source/Editor/Sequencer/Private/KeyPropertyParams.cpp
+++ b/Engine/Source/Editor/Sequencer/Private/KeyPropertyParams.cpp
@@ -64,7 +64,7 @@ const UStruct* FCanKeyPropertyParams::FindPropertyOwner(const FProperty* ForProp
 	bool bFoundProperty = false;
 	for (int32 Index = PropertyPath.GetNumProperties() - 1; Index >= 0; --Index)
 	{
-		FProperty* Property = PropertyPath.GetPropertyInfo(Index).Property.Get();
+		const FProperty* Property = PropertyPath.GetPropertyInfo(Index).Property.Get();
 		if (!bFoundProperty)
 		{
 			bFoundProperty = Property == ForProperty;
@@ -88,7 +88,7 @@ const UStruct* FCanKeyPropertyParams::FindPropertyContainer(const FProperty* For
 	bool bFoundProperty = false;
 	for (int32 Index = PropertyPath.GetNumProperties() - 1; Index >= 0; --Index)
 	{
-		FProperty* Property = PropertyPath.GetPropertyInfo(Index).Property.Get();
+		const FProperty* Property = PropertyPath.GetPropertyInfo(Index).Property.Get();
 		if (!bFoundProperty)
 		{
 			bFoundProperty = Property == ForProperty;
diff --git a/Engine/Source/Editor/Sequencer/Private/MVVM/ViewModels/ObjectBindingModel.cpp b/Engine/Source/Editor/Sequencer/Private/MVVM/ViewModels/ObjectBindingModel.cpp
--- a/Engine/Source/Editor/Sequencer/Private/MVVM/ViewModels/ObjectBindingModel.cpp
+++ b/Engine/Source/Editor/Sequencer/Private/MVVM/ViewModels/ObjectBindingModel.cpp
@@ -679,7 +679,7 @@ TSharedRef<SWidget> FObjectBindingModel::GetAddTrackMenuContent()
 	TMap<FString, TArray<PropertyMenuData>> KeyablePropertyMenuData;
 	for (const FPropertyPath& KeyablePropertyPath : KeyablePropertyPaths)
 	{
-		FProperty* Property = KeyablePropertyPath.GetRootProperty().Property.Get();
+		const FProperty* Property = KeyablePropertyPath.GetRootProperty().Property.Get();
 		if (Property)
 		{
 			PropertyMenuData KeyableMenuData;
@@ -797,7 +797,7 @@ void FObjectBindingModel::HandleAddTrackSubMenuNew(FMenuBuilder& AddTrackMenuBui
 	// [PostProcessSettings] [ColorGrading]
 
 	// Create property menu data based on keyable property paths
-	TArray<FProperty*> PropertiesTraversed;
+	TArray<const FProperty*> PropertiesTraversed;
 	TArray<int32> ArrayIndicesTraversed;
 	TArray<PropertyMenuData> KeyablePropertyMenuData;
 	for (const FPropertyPath& KeyablePropertyPath : KeyablePropertyPaths)
@@ -809,7 +809,7 @@ void FObjectBindingModel::HandleAddTrackSubMenuNew(FMenuBuilder& AddTrackMenuBui
 		if (KeyablePropertyPath.GetNumProperties() > 1) //@todo
 		{
 			const FPropertyInfo& PropertyInfo = KeyablePropertyPath.GetPropertyInfo(1);
-			FProperty* Property = PropertyInfo.Property.Get();
+			const FProperty* Property = PropertyInfo.Property.Get();
 
 			// Search for any array elements
 			int32 ArrayIndex = INDEX_NONE;
diff --git a/Engine/Source/Editor/Sequencer/Private/SequencerObjectChangeListener.cpp b/Engine/Source/Editor/Sequencer/Private/SequencerObjectChangeListener.cpp
--- a/Engine/Source/Editor/Sequencer/Private/SequencerObjectChangeListener.cpp
+++ b/Engine/Source/Editor/Sequencer/Private/SequencerObjectChangeListener.cpp
@@ -65,7 +65,7 @@ void FSequencerObjectChangeListener::BroadcastPropertyChanged( FKeyPropertyParam
 	// both the CameraActor and the CameraComponent.
 	TArray<UObject*> KeyableObjects;
 	FOnAnimatablePropertyChanged Delegate;
-	FProperty* Property = nullptr;
+	const FProperty* Property = nullptr;
 	FPropertyPath PropertyPath;
 	for (auto ObjectToKey : KeyPropertyParams.ObjectsToKey)
 	{
@@ -276,19 +276,19 @@ const FOnAnimatablePropertyChanged* FSequencerObjectChangeListener::FindProperty
 bool FSequencerObjectChangeListener::CanKeyProperty(FCanKeyPropertyParams CanKeyPropertyParams) const
 {
 	FOnAnimatablePropertyChanged Delegate;
-	FProperty* Property = nullptr;
+	const FProperty* Property = nullptr;
 	FPropertyPath PropertyPath;
 	return CanKeyProperty_Internal(CanKeyPropertyParams, Delegate, Property, PropertyPath);
 }
 
 bool FSequencerObjectChangeListener::CanKeyProperty(FCanKeyPropertyParams KeyPropertyParams, FPropertyPath& OutPropertyPath) const
 {
-	FProperty* Property = nullptr;
+	const FProperty* Property = nullptr;
 	FOnAnimatablePropertyChanged Delegate;
 	return CanKeyProperty_Internal(KeyPropertyParams, Delegate, Property, OutPropertyPath);
 }
 
-bool FSequencerObjectChangeListener::CanKeyProperty_Internal(FCanKeyPropertyParams CanKeyPropertyParams, FOnAnimatablePropertyChanged& InOutDelegate, FProperty*& InOutProperty, FPropertyPath& InOutPropertyPath) const
+bool FSequencerObjectChangeListener::CanKeyProperty_Internal(FCanKeyPropertyParams CanKeyPropertyParams, FOnAnimatablePropertyChanged& InOutDelegate, const FProperty*& InOutProperty, FPropertyPath& InOutPropertyPath) const
 {
 	if (CanKeyPropertyParams.PropertyPath.GetNumProperties() == 0)
 	{
@@ -304,7 +304,7 @@ bool FSequencerObjectChangeListener::CanKeyProperty_Internal(FCanKeyPropertyPara
 		// Add this to our 'potentially truncated' path
 		InOutPropertyPath.AddProperty(PropertyInfo);
 
-		FProperty* Property = CanKeyPropertyParams.PropertyPath.GetPropertyInfo(Index).Property.Get();
+		const FProperty* Property = CanKeyPropertyParams.PropertyPath.GetPropertyInfo(Index).Property.Get();
 		if (Property)
 		{
 			if (Property->IsA<FArrayProperty>())
@@ -356,7 +356,7 @@ bool FSequencerObjectChangeListener::CanKeyProperty_Internal(FCanKeyPropertyPara
 				return true;
 			}
 
-			FObjectPropertyBase* ObjectProperty = CastField<FObjectPropertyBase>(Property);
+			const FObjectPropertyBase* ObjectProperty = CastField<FObjectPropertyBase>(Property);
 
 			// Check each level of the property hierarchy
 			FFieldClass* PropertyType = Property->GetClass();
diff --git a/Engine/Source/Editor/Sequencer/Private/SequencerObjectChangeListener.h b/Engine/Source/Editor/Sequencer/Private/SequencerObjectChangeListener.h
--- a/Engine/Source/Editor/Sequencer/Private/SequencerObjectChangeListener.h
+++ b/Engine/Source/Editor/Sequencer/Private/SequencerObjectChangeListener.h
@@ -86,7 +86,7 @@ private:
 	 * @param	InOutPropertyPath	The property path of the proepty we will actually key
 	 * @return true if we can key the property
 	 */
-	bool CanKeyProperty_Internal(FCanKeyPropertyParams KeyPropertyParams, FOnAnimatablePropertyChanged& InOutDelegate, FProperty*& InOutProperty, FPropertyPath& InOutPropertyPath) const;
+	bool CanKeyProperty_Internal(FCanKeyPropertyParams KeyPropertyParams, FOnAnimatablePropertyChanged& InOutDelegate, const FProperty*& InOutProperty, FPropertyPath& InOutPropertyPath) const;
 
 private:
 	/** Mapping of object to a listener used to check for property changes */
diff --git a/Engine/Source/Editor/Sequencer/Private/SequencerPropertyKeyedStatus.cpp b/Engine/Source/Editor/Sequencer/Private/SequencerPropertyKeyedStatus.cpp
--- a/Engine/Source/Editor/Sequencer/Private/SequencerPropertyKeyedStatus.cpp
+++ b/Engine/Source/Editor/Sequencer/Private/SequencerPropertyKeyedStatus.cpp
@@ -58,7 +58,7 @@ FName FSequencerPropertyKeyedStatusHandler::FPropertyParameters::FindTrackProper
 {
 	if (PropertyPath.GetNumProperties() > 0)
 	{
-		FStructProperty* StructProperty = CastField<FStructProperty>(PropertyPath.GetLeafMostProperty().Property.Get());
+		const FStructProperty* StructProperty = CastField<FStructProperty>(PropertyPath.GetLeafMostProperty().Property.Get());
 		if (StructProperty && StructProperty->Struct)
 		{
 			return StructProperty->Struct->GetFName();
diff --git a/Engine/Source/Editor/Sequencer/Public/KeyPropertyParams.h b/Engine/Source/Editor/Sequencer/Public/KeyPropertyParams.h
--- a/Engine/Source/Editor/Sequencer/Public/KeyPropertyParams.h
+++ b/Engine/Source/Editor/Sequencer/Public/KeyPropertyParams.h
@@ -97,17 +97,17 @@ public:
 		for (int32 i = 0; i < PropertyPath.GetNumProperties(); i++)
 		{
 			const FPropertyInfo& PropertyInfo = PropertyPath.GetPropertyInfo(i);
-			if (FProperty* Property = PropertyInfo.Property.Get())
+			if (const FProperty* Property = PropertyInfo.Property.Get())
 			{
 				int32 ArrayIndex = FMath::Max(0, PropertyInfo.ArrayIndex);
-				if (FArrayProperty* ArrayProp = CastField<FArrayProperty>(Property))
+				if (const FArrayProperty* ArrayProp = CastField<FArrayProperty>(Property))
 				{
 					// Sometimes property paths have the array property twice, first with no array index,
 					// then a second so we skip over this property if that's the case
 					if (PropertyInfo.ArrayIndex == INDEX_NONE && i < PropertyPath.GetNumProperties()-1)
 					{
 						const FPropertyInfo& InnerPropertyInfo = PropertyPath.GetPropertyInfo(i+1);
-						FProperty* InnerProperty = InnerPropertyInfo.Property.Get();
+						const FProperty* InnerProperty = InnerPropertyInfo.Property.Get();
 						if (InnerProperty && InnerProperty->GetOwner<FProperty>() == ArrayProp)
 						{
 							ArrayIndex = InnerPropertyInfo.ArrayIndex;
diff --git a/Engine/Source/Editor/StatsViewer/Private/ActorArrayHyperlinkColumn.cpp b/Engine/Source/Editor/StatsViewer/Private/ActorArrayHyperlinkColumn.cpp
--- a/Engine/Source/Editor/StatsViewer/Private/ActorArrayHyperlinkColumn.cpp
+++ b/Engine/Source/Editor/StatsViewer/Private/ActorArrayHyperlinkColumn.cpp
@@ -125,10 +125,10 @@ bool FActorArrayHyperlinkColumn::Supports( const TSharedRef< IPropertyTableColum
 		if( PropertyPath.IsValid() && PropertyPath->GetNumProperties() > 0 )
 		{
 			const FPropertyInfo& PropertyInfo = PropertyPath->GetRootProperty();
-			FProperty* Property = PropertyInfo.Property.Get();
+			const FProperty* Property = PropertyInfo.Property.Get();
 			if( Property->IsA( FArrayProperty::StaticClass() ) )
 			{
-				FArrayProperty* ArrayProperty = CastField<FArrayProperty>(Property);
+				const FArrayProperty* ArrayProperty = CastField<FArrayProperty>(Property);
 				if( ArrayProperty->Inner->IsA(FWeakObjectProperty::StaticClass()) )
 				{
 					return true;
diff --git a/Engine/Source/Editor/StatsViewer/Private/ObjectHyperlinkColumn.cpp b/Engine/Source/Editor/StatsViewer/Private/ObjectHyperlinkColumn.cpp
--- a/Engine/Source/Editor/StatsViewer/Private/ObjectHyperlinkColumn.cpp
+++ b/Engine/Source/Editor/StatsViewer/Private/ObjectHyperlinkColumn.cpp
@@ -138,7 +138,7 @@ bool FObjectHyperlinkColumn::Supports( const TSharedRef< IPropertyTableColumn >&
 		if( PropertyPath.IsValid() && PropertyPath->GetNumProperties() > 0 )
 		{
 			const FPropertyInfo& PropertyInfo = PropertyPath->GetRootProperty();
-			FProperty* Property = PropertyInfo.Property.Get();
+			const FProperty* Property = PropertyInfo.Property.Get();
 			if( Property->IsA( FWeakObjectProperty::StaticClass() ) )
 			{
 				const UClass* PropertyClass = CastField<FWeakObjectProperty>(Property)->PropertyClass;
diff --git a/Engine/Source/Editor/StatsViewer/Private/SStatsViewer.cpp b/Engine/Source/Editor/StatsViewer/Private/SStatsViewer.cpp
--- a/Engine/Source/Editor/StatsViewer/Private/SStatsViewer.cpp
+++ b/Engine/Source/Editor/StatsViewer/Private/SStatsViewer.cpp
@@ -478,7 +478,7 @@ FReply SStatsViewer::OnExportClicked()
 				{
 					TSharedPtr< FPropertyPath > PropertyPath = Column->GetDataSource()->AsPropertyPath();
 					const FPropertyInfo& PropertyInfo = PropertyPath->GetRootProperty();
-					const TWeakFieldPtr< FProperty > Property = PropertyInfo.Property;
+					const TWeakFieldPtr< const FProperty > Property = PropertyInfo.Property;
 					FString Name = UEditorEngine::GetFriendlyName(Property.Get());
 					Name.ReplaceInline( *Delimiter, TEXT(" ") );
 					HeaderRow += Name + Delimiter;
diff --git a/Engine/Source/Editor/StatsViewer/Private/StatsCustomColumn.cpp b/Engine/Source/Editor/StatsViewer/Private/StatsCustomColumn.cpp
--- a/Engine/Source/Editor/StatsViewer/Private/StatsCustomColumn.cpp
+++ b/Engine/Source/Editor/StatsViewer/Private/StatsCustomColumn.cpp
@@ -140,7 +140,7 @@ FText FStatsCustomColumn::GetTotalText( TSharedRef< IPropertyTableColumn > Colum
 	return FText::GetEmpty();
 }
 
-bool FStatsCustomColumn::SupportsProperty( FProperty* Property )
+bool FStatsCustomColumn::SupportsProperty( const FProperty* Property )
 {
 	if( Property->IsA( FFloatProperty::StaticClass() ) || Property->IsA( FIntProperty::StaticClass() ) )
 	{
diff --git a/Engine/Source/Editor/StatsViewer/Private/StatsCustomColumn.h b/Engine/Source/Editor/StatsViewer/Private/StatsCustomColumn.h
--- a/Engine/Source/Editor/StatsViewer/Private/StatsCustomColumn.h
+++ b/Engine/Source/Editor/StatsViewer/Private/StatsCustomColumn.h
@@ -39,7 +39,7 @@ public:
 	 * @param	Property	The property to check
 	 * @returns true if this property is supported for display by this column type
 	 */
-	static bool SupportsProperty( FProperty* Property );
+	static bool SupportsProperty( const FProperty* Property );
 
 	/** 
 	 * Helper function to get the text we would display for this property 
diff --git a/Engine/Source/Editor/TranslationEditor/Private/CustomFontColumn.cpp b/Engine/Source/Editor/TranslationEditor/Private/CustomFontColumn.cpp
--- a/Engine/Source/Editor/TranslationEditor/Private/CustomFontColumn.cpp
+++ b/Engine/Source/Editor/TranslationEditor/Private/CustomFontColumn.cpp
@@ -40,7 +40,7 @@ bool FCustomFontColumn::Supports( const TSharedRef< IPropertyTableColumn >& Colu
 		if( PropertyPath.IsValid() && PropertyPath->GetNumProperties() > 0 )
 		{
 			const FPropertyInfo& PropertyInfo = PropertyPath->GetRootProperty();
-			FProperty* Property = PropertyInfo.Property.Get();
+			const FProperty* Property = PropertyInfo.Property.Get();
 			if (SupportedProperties.Contains(Property))
 			{
 				IsSupported = true;
diff --git a/Engine/Source/Editor/TranslationEditor/Private/TranslationEditor.cpp b/Engine/Source/Editor/TranslationEditor/Private/TranslationEditor.cpp
--- a/Engine/Source/Editor/TranslationEditor/Private/TranslationEditor.cpp
+++ b/Engine/Source/Editor/TranslationEditor/Private/TranslationEditor.cpp
@@ -1070,7 +1070,7 @@ void FTranslationEditor::UpdateContextSelection()
 	{
 		TSet<TSharedRef<IPropertyTableRow>> SelectedRows = ContextPropertyTable->GetSelectedRows();
 		TSharedRef<FPropertyPath> InitialPath = ContextPropertyTable->GetRootPath();
-		FProperty* PropertyToFind = InitialPath->GetRootProperty().Property.Get();
+		const FProperty* PropertyToFind = InitialPath->GetRootProperty().Property.Get();
 
 		// Can only really handle single selection
 		if (SelectedRows.Num() == 1)
diff --git a/Engine/Source/Editor/UnrealEd/Private/ObjectTools.cpp b/Engine/Source/Editor/UnrealEd/Private/ObjectTools.cpp
--- a/Engine/Source/Editor/UnrealEd/Private/ObjectTools.cpp
+++ b/Engine/Source/Editor/UnrealEd/Private/ObjectTools.cpp
@@ -505,8 +505,8 @@ namespace ObjectTools
 	{
 		for (TPropertyValueIterator<FObjectProperty> PIter(InObject->GetClass(), InObject); PIter; ++PIter)
 		{
-			FObjectProperty* Property = PIter.Key();
-			void* Value = const_cast<void*>(PIter->Value);
+			const FObjectProperty* Property = PIter.Key();
+			const void* Value = PIter->Value;
 			if (const UObject* ValueObject = Property->GetPropertyValue(Value))
 			{
 				if (int32* TimesEncountered = CheckedObjects.Find(ValueObject))
