SPropertyEditorArray  : Allow drop onto arrays of any FObjectPropertyBase to match the behavior of non-array properties.

This enables drag & drop of assets from the Content Browser onto a property type of e.g. TArray<TSoftObjectPtr>.  We already supported drag & drop onto a property of e.g. TSoftObjectPtr but the user previously would have had to add them one by one to an array property.

diff --git a/Engine/Source/Editor/PropertyEditor/Private/UserInterface/PropertyEditor/SPropertyEditorArray.cpp b/Engine/Source/Editor/PropertyEditor/Private/UserInterface/PropertyEditor/SPropertyEditorArray.cpp
--- a/Engine/Source/Editor/PropertyEditor/Private/UserInterface/PropertyEditor/SPropertyEditorArray.cpp
+++ b/Engine/Source/Editor/PropertyEditor/Private/UserInterface/PropertyEditor/SPropertyEditorArray.cpp
@@ -85,10 +85,11 @@ bool SPropertyEditorArray::CanEdit() const
 
 FReply SPropertyEditorArray::OnDragDropTarget(const FGeometry& InGeometry, const FDragDropEvent& InDragDropEvent)
 {
-	FObjectProperty* ObjectProperty = nullptr;
+	FObjectPropertyBase* ObjectProperty = nullptr;
+
 	if (const FArrayProperty* ArrayProperty = CastField<FArrayProperty>(PropertyEditor->GetProperty()))
 	{
-		ObjectProperty = CastField<FObjectProperty>(ArrayProperty->Inner);
+		ObjectProperty = CastField<FObjectPropertyBase>(ArrayProperty->Inner);
 	}
 
 	// Only try to add entries if we are dropping on an asset array
@@ -116,10 +117,11 @@ FReply SPropertyEditorArray::OnDragDropTarget(const FGeometry& InGeometry, const
 
 bool SPropertyEditorArray::IsValidAssetDropOp(TSharedPtr<FDragDropOperation> InOperation)
 {
-	FObjectProperty* ObjectProperty = nullptr;
+	FObjectPropertyBase* ObjectProperty = nullptr;
+
 	if (const FArrayProperty* ArrayProperty = CastField<FArrayProperty>(PropertyEditor->GetProperty()))
 	{
-		ObjectProperty = CastField<FObjectProperty>(ArrayProperty->Inner);
+		ObjectProperty = CastField<FObjectPropertyBase>(ArrayProperty->Inner);
 	}
 
 	// Only try to add entries if we are dropping on an asset array
@@ -132,11 +134,11 @@ bool SPropertyEditorArray::IsValidAssetDropOp(TSharedPtr<FDragDropOperation> InO
 
 bool SPropertyEditorArray::WillAddValidElements(TSharedPtr<FDragDropOperation> InOperation)
 {
-	FObjectProperty* ObjectProperty = nullptr;
+	FObjectPropertyBase* ObjectProperty = nullptr;
 
 	if (const FArrayProperty* ArrayProperty = CastField<FArrayProperty>(PropertyEditor->GetProperty()))
 	{
-		ObjectProperty = CastField<FObjectProperty>(ArrayProperty->Inner);
+		ObjectProperty = CastField<FObjectPropertyBase>(ArrayProperty->Inner);
 	}
 
 	// Only try to add entries if we are dropping on an asset array
