[PropertyEditor] Category Copy/Paste accounts for EditConst flag (but allows pasting to EditCondition properties)

diff --git a/Engine/Source/Editor/PropertyEditor/Private/Presentation/PropertyEditor/PropertyEditor.cpp b/Engine/Source/Editor/PropertyEditor/Private/Presentation/PropertyEditor/PropertyEditor.cpp
--- a/Engine/Source/Editor/PropertyEditor/Private/Presentation/PropertyEditor/PropertyEditor.cpp
+++ b/Engine/Source/Editor/PropertyEditor/Private/Presentation/PropertyEditor/PropertyEditor.cpp
@@ -520,9 +520,9 @@ bool FPropertyEditor::DoesPassFilterRestrictions() const
 	return PropertyNode->HasNodeFlags( EPropertyNodeFlags::IsSeenDueToFiltering ) != 0;
 }
 
-bool FPropertyEditor::IsEditConst() const
+bool FPropertyEditor::IsEditConst(const bool bIncludeEditCondition) const
 {
-	return PropertyNode->IsEditConst();
+	return PropertyNode->IsEditConst(bIncludeEditCondition);
 }
 
 bool FPropertyEditor::SupportsEditConditionToggle() const
diff --git a/Engine/Source/Editor/PropertyEditor/Private/Presentation/PropertyEditor/PropertyEditor.h b/Engine/Source/Editor/PropertyEditor/Private/Presentation/PropertyEditor/PropertyEditor.h
--- a/Engine/Source/Editor/PropertyEditor/Private/Presentation/PropertyEditor/PropertyEditor.h
+++ b/Engine/Source/Editor/PropertyEditor/Private/Presentation/PropertyEditor/PropertyEditor.h
@@ -42,8 +42,9 @@ public:
 
 	void RequestRefresh();
 
-	/**	@return Whether the property is editconst */
-	bool IsEditConst() const;
+	/**	@return Whether the property is editconst. If bIncludeEditCondition is false,
+	 *	this will return true if the property is ever editable (as though it always passes any specified EditCondition). */
+	bool IsEditConst(const bool bIncludeEditCondition = true) const;
 
 	/** @return Whether this property should have an edit condition toggle. */
 	bool SupportsEditConditionToggle() const;
diff --git a/Engine/Source/Editor/PropertyEditor/Private/PropertyNode.cpp b/Engine/Source/Editor/PropertyEditor/Private/PropertyNode.cpp
--- a/Engine/Source/Editor/PropertyEditor/Private/PropertyNode.cpp
+++ b/Engine/Source/Editor/PropertyEditor/Private/PropertyNode.cpp
@@ -1325,7 +1325,7 @@ bool FPropertyNode::IsPropertyConst() const
 }
 
 /** @return whether this window's property is constant (can't be edited by the user) */
-bool FPropertyNode::IsEditConst() const
+bool FPropertyNode::IsEditConst(const bool bIncludeEditCondition) const
 {
 	if (bUpdateEditConstState || UpdateEditConstStateEpoch != PropertyEditorPolicy::Get().GetPolicyEpoch())
 	{
@@ -1354,7 +1354,7 @@ bool FPropertyNode::IsEditConst() const
 					}
 				}
 
-				if (CurParent->IsEditConst())
+				if (CurParent->IsEditConst(bIncludeEditCondition))
 				{
 					// An owning struct is edit const, so the child property is too
 					bIsEditConst = true;
@@ -1421,6 +1421,9 @@ bool FPropertyNode::IsEditConst() const
 			}
 		}
 
+		// this ignores EditCondition check below
+		bIsEditConstWithoutCondition = bIsEditConst;
+
 		// check edit condition
 		if (!bIsEditConst && HasEditCondition())
 		{
@@ -1430,7 +1433,7 @@ bool FPropertyNode::IsEditConst() const
 		bUpdateEditConstState = false;
 	}
 
-	return bIsEditConst;
+	return bIncludeEditCondition ? bIsEditConst : bIsEditConstWithoutCondition;
 }
 
 bool FPropertyNode::ShouldSkipSerialization() const
diff --git a/Engine/Source/Editor/PropertyEditor/Private/PropertyNode.h b/Engine/Source/Editor/PropertyEditor/Private/PropertyNode.h
--- a/Engine/Source/Editor/PropertyEditor/Private/PropertyNode.h
+++ b/Engine/Source/Editor/PropertyEditor/Private/PropertyNode.h
@@ -542,8 +542,8 @@ public:
 	 */
 	bool IsPropertyConst() const;
 
-	/** @return whether this window's property is constant (can't be edited by the user) */
-	bool IsEditConst() const;
+	/** @return whether this window's property is constant (can't be edited by the user). If bIncludeEditCondition is false, the property is considered editable regardless of the EditCondition result. */
+	bool IsEditConst(const bool bIncludeEditCondition = true) const;
 
 	/**
 	 * Returns whether this window's property should not be serialized (determined by the CPF_SkipSerialization flag).
@@ -1321,7 +1321,8 @@ protected:
 	* Cached state of flags that are expensive to update
 	* These update when values are changed in the details panel
 	*/
-	mutable bool bIsEditConst;
+	mutable bool bIsEditConst; // Includes EditCondition state
+	mutable bool bIsEditConstWithoutCondition; // Ignores EditCondition state
 	mutable bool bUpdateEditConstState;
 	mutable int32 UpdateEditConstStateEpoch;
 	mutable bool bDiffersFromDefault;
diff --git a/Engine/Source/Editor/PropertyEditor/Private/SDetailCategoryTableRow.cpp b/Engine/Source/Editor/PropertyEditor/Private/SDetailCategoryTableRow.cpp
--- a/Engine/Source/Editor/PropertyEditor/Private/SDetailCategoryTableRow.cpp
+++ b/Engine/Source/Editor/PropertyEditor/Private/SDetailCategoryTableRow.cpp
@@ -276,7 +276,7 @@ void SDetailCategoryTableRow::PopulateContextMenu(UToolMenu* ToolMenu)
 							return CanPasteCategory()
 								? NSLOCTEXT("PropertyView", "PasteCategoryProperties_ToolTip", "Paste the copied property values here")
 								// @note: this is specific to the constraint that the destination category has to match the source category (copied from) exactly 
-								: NSLOCTEXT("PropertyView", "CantPasteCategoryProperties_ToolTip", "The properties in this category don't match the contents of the clipboard");
+								: NSLOCTEXT("PropertyView", "CantPasteCategoryProperties_ToolTip", "The properties in this category don't match the contents of the clipboard, or the properties aren't editable");
 						}),
 						FSlateIcon(FCoreStyle::Get().GetStyleSetName(), "GenericCommands.Paste"),
 						PasteAction);
@@ -447,7 +447,7 @@ void SDetailCategoryTableRow::OnPasteCategory()
 			const FGuid OperationGuid = FGuid::NewGuid();			
 			for (const TPair<FName, FString>& KVP : PreviousClipboardData.PropertyValues)
 			{
-				if ( OnPasteFromTextDelegate.IsValid() )
+				if (OnPasteFromTextDelegate.IsValid())
 				{
 					OnPasteFromTextDelegate->Broadcast(KVP.Key.ToString(), KVP.Value, OperationGuid);
 				}
@@ -477,13 +477,17 @@ bool SDetailCategoryTableRow::CanPasteCategory()
 	}
 
 	const TArray<TSharedPtr<IPropertyHandle>> PropertyHandles = GetPropertyHandles(true);
+	const TArray<TSharedPtr<FPropertyNode>> PropertyNodes = GetPropertyNodesFromHandles(PropertyHandles);
 
-	// @note: Usually we'd check for IsEditConst or IsEditConditionMet, but if used for PP settings (for example),
-	// by default no properties are editable unless explicitly overridden, so this check would in all cases would prevent paste.
-	constexpr bool bHasEditables = true;
+	// @note: We allow pasting to properties that are disabled due to an EditCondition, but not those that are never editable (ie. VisibleAnywhere).
+	const bool bHasEditables = Algo::AnyOf(PropertyNodes, [](const TSharedPtr<FPropertyNode>& InPropertyNode)
+	{
+		constexpr bool bIncludeEditConditionForConstCheck = false;
+		return InPropertyNode->IsEditConst(bIncludeEditConditionForConstCheck);
+	});
 	
 	// No editable properties to write to
-	if constexpr (!bHasEditables)
+	if (!bHasEditables)
 	{
 		return false;
 	}
diff --git a/Engine/Source/Editor/PropertyEditor/Private/SDetailSingleItemRow.cpp b/Engine/Source/Editor/PropertyEditor/Private/SDetailSingleItemRow.cpp
--- a/Engine/Source/Editor/PropertyEditor/Private/SDetailSingleItemRow.cpp
+++ b/Engine/Source/Editor/PropertyEditor/Private/SDetailSingleItemRow.cpp
@@ -2,6 +2,7 @@
 
 #include "SDetailSingleItemRow.h"
 
+#include "Algo/AnyOf.h"
 #include "Algo/Compare.h"
 #include "DetailGroup.h"
 #include "DetailPropertyRow.h"
@@ -987,13 +988,17 @@ bool SDetailSingleItemRow::CanPasteGroup()
 	}
 
 	const TArray<TSharedPtr<IPropertyHandle>> GroupPropertyHandles = GetPropertyHandles(true);
+	const TArray<TSharedPtr<FPropertyNode>> GroupPropertyNodes = GetPropertyNodesFromHandles(GroupPropertyHandles);
+
+	// @note: We allow pasting to properties that are disabled due to an EditCondition, but not those that are never editable (ie. VisibleAnywhere).
+	const bool bHasEditables = Algo::AnyOf(GroupPropertyNodes, [](const TSharedPtr<FPropertyNode>& InPropertyNode)
+	{
+		constexpr bool bIncludeEditConditionForConstCheck = false;
+		return InPropertyNode->IsEditConst(bIncludeEditConditionForConstCheck);
+	});
 
-	// @note: Usually we'd check for IsEditConst or IsEditConditionMet, but if used for PP settings,
-	// by default no properties are editable unless explicitly overridden, so this check would in all cases would prevent paste.
-	constexpr bool bHasEditables = true;
-	
 	// No editable properties to write to
-	if constexpr (!bHasEditables)
+	if (!bHasEditables)
 	{
 		return false;
 	}
@@ -1104,7 +1109,7 @@ void SDetailSingleItemRow::PopulateContextMenu(UToolMenu* ToolMenu)
 							return CanPasteGroup()
 								? NSLOCTEXT("PropertyView", "PasteGroupProperties_ToolTip", "Paste the copied property values here")
 								// @note: this is specific to the constraint that the destination group has to match the source group (copied from) exactly 
-								: NSLOCTEXT("PropertyView", "CantPasteGroupProperties_ToolTip", "The properties in this group don't match the contents of the clipboard");
+								: NSLOCTEXT("PropertyView", "CantPasteGroupProperties_ToolTip", "The properties in this group don't match the contents of the clipboard, or the properties aren't editable");
 						});
 					}
 					else
@@ -1241,7 +1246,7 @@ void SDetailSingleItemRow::PopulateContextMenu(UToolMenu* ToolMenu)
 	}
 }
 
-TArray<TSharedPtr<IPropertyHandle>> SDetailSingleItemRow::GetPropertyHandles(const bool& bRecursive) const
+TArray<TSharedPtr<IPropertyHandle>> SDetailSingleItemRow::GetPropertyHandles(const bool bRecursive) const
 {
 	if (TArray<TSharedPtr<IPropertyHandle>> PropertyHandles = SDetailTableRowBase::GetPropertyHandles(bRecursive);
 		!PropertyHandles.IsEmpty())
@@ -1377,6 +1382,16 @@ bool SDetailSingleItemRow::CanCopyPropertyInternalName()
 
 bool SDetailSingleItemRow::CanPasteProperty() const
 {
+	TSharedPtr<FPropertyNode> PropertyNode = GetPropertyNode();
+	check(PropertyNode.IsValid());
+
+	// Check if the property is editable first (a failed EditCondition is still considered to be Editable, to make copying to PostProcess settings etc. practical)
+	constexpr bool bIncludeEditConditionForConstCheck = false;
+	if (PropertyNode->IsEditConst(bIncludeEditConditionForConstCheck)) // Ignore EditCondition state
+	{
+		return false;
+	}
+	
 	FString ClipboardContent;
 	FPropertyEditorClipboard::ClipboardPaste(ClipboardContent);
 
diff --git a/Engine/Source/Editor/PropertyEditor/Private/SDetailSingleItemRow.h b/Engine/Source/Editor/PropertyEditor/Private/SDetailSingleItemRow.h
--- a/Engine/Source/Editor/PropertyEditor/Private/SDetailSingleItemRow.h
+++ b/Engine/Source/Editor/PropertyEditor/Private/SDetailSingleItemRow.h
@@ -70,7 +70,7 @@ public:
 
 protected:
 	virtual void PopulateContextMenu(UToolMenu* ToolMenu) override;
-	virtual TArray<TSharedPtr<IPropertyHandle>> GetPropertyHandles(const bool& bRecursive = false) const override;
+	virtual TArray<TSharedPtr<IPropertyHandle>> GetPropertyHandles(const bool bRecursive = false) const override;
 
 private:
 	void OnCopyProperty();
diff --git a/Engine/Source/Editor/PropertyEditor/Private/SDetailTableRowBase.cpp b/Engine/Source/Editor/PropertyEditor/Private/SDetailTableRowBase.cpp
--- a/Engine/Source/Editor/PropertyEditor/Private/SDetailTableRowBase.cpp
+++ b/Engine/Source/Editor/PropertyEditor/Private/SDetailTableRowBase.cpp
@@ -115,14 +115,18 @@ void SDetailTableRowBase::PopulateContextMenu(UToolMenu* ToolMenu)
 	}
 }
 
-TArray<TSharedPtr<FPropertyNode>> SDetailTableRowBase::GetPropertyNodes(const bool& bRecursive) const
+TArray<TSharedPtr<FPropertyNode>> SDetailTableRowBase::GetPropertyNodes(const bool bRecursive) const
 {
 	TArray<TSharedPtr<IPropertyHandle>> PropertyHandles = GetPropertyHandles(bRecursive);
+	return GetPropertyNodesFromHandles(PropertyHandles);
+}
 
+TArray<TSharedPtr<FPropertyNode>> SDetailTableRowBase::GetPropertyNodesFromHandles(const TConstArrayView<TSharedPtr<IPropertyHandle>>& InPropertyHandles) const
+{
 	TArray<TSharedPtr<FPropertyNode>> PropertyNodes;
-	PropertyNodes.Reserve(PropertyHandles.Num());
+	PropertyNodes.Reserve(InPropertyHandles.Num());
 	
-	for (const TSharedPtr<IPropertyHandle>& PropertyHandle : PropertyHandles)
+	for (const TSharedPtr<IPropertyHandle>& PropertyHandle : InPropertyHandles)
 	{
 		if (PropertyHandle->IsValidHandle())
 		{
@@ -133,7 +137,7 @@ TArray<TSharedPtr<FPropertyNode>> SDetailTableRowBase::GetPropertyNodes(const bo
 	return PropertyNodes;
 }
 
-TArray<TSharedPtr<IPropertyHandle>> SDetailTableRowBase::GetPropertyHandles(const bool& bRecursive) const
+TArray<TSharedPtr<IPropertyHandle>> SDetailTableRowBase::GetPropertyHandles(const bool bRecursive) const
 {
 	TFunction<bool(const TSharedPtr<IDetailTreeNode>&, TArray<TSharedPtr<IPropertyHandle>>&)> AppendPropertyHandles;
 	AppendPropertyHandles = [&AppendPropertyHandles, bRecursive]
diff --git a/Engine/Source/Editor/PropertyEditor/Private/SDetailTableRowBase.h b/Engine/Source/Editor/PropertyEditor/Private/SDetailTableRowBase.h
--- a/Engine/Source/Editor/PropertyEditor/Private/SDetailTableRowBase.h
+++ b/Engine/Source/Editor/PropertyEditor/Private/SDetailTableRowBase.h
@@ -41,10 +41,13 @@ protected:
 	TSharedPtr<FPropertyUpdatedWidgetBuilder> PropertyUpdatedWidgetBuilder;
 
 	/** Retrieve all property nodes represented by this row, and it's children if Recursive specified. */
-	TArray<TSharedPtr<FPropertyNode>> GetPropertyNodes(const bool& bRecursive = false) const;
+	TArray<TSharedPtr<FPropertyNode>> GetPropertyNodes(const bool bRecursive = false) const;
+	
+	/** Retrieve all property nodes represented by the given property handles. */
+	TArray<TSharedPtr<FPropertyNode>> GetPropertyNodesFromHandles(const TConstArrayView<TSharedPtr<IPropertyHandle>>& InPropertyHandles) const;
 
 	/** Retrieve all property handles represented by this row, and it's children if Recursive specified. */
-	virtual TArray<TSharedPtr<IPropertyHandle>> GetPropertyHandles(const bool& bRecursive = false) const;
+	virtual TArray<TSharedPtr<IPropertyHandle>> GetPropertyHandles(const bool bRecursive = false) const;
 
 	/** Refreshes widget and associated list views. */
 	virtual void ForceRefresh();
