Fix crash in FStructurePropertyNode that sometimes happens when committing values when the StructProvider is null. This appears to be a regression from CL 33428724.

diff --git a/Engine/Source/Editor/PropertyEditor/Private/StructurePropertyNode.h b/Engine/Source/Editor/PropertyEditor/Private/StructurePropertyNode.h
--- a/Engine/Source/Editor/PropertyEditor/Private/StructurePropertyNode.h
+++ b/Engine/Source/Editor/PropertyEditor/Private/StructurePropertyNode.h
@@ -223,19 +223,25 @@ public:
 			}
 			else if (const FStructurePropertyNode* StructNode = CurrentNode->AsStructureNode())
 			{
-				TSharedPtr<IStructureDataProvider> TempStructProvider = StructNode->GetStructProvider();
-				if (TempStructProvider->IsPropertyIndirection())
+				if (TSharedPtr<IStructureDataProvider> TempStructProvider = StructNode->GetStructProvider())
 				{
-					// If the struct provider is marked as property indirection, it is assumed that it handles indirection between it's
-					// parent property, and some data inside that property (e.g. FInstancedStruct).
-					const FPropertyNode* ParentNode = CurrentNode->GetParentNode();
-					CurrentNode = ParentNode->FindComplexParent();
+					if (TempStructProvider->IsPropertyIndirection())
+					{
+						// If the struct provider is marked as property indirection, it is assumed that it handles indirection between it's
+						// parent property, and some data inside that property (e.g. FInstancedStruct).
+						const FPropertyNode* ParentNode = CurrentNode->GetParentNode();
+						CurrentNode = ParentNode->FindComplexParent();
+					}
+					else
+					{
+						TArray<TSharedPtr<FStructOnScope>> Instances;
+						TempStructProvider->GetInstances(Instances, StructNode->WeakCachedBaseStruct.Get());
+						return Instances.Num();
+					}
 				}
 				else
 				{
-					TArray<TSharedPtr<FStructOnScope>> Instances;
-					TempStructProvider->GetInstances(Instances, StructNode->WeakCachedBaseStruct.Get());
-					return Instances.Num();
+					return 0;
 				}
 			}
 		}
