[PropertyEditor] Category Copy/Paste accounts for EditConst flag

diff --git a/Engine/Source/Editor/PropertyEditor/Private/SDetailCategoryTableRow.cpp b/Engine/Source/Editor/PropertyEditor/Private/SDetailCategoryTableRow.cpp
--- a/Engine/Source/Editor/PropertyEditor/Private/SDetailCategoryTableRow.cpp
+++ b/Engine/Source/Editor/PropertyEditor/Private/SDetailCategoryTableRow.cpp
@@ -483,9 +483,9 @@ bool SDetailCategoryTableRow::CanPasteCategory()
 	const bool bHasEditables = Algo::AnyOf(PropertyNodes, [](const TSharedPtr<FPropertyNode>& InPropertyNode)
 	{
 		constexpr bool bIncludeEditConditionForConstCheck = false;
-		return InPropertyNode->IsEditConst(bIncludeEditConditionForConstCheck);
+		return !InPropertyNode->IsEditConst(bIncludeEditConditionForConstCheck);
 	});
-	
+
 	// No editable properties to write to
 	if (!bHasEditables)
 	{
diff --git a/Engine/Source/Editor/PropertyEditor/Private/SDetailSingleItemRow.cpp b/Engine/Source/Editor/PropertyEditor/Private/SDetailSingleItemRow.cpp
--- a/Engine/Source/Editor/PropertyEditor/Private/SDetailSingleItemRow.cpp
+++ b/Engine/Source/Editor/PropertyEditor/Private/SDetailSingleItemRow.cpp
@@ -994,7 +994,7 @@ bool SDetailSingleItemRow::CanPasteGroup()
 	const bool bHasEditables = Algo::AnyOf(GroupPropertyNodes, [](const TSharedPtr<FPropertyNode>& InPropertyNode)
 	{
 		constexpr bool bIncludeEditConditionForConstCheck = false;
-		return InPropertyNode->IsEditConst(bIncludeEditConditionForConstCheck);
+		return !InPropertyNode->IsEditConst(bIncludeEditConditionForConstCheck);
 	});
 
 	// No editable properties to write to
