Fix PropertyPermissionList breaking if an entry contained no Allow, Deny or DenyAll entries.
#rnx

diff --git a/Engine/Source/Editor/PropertyEditor/Private/PropertyPermissionList.cpp b/Engine/Source/Editor/PropertyEditor/Private/PropertyPermissionList.cpp
--- a/Engine/Source/Editor/PropertyEditor/Private/PropertyPermissionList.cpp
+++ b/Engine/Source/Editor/PropertyEditor/Private/PropertyPermissionList.cpp
@@ -83,10 +83,15 @@ void FPropertyPermissionList::ClearCacheAndBroadcast(TSoftObjectPtr<UStruct> Obj
 	}
 }
 
-void FPropertyPermissionList::AddPermissionList(TSoftObjectPtr<UStruct> Struct, const FNamePermissionList& PermissionList, EPropertyPermissionListRules Rules)
+void FPropertyPermissionList::AddPermissionList(TSoftObjectPtr<UStruct> Struct, const FNamePermissionList& PermissionList, const EPropertyPermissionListRules Rules, const TConstArrayView<FName> InAdditionalOwnerNames)
 {
 	FPropertyPermissionListEntry& Entry = RawPropertyPermissionList.FindOrAdd(Struct);
 	Entry.PermissionList = PermissionList;
+
+	// Track additional owners to prevent entry from being removed by FPropertyPermissionList::UnregisterOwner()
+	// PermissionList without AllowList/DenyList/DenyAll entries will have no owners and be deleted from RawPropertyPermissionList unless it provides AdditionalOwnerNames here
+	Entry.AdditionalOwnerNames = InAdditionalOwnerNames;
+
 	// Always use the most permissive rule previously set
 	if (Entry.Rules > Rules)
 	{
@@ -96,6 +101,11 @@ void FPropertyPermissionList::AddPermissionList(TSoftObjectPtr<UStruct> Struct,
 	ClearCacheAndBroadcast(Struct);
 }
 
+void FPropertyPermissionList::AddPermissionList(TSoftObjectPtr<UStruct> Struct, const FNamePermissionList& PermissionList, const EPropertyPermissionListRules Rules)
+{
+	AddPermissionList(Struct, PermissionList, Rules, {});
+}
+
 void FPropertyPermissionList::RemovePermissionList(TSoftObjectPtr<UStruct> Struct)
 {
 	if (RawPropertyPermissionList.Remove(Struct) > 0)
@@ -118,7 +128,8 @@ void FPropertyPermissionList::UnregisterOwner(const FName Owner)
 	for (TPair<TSoftObjectPtr<UStruct>, FPropertyPermissionListEntry>& Pair : RawPropertyPermissionList)
 	{
 		Pair.Value.PermissionList.UnregisterOwner(Owner);
-		if (Pair.Value.PermissionList.GetOwnerNames().Num() == 0)
+		Pair.Value.AdditionalOwnerNames.Remove(Owner);
+		if (Pair.Value.AdditionalOwnerNames.Num() == 0 && Pair.Value.PermissionList.GetOwnerNames().Num() == 0)
 		{
 			StructsToRemove.Add(Pair.Key);
 		}
@@ -313,6 +324,8 @@ const FNamePermissionList& FPropertyPermissionList::GetCachedPermissionListForSt
 		}
 
 		PermissionListEntry = &CachedPropertyPermissionList.Add(Struct, { MoveTemp(NewPermissionList),
+			// Note: The AdditionalOwnerNames are only used to prevent removal from RawPropertyPermissionList by calls to RemoveOwner()
+			TArray<FName>(),
 			// propagate the allow list all properties setting, this is to make AllowListAllSubclassProperties work properly 
 			// when a sub class was already previously cached from a base class that has AllowListAllSubclassProperties set on it and another class is then building its list off of the cached subclass list
 			bInOutShouldAllowListAllProperties ? EPropertyPermissionListRules::AllowListAllProperties : EPropertyPermissionListRules::UseExistingPermissionList });
diff --git a/Engine/Source/Editor/PropertyEditor/Public/PropertyPermissionList.h b/Engine/Source/Editor/PropertyEditor/Public/PropertyPermissionList.h
--- a/Engine/Source/Editor/PropertyEditor/Public/PropertyPermissionList.h
+++ b/Engine/Source/Editor/PropertyEditor/Public/PropertyPermissionList.h
@@ -41,6 +41,8 @@ enum class EPropertyPermissionListRules : uint8
 struct FPropertyPermissionListEntry
 {
     FNamePermissionList PermissionList;
+	// When the permission list does not contain any AllowList, DenyList or DenyListAll it needs an AdditionalOwnerName to be tracked here so stays alive without being removed
+	TArray<FName> AdditionalOwnerNames;
 	EPropertyPermissionListRules Rules = EPropertyPermissionListRules::UseExistingPermissionList;
 };
 
@@ -48,7 +50,10 @@ class PROPERTYEDITOR_API FPropertyPermissionList
 {
 public:
 	/** Add a set of rules for a specific base UStruct to determine which properties are visible in all details panels */
-	void AddPermissionList(TSoftObjectPtr<UStruct> Struct, const FNamePermissionList& PermissionList, EPropertyPermissionListRules Rules = EPropertyPermissionListRules::UseExistingPermissionList);
+	UE_DEPRECATED(5.5, "Call AddPermissionList with additional required arguments instead.")
+	void AddPermissionList(TSoftObjectPtr<UStruct> Struct, const FNamePermissionList& PermissionList, const EPropertyPermissionListRules Rules = EPropertyPermissionListRules::UseExistingPermissionList);
+	/** Add a set of rules for a specific base UStruct to determine which properties are visible in all details panels */
+	void AddPermissionList(TSoftObjectPtr<UStruct> Struct, const FNamePermissionList& PermissionList, const EPropertyPermissionListRules Rules, const TConstArrayView<FName> InAdditionalOwnerNames);
 	/** Remove a set of rules for a specific base UStruct to determine which properties are visible in all details panels */
 	void RemovePermissionList(TSoftObjectPtr<UStruct> Struct);
 	/** Remove all rules */
