PropertyEditor: Add property handle layout overrides, so that systems can inject their own custom layout selection logic

diff --git a/Engine/Source/Editor/PropertyEditor/Private/PropertyEditorModule.cpp b/Engine/Source/Editor/PropertyEditor/Private/PropertyEditorModule.cpp
--- a/Engine/Source/Editor/PropertyEditor/Private/PropertyEditorModule.cpp
+++ b/Engine/Source/Editor/PropertyEditor/Private/PropertyEditorModule.cpp
@@ -519,6 +519,24 @@ void FPropertyEditorModule::UnregisterCustomPropertyTypeLayout( FName PropertyTy
 	}
 }
 
+FDelegateHandle FPropertyEditorModule::RegisterPropertyHandleLayoutOverride(const FName PropertyTypeName,
+	const FPropertyHandleLayoutOverride& Delegate)
+{
+	PropertyHandleLayoutOverrides.Add(PropertyTypeName, Delegate);
+	return Delegate.GetHandle();
+}
+
+void FPropertyEditorModule::UnregisterPropertyHandleLayoutOverride(const FDelegateHandle DelegateHandle)
+{
+	for (auto It = PropertyHandleLayoutOverrides.CreateIterator(); It; ++It)
+	{
+		if (It->Value.GetHandle() == DelegateHandle)
+		{
+			It.RemoveCurrent();
+		}
+	}
+}
+
 void FPropertySection::AddCategory(FName CategoryName)
 {
 	checkf(!CategoryName.ToString().Contains(TEXT("|")), TEXT("Cannnot register a section mapping for a subcategory. Section: '%s', Category: '%s'"), *Name.ToString(), *CategoryName.ToString());
@@ -1036,6 +1054,16 @@ FPropertyTypeLayoutCallback FPropertyEditorModule::FindPropertyTypeLayoutCallbac
 {
 	if (PropertyTypeName != NAME_None)
 	{
+		for (auto It = PropertyHandleLayoutOverrides.CreateConstKeyIterator(PropertyTypeName); It; ++It)
+		{
+			const FName TypeOverrideName = It->Value.Execute(PropertyHandle);
+			if (TypeOverrideName != NAME_None)
+			{
+				PropertyTypeName = TypeOverrideName;
+				break;
+			}
+		}
+
 		const FPropertyTypeLayoutCallbackList* LayoutCallbacks = InstancedPropertyTypeLayoutMap.Find( PropertyTypeName );
 	
 		if( !LayoutCallbacks )
diff --git a/Engine/Source/Editor/PropertyEditor/Public/PropertyEditorDelegates.h b/Engine/Source/Editor/PropertyEditor/Public/PropertyEditorDelegates.h
--- a/Engine/Source/Editor/PropertyEditor/Public/PropertyEditorDelegates.h
+++ b/Engine/Source/Editor/PropertyEditor/Public/PropertyEditorDelegates.h
@@ -98,6 +98,12 @@ DECLARE_DELEGATE_RetVal_TwoParams(TSharedRef< SWidget >, FConstructExternalColum
 /** Delegate called to see if a property editing is enabled */
 DECLARE_DELEGATE_RetVal(bool, FIsPropertyEditingEnabled);
 
+/**
+ * A delegate that returns the name of a type that is to be used to look up the layout callback instead of the
+ * property's type. Return None to not provide an override.
+ */
+DECLARE_DELEGATE_RetVal_OneParam(FName, FPropertyHandleLayoutOverride, const IPropertyHandle& Handle)
+
 /**
  * A delegate which is called after properties have been edited and PostEditChange has been called on all objects.
  * This can be used to safely make changes to data that the details panel is observing instead of during PostEditChange (which is
diff --git a/Engine/Source/Editor/PropertyEditor/Public/PropertyEditorModule.h b/Engine/Source/Editor/PropertyEditor/Public/PropertyEditorModule.h
--- a/Engine/Source/Editor/PropertyEditor/Public/PropertyEditorModule.h
+++ b/Engine/Source/Editor/PropertyEditor/Public/PropertyEditorModule.h
@@ -297,6 +297,21 @@ public:
 	 */
 	virtual void UnregisterCustomPropertyTypeLayout( FName PropertyTypeName, TSharedPtr<IPropertyTypeIdentifier> InIdentifier = nullptr);
 
+	/**
+	 * Registers a property layout override callback
+	 *
+	 * @param PropertyTypeName Name of the property type (e.g. struct FName) to override
+	 * @param Delegate The delegate to call when we're looking for a property handle's layout override
+	 */
+	virtual FDelegateHandle RegisterPropertyHandleLayoutOverride(FName PropertyTypeName, const FPropertyHandleLayoutOverride& Delegate);
+
+	/**
+	 * Unregisters a property layout override callback
+	 *
+	 * @param DelegateHandle The handle returned by RegisterPropertyHandleLayoutOverride
+	 */
+	virtual void UnregisterPropertyHandleLayoutOverride(FDelegateHandle DelegateHandle);
+
 	/**
 	 * Find an existing section or create a section for a class.
 	 * 
@@ -473,6 +488,8 @@ private:
 	FCustomDetailLayoutNameMap ClassNameToDetailLayoutNameMap;
 	/** A mapping of property names to property type layout delegates, called when querying for custom property layouts */
 	FCustomPropertyTypeLayoutMap GlobalPropertyTypeToLayoutMap;
+	/** Registered list of override callbacks. First one returning non-None is going to be applied (if any). */
+	TMultiMap<FName, FPropertyHandleLayoutOverride> PropertyHandleLayoutOverrides;
 	/** A mapping of class names to section mappings. */
 	TMap<FName, TSharedPtr<FClassSectionMapping>> ClassSectionMappings;
 	/** Event to be called when a property editor is opened */
