Ensure SingleProperty Handle generation results in valid handle for built-in struct types

Many built-in PropertyHandles create child PropertyHandles based on the ValueNode (see. FPropertyHandleVector, FPropertyHandleRotator).
Moving the creation of the PropertyHandle until after RebuildChildren is called ensures that the PropertyNodes that the PropertyHandles rely upon have valid data to generate child Handles
Added unit test to ensure behaviour is tested

diff --git a/Engine/Source/Editor/PropertyEditor/Private/SSingleProperty.cpp b/Engine/Source/Editor/PropertyEditor/Private/SSingleProperty.cpp
--- a/Engine/Source/Editor/PropertyEditor/Private/SSingleProperty.cpp
+++ b/Engine/Source/Editor/PropertyEditor/Private/SSingleProperty.cpp
@@ -212,7 +212,6 @@ bool SSingleProperty::GeneratePropertyCustomization()
 		//bIsAcceptableProperty &= !( Property->IsA( FArrayProperty::StaticClass() ) || (Property->ArrayDim > 1 && ValueNode->GetArrayIndex() == INDEX_NONE) );
 		// not a struct property unless its a built in type like a vector
 		//bIsAcceptableProperty &= ( !Property->IsA( FStructProperty::StaticClass() ) || PropertyEditorHelpers::IsBuiltInStructProperty( Property ) );
-		PropertyHandle = PropertyEditorHelpers::GetPropertyHandle(ValueNode.ToSharedRef(), NotifyHook, PropertyUtilities);
 	}
 
 	if( bIsAcceptableProperty )
@@ -221,6 +220,8 @@ bool SSingleProperty::GeneratePropertyCustomization()
 
 		TSharedRef< FPropertyEditor > PropertyEditor = FPropertyEditor::Create( ValueNode.ToSharedRef(), TSharedPtr< IPropertyUtilities >( PropertyUtilities ).ToSharedRef() );
 		ValueNode->SetDisplayNameOverride( NameOverride );
+		
+		PropertyHandle = PropertyEditorHelpers::GetPropertyHandle(ValueNode.ToSharedRef(), NotifyHook, PropertyUtilities);
 
 		TSharedPtr<SHorizontalBox> HorizontalBox;
 
@@ -284,6 +285,12 @@ bool SSingleProperty::GeneratePropertyCustomization()
 	}
 	else
 	{
+		if (ValueNode.IsValid())
+		{
+			// Still create a PropertyHandle, though it may not be fully initialized as the ValueNode will not have had children rebuilt
+			PropertyHandle = PropertyEditorHelpers::GetPropertyHandle(ValueNode.ToSharedRef(), NotifyHook, PropertyUtilities);
+		}
+		
 		ChildSlot
 		[
 			SNew(STextBlock)
diff --git a/Engine/Source/Editor/PropertyEditor/Private/Tests/SinglePropertyTests.cpp b/Engine/Source/Editor/PropertyEditor/Private/Tests/SinglePropertyTests.cpp
new file mode 100644
--- /dev/null
+++ b/Engine/Source/Editor/PropertyEditor/Private/Tests/SinglePropertyTests.cpp
@@ -0,0 +1,54 @@
+﻿// Copyright Epic Games, Inc. All Rights Reserved.
+
+#include "SinglePropertyTests.h"
+#include "ISinglePropertyView.h"
+#include "PropertyHandle.h"
+
+IMPLEMENT_SIMPLE_AUTOMATION_TEST(FPropertyEditorTests_SingleProperty, "PropertyEditor.SingleProperty", EAutomationTestFlags::EditorContext | EAutomationTestFlags::EngineFilter)
+
+#define REQUIRE(expr) \
+do { \
+	if (!(expr)) { \
+		AddError(FString::Printf(TEXT("Expression '%s' to be true."), TEXT(#expr))); \
+		return false; \
+	} \
+} while(false)
+
+#define EXPECT(expr) \
+do { \
+	if (!(expr)) { \
+		AddError(FString::Printf(TEXT("Expression '%s' to be true."), TEXT(#expr))); \
+	} \
+} while(false)
+
+bool FPropertyEditorTests_SingleProperty::RunTest(const FString& Parameters)
+{
+	const FVector ExpectedValue(1.0, 2.0, 3.0);
+
+	UPropertyEditorSinglePropertyTestClass* TestObject = NewObject<UPropertyEditorSinglePropertyTestClass>();
+	TestObject->Vector = ExpectedValue;
+
+	FPropertyEditorModule& PropertyEditorModule = FModuleManager::GetModuleChecked<FPropertyEditorModule>("PropertyEditor");
+	TSharedPtr<ISinglePropertyView> PropertyView =
+		PropertyEditorModule.CreateSingleProperty(TestObject, "Vector", FSinglePropertyParams{});
+
+	REQUIRE(PropertyView.IsValid());
+
+	const TSharedPtr<IPropertyHandle> PropertyHandle = PropertyView->GetPropertyHandle();
+
+	REQUIRE(PropertyHandle.IsValid());
+	REQUIRE(PropertyHandle->IsValidHandle());
+
+	FVector ReadVector;
+	FPropertyAccess::Result ReadResult = PropertyHandle->GetValue(ReadVector);
+
+	EXPECT(FPropertyAccess::Success == ReadResult);
+	EXPECT(ExpectedValue.X == ReadVector.X);
+	EXPECT(ExpectedValue.Y == ReadVector.Y);
+	EXPECT(ExpectedValue.Z == ReadVector.Z);
+
+	return true;
+}
+
+#undef REQUIRE
+#undef EXPECT
\ No newline at end of file
diff --git a/Engine/Source/Editor/PropertyEditor/Private/Tests/SinglePropertyTests.h b/Engine/Source/Editor/PropertyEditor/Private/Tests/SinglePropertyTests.h
new file mode 100644
--- /dev/null
+++ b/Engine/Source/Editor/PropertyEditor/Private/Tests/SinglePropertyTests.h
@@ -0,0 +1,25 @@
+﻿// Copyright Epic Games, Inc. All Rights Reserved.
+#pragma once
+
+//#if WITH_DEV_AUTOMATION_TESTS
+
+#include "GameFramework/Actor.h"
+#include "Editor/PropertyEditorTestObject.h"
+#include "Modules/ModuleManager.h"
+#include "Misc/AutomationTest.h"
+#include "PropertyEditorModule.h"
+#include "CoreUObject.h"
+
+#include "SinglePropertyTests.generated.h"
+
+UCLASS()
+class UPropertyEditorSinglePropertyTestClass : public UObject
+{
+	GENERATED_BODY()
+public:
+
+	UPROPERTY(EditAnywhere, Category = "Properties")
+	FVector Vector;
+};
+
+//#endif
\ No newline at end of file
