PropertyEditor: Display the internal & display name of the property in the copy context action tooltip

diff --git a/Engine/Source/Editor/PropertyEditor/Private/SDetailSingleItemRow.cpp b/Engine/Source/Editor/PropertyEditor/Private/SDetailSingleItemRow.cpp
--- a/Engine/Source/Editor/PropertyEditor/Private/SDetailSingleItemRow.cpp
+++ b/Engine/Source/Editor/PropertyEditor/Private/SDetailSingleItemRow.cpp
@@ -1135,10 +1135,12 @@ void SDetailSingleItemRow::PopulateContextMenu(UToolMenu* ToolMenu)
 			FUIAction CopyDisplayNameAction = FExecuteAction::CreateSP(this, &SDetailSingleItemRow::OnCopyPropertyDisplayName);
 			CopyDisplayNameAction.CanExecuteAction = FCanExecuteAction::CreateSP(this, &SDetailSingleItemRow::CanCopyPropertyDisplayName);
 
+			static const FTextFormat TooltipFormat = NSLOCTEXT("PropertyView", "CopyPropertyDisplayName_ToolTip", "Copy the display name of this property to the system clipboard:\n{0}");
+
 			EditSection.AddMenuEntry(
 				TEXT("CopyDisplayName"),
 				NSLOCTEXT("PropertyView", "CopyPropertyDisplayName", "Copy Display Name"),
-				NSLOCTEXT("PropertyView", "CopyPropertyDisplayName_ToolTip", "Copy the display name of this property to the system clipboard."),
+				FText::Format(TooltipFormat, GetPropertyDisplayName()),
 				FSlateIcon(FCoreStyle::Get().GetStyleSetName(), "GenericCommands.Copy"),
 				CopyDisplayNameAction);
 		}
@@ -1148,10 +1150,12 @@ void SDetailSingleItemRow::PopulateContextMenu(UToolMenu* ToolMenu)
 			FUIAction CopyInternalNameAction = FExecuteAction::CreateSP(this, &SDetailSingleItemRow::OnCopyPropertyInternalName);
 			CopyInternalNameAction.CanExecuteAction = FCanExecuteAction::CreateSP(this, &SDetailSingleItemRow::CanCopyPropertyInternalName);
 
+			static const FTextFormat TooltipFormat = NSLOCTEXT("PropertyView", "CopyPropertyInternalName_ToolTip", "Copy the internal name of this property to the system clipboard:\n{0}");
+
 			EditSection.AddMenuEntry(
 				TEXT("CopyInternalName"),
 				NSLOCTEXT("PropertyView", "CopyPropertyInternalName", "Copy Internal Name"),
-				NSLOCTEXT("PropertyView", "CopyPropertyInternalName_ToolTip", "Copy the internal name of this property to the system clipboard."),
+				FText::Format(TooltipFormat, FText::FromString(GetPropertyInternalName())),
 				FSlateIcon(FCoreStyle::Get().GetStyleSetName(), "GenericCommands.Copy"),
 				CopyInternalNameAction);
 		}
@@ -1277,107 +1281,77 @@ void SDetailSingleItemRow::OnCopyProperty()
 	}
 }
 
-void SDetailSingleItemRow::OnCopyPropertyDisplayName()
+FText SDetailSingleItemRow::GetPropertyDisplayName() const
 {
 	if (!OwnerTreeNode.IsValid())
 	{
-		return;
+		return { };
 	}
 
-	TSharedPtr<FPropertyNode> PropertyNode = GetPropertyNode();
+	const TSharedPtr<FPropertyNode> PropertyNode = GetPropertyNode();
 	if (!PropertyNode.IsValid())
 	{
-		return;
+		return { };
 	}
 
 	if (PropertyNode->IsOptionalValueNode())
 	{
-		FPropertyEditorClipboard::ClipboardCopy(*PropertyNode->GetParentNode()->GetDisplayName().ToString());
-	}
-	else
-	{
-		FPropertyEditorClipboard::ClipboardCopy(*PropertyNode->GetDisplayName().ToString());
+		return PropertyNode->GetParentNode()->GetDisplayName();
 	}
+
+	return PropertyNode->GetDisplayName();
 }
 
-bool SDetailSingleItemRow::CanCopyPropertyDisplayName()
+void SDetailSingleItemRow::OnCopyPropertyDisplayName()
 {
-	if (!OwnerTreeNode.IsValid())
-	{
-		return false;
-	}
-
-	TSharedPtr<FPropertyNode> PropertyNode = GetPropertyNode();
-	if (!PropertyNode.IsValid())
-	{
-		return false;
-	}
-
-	if (PropertyNode->IsOptionalValueNode() && PropertyNode->GetParentNode()->GetDisplayName().IsEmpty())
-	{
-		return false;
-	}
-	else if (PropertyNode->GetDisplayName().IsEmpty())
-	{
-		return false;
-	}
+	FPropertyEditorClipboard::ClipboardCopy(*GetPropertyDisplayName().ToString());
+}
 
-	return true;
+bool SDetailSingleItemRow::CanCopyPropertyDisplayName()
+{
+	return !GetPropertyDisplayName().IsEmpty();
 }
 
-void SDetailSingleItemRow::OnCopyPropertyInternalName()
+FString SDetailSingleItemRow::GetPropertyInternalName() const
 {
 	if (!OwnerTreeNode.IsValid())
 	{
-		return;
+		return { };
 	}
 
-	TSharedPtr<FPropertyNode> PropertyNode = GetPropertyNode();
+	const TSharedPtr<FPropertyNode> PropertyNode = GetPropertyNode();
 	if (!PropertyNode.IsValid())
 	{
-		return;
+		return { };
 	}
 
 	const FProperty* Property = PropertyNode->IsOptionalValueNode() ? PropertyNode->GetParentNode()->GetProperty() : PropertyNode->GetProperty();
 	if (!Property)
 	{
-		return;
+		return { };
 	}
 
-	if (const UStruct* OwnerStruct = Property->GetOwnerStruct())
+	const UStruct* OwnerStruct = Property->GetOwnerStruct();
+	if (!OwnerStruct)
 	{
-		FPropertyEditorClipboard::ClipboardCopy(*OwnerStruct->GetAuthoredNameForField(Property));
+		return { };
 	}
+
+	return OwnerStruct->GetAuthoredNameForField(Property);
 }
 
-bool SDetailSingleItemRow::CanCopyPropertyInternalName()
+void SDetailSingleItemRow::OnCopyPropertyInternalName()
 {
-	if (!OwnerTreeNode.IsValid())
-	{
-		return false;
-	}
-
-	TSharedPtr<FPropertyNode> PropertyNode = GetPropertyNode();
-	if (!PropertyNode.IsValid())
+	const FString InternalName = GetPropertyInternalName();
+	if (!InternalName.IsEmpty())
 	{
-		return false;
-	}
-
-	const FProperty* Property = PropertyNode->IsOptionalValueNode() ? PropertyNode->GetParentNode()->GetProperty() : PropertyNode->GetProperty();
-	if (!Property)
-	{
-		return false;
-	}
-
-	if (const UStruct* OwnerStruct = Property->GetOwnerStruct())
-	{
-		if (OwnerStruct->GetAuthoredNameForField(Property).IsEmpty())
-		{
-			return false;
-		}
+		FPropertyEditorClipboard::ClipboardCopy(*InternalName);
 	}
+}
 
-	return true;
+bool SDetailSingleItemRow::CanCopyPropertyInternalName()
+{
+	return !GetPropertyInternalName().IsEmpty();
 }
 
 bool SDetailSingleItemRow::CanPasteProperty() const
diff --git a/Engine/Source/Editor/PropertyEditor/Private/SDetailSingleItemRow.h b/Engine/Source/Editor/PropertyEditor/Private/SDetailSingleItemRow.h
--- a/Engine/Source/Editor/PropertyEditor/Private/SDetailSingleItemRow.h
+++ b/Engine/Source/Editor/PropertyEditor/Private/SDetailSingleItemRow.h
@@ -86,6 +86,9 @@ private:
 	void OnPasteGroup();
 	bool CanPasteGroup();
 
+	FString GetPropertyInternalName() const;
+	FText GetPropertyDisplayName() const;
+
 	/**
 	 * Returns the TSharedRef<SWidget> for the Name column widget
 	 * 
