[Backout] - CL34072649
Original CL Desc
-----------------------------------------------------------------
PropertyEditor: Display the internal & display name of the property in the copy context action tooltip

diff --git a/Engine/Source/Editor/PropertyEditor/Private/SDetailSingleItemRow.cpp b/Engine/Source/Editor/PropertyEditor/Private/SDetailSingleItemRow.cpp
--- a/Engine/Source/Editor/PropertyEditor/Private/SDetailSingleItemRow.cpp
+++ b/Engine/Source/Editor/PropertyEditor/Private/SDetailSingleItemRow.cpp
@@ -1135,12 +1135,10 @@ void SDetailSingleItemRow::PopulateContextMenu(UToolMenu* ToolMenu)
 			FUIAction CopyDisplayNameAction = FExecuteAction::CreateSP(this, &SDetailSingleItemRow::OnCopyPropertyDisplayName);
 			CopyDisplayNameAction.CanExecuteAction = FCanExecuteAction::CreateSP(this, &SDetailSingleItemRow::CanCopyPropertyDisplayName);
 
-			static const FTextFormat TooltipFormat = NSLOCTEXT("PropertyView", "CopyPropertyDisplayName_ToolTip", "Copy the display name of this property to the system clipboard:\n{0}");
-
 			EditSection.AddMenuEntry(
 				TEXT("CopyDisplayName"),
 				NSLOCTEXT("PropertyView", "CopyPropertyDisplayName", "Copy Display Name"),
-				FText::Format(TooltipFormat, GetPropertyDisplayName()),
+				NSLOCTEXT("PropertyView", "CopyPropertyDisplayName_ToolTip", "Copy the display name of this property to the system clipboard."),
 				FSlateIcon(FCoreStyle::Get().GetStyleSetName(), "GenericCommands.Copy"),
 				CopyDisplayNameAction);
 		}
@@ -1150,12 +1148,10 @@ void SDetailSingleItemRow::PopulateContextMenu(UToolMenu* ToolMenu)
 			FUIAction CopyInternalNameAction = FExecuteAction::CreateSP(this, &SDetailSingleItemRow::OnCopyPropertyInternalName);
 			CopyInternalNameAction.CanExecuteAction = FCanExecuteAction::CreateSP(this, &SDetailSingleItemRow::CanCopyPropertyInternalName);
 
-			static const FTextFormat TooltipFormat = NSLOCTEXT("PropertyView", "CopyPropertyInternalName_ToolTip", "Copy the internal name of this property to the system clipboard:\n{0}");
-
 			EditSection.AddMenuEntry(
 				TEXT("CopyInternalName"),
 				NSLOCTEXT("PropertyView", "CopyPropertyInternalName", "Copy Internal Name"),
-				FText::Format(TooltipFormat, FText::FromString(GetPropertyInternalName())),
+				NSLOCTEXT("PropertyView", "CopyPropertyInternalName_ToolTip", "Copy the internal name of this property to the system clipboard."),
 				FSlateIcon(FCoreStyle::Get().GetStyleSetName(), "GenericCommands.Copy"),
 				CopyInternalNameAction);
 		}
@@ -1281,77 +1277,107 @@ void SDetailSingleItemRow::OnCopyProperty()
 	}
 }
 
-FText SDetailSingleItemRow::GetPropertyDisplayName() const
+void SDetailSingleItemRow::OnCopyPropertyDisplayName()
 {
 	if (!OwnerTreeNode.IsValid())
 	{
-		return { };
+		return;
 	}
 
-	const TSharedPtr<FPropertyNode> PropertyNode = GetPropertyNode();
+	TSharedPtr<FPropertyNode> PropertyNode = GetPropertyNode();
 	if (!PropertyNode.IsValid())
 	{
-		return { };
+		return;
 	}
 
 	if (PropertyNode->IsOptionalValueNode())
 	{
-		return PropertyNode->GetParentNode()->GetDisplayName();
+		FPropertyEditorClipboard::ClipboardCopy(*PropertyNode->GetParentNode()->GetDisplayName().ToString());
+	}
+	else
+	{
+		FPropertyEditorClipboard::ClipboardCopy(*PropertyNode->GetDisplayName().ToString());
 	}
-
-	return PropertyNode->GetDisplayName();
-}
-
-void SDetailSingleItemRow::OnCopyPropertyDisplayName()
-{
-	FPropertyEditorClipboard::ClipboardCopy(*GetPropertyDisplayName().ToString());
 }
 
 bool SDetailSingleItemRow::CanCopyPropertyDisplayName()
 {
-	return !GetPropertyDisplayName().IsEmpty();
+	if (!OwnerTreeNode.IsValid())
+	{
+		return false;
+	}
+
+	TSharedPtr<FPropertyNode> PropertyNode = GetPropertyNode();
+	if (!PropertyNode.IsValid())
+	{
+		return false;
+	}
+
+	if (PropertyNode->IsOptionalValueNode() && PropertyNode->GetParentNode()->GetDisplayName().IsEmpty())
+	{
+		return false;
+	}
+	else if (PropertyNode->GetDisplayName().IsEmpty())
+	{
+		return false;
+	}
+
+	return true;
 }
 
-FString SDetailSingleItemRow::GetPropertyInternalName() const
+void SDetailSingleItemRow::OnCopyPropertyInternalName()
 {
 	if (!OwnerTreeNode.IsValid())
 	{
-		return { };
+		return;
 	}
 
-	const TSharedPtr<FPropertyNode> PropertyNode = GetPropertyNode();
+	TSharedPtr<FPropertyNode> PropertyNode = GetPropertyNode();
 	if (!PropertyNode.IsValid())
 	{
-		return { };
+		return;
 	}
 
 	const FProperty* Property = PropertyNode->IsOptionalValueNode() ? PropertyNode->GetParentNode()->GetProperty() : PropertyNode->GetProperty();
 	if (!Property)
 	{
-		return { };
+		return;
 	}
 
-	const UStruct* OwnerStruct = Property->GetOwnerStruct();
-	if (!OwnerStruct)
+	if (const UStruct* OwnerStruct = Property->GetOwnerStruct())
 	{
-		return { };
+		FPropertyEditorClipboard::ClipboardCopy(*OwnerStruct->GetAuthoredNameForField(Property));
 	}
-
-	return OwnerStruct->GetAuthoredNameForField(Property);
 }
 
-void SDetailSingleItemRow::OnCopyPropertyInternalName()
+bool SDetailSingleItemRow::CanCopyPropertyInternalName()
 {
-	const FString InternalName = GetPropertyInternalName();
-	if (!InternalName.IsEmpty())
+	if (!OwnerTreeNode.IsValid())
 	{
-		FPropertyEditorClipboard::ClipboardCopy(*InternalName);
+		return false;
 	}
-}
 
-bool SDetailSingleItemRow::CanCopyPropertyInternalName()
-{
-	return !GetPropertyInternalName().IsEmpty();
+	TSharedPtr<FPropertyNode> PropertyNode = GetPropertyNode();
+	if (!PropertyNode.IsValid())
+	{
+		return false;
+	}
+
+	const FProperty* Property = PropertyNode->IsOptionalValueNode() ? PropertyNode->GetParentNode()->GetProperty() : PropertyNode->GetProperty();
+	if (!Property)
+	{
+		return false;
+	}
+
+	if (const UStruct* OwnerStruct = Property->GetOwnerStruct())
+	{
+		if (OwnerStruct->GetAuthoredNameForField(Property).IsEmpty())
+		{
+			return false;
+		}
+	}
+
+	return true;
 }
 
 bool SDetailSingleItemRow::CanPasteProperty() const
diff --git a/Engine/Source/Editor/PropertyEditor/Private/SDetailSingleItemRow.h b/Engine/Source/Editor/PropertyEditor/Private/SDetailSingleItemRow.h
--- a/Engine/Source/Editor/PropertyEditor/Private/SDetailSingleItemRow.h
+++ b/Engine/Source/Editor/PropertyEditor/Private/SDetailSingleItemRow.h
@@ -86,9 +86,6 @@ private:
 	void OnPasteGroup();
 	bool CanPasteGroup();
 
-	FString GetPropertyInternalName() const;
-	FText GetPropertyDisplayName() const;
-
 	/**
 	 * Returns the TSharedRef<SWidget> for the Name column widget
 	 * 
