Fix include paths
Mac compilation issues warning on monolithic include.

#rnx

diff --git a/Engine/Source/Editor/PropertyEditor/Private/Tests/SinglePropertyTests.h b/Engine/Source/Editor/PropertyEditor/Private/Tests/SinglePropertyTests.h
--- a/Engine/Source/Editor/PropertyEditor/Private/Tests/SinglePropertyTests.h
+++ b/Engine/Source/Editor/PropertyEditor/Private/Tests/SinglePropertyTests.h
@@ -8,7 +8,7 @@
 #include "Modules/ModuleManager.h"
 #include "Misc/AutomationTest.h"
 #include "PropertyEditorModule.h"
-#include "CoreUObject.h"
+#include "UObject/Object.h"
 
 #include "SinglePropertyTests.generated.h"
 
