[Backout] - CL34072186
Original CL Desc
-----------------------------------------------------------------
Ensure SingleProperty Handle generation results in valid handle for built-in struct types

Many built-in PropertyHandles create child PropertyHandles based on the ValueNode (see. FPropertyHandleVector, FPropertyHandleRotator).
Moving the creation of the PropertyHandle until after RebuildChildren is called ensures that the PropertyNodes that the PropertyHandles rely upon have valid data to generate child Handles
Added unit test to ensure behaviour is tested

diff --git a/Engine/Source/Editor/PropertyEditor/Private/SSingleProperty.cpp b/Engine/Source/Editor/PropertyEditor/Private/SSingleProperty.cpp
--- a/Engine/Source/Editor/PropertyEditor/Private/SSingleProperty.cpp
+++ b/Engine/Source/Editor/PropertyEditor/Private/SSingleProperty.cpp
@@ -212,6 +212,7 @@ bool SSingleProperty::GeneratePropertyCustomization()
 		//bIsAcceptableProperty &= !( Property->IsA( FArrayProperty::StaticClass() ) || (Property->ArrayDim > 1 && ValueNode->GetArrayIndex() == INDEX_NONE) );
 		// not a struct property unless its a built in type like a vector
 		//bIsAcceptableProperty &= ( !Property->IsA( FStructProperty::StaticClass() ) || PropertyEditorHelpers::IsBuiltInStructProperty( Property ) );
+		PropertyHandle = PropertyEditorHelpers::GetPropertyHandle(ValueNode.ToSharedRef(), NotifyHook, PropertyUtilities);
 	}
 
 	if( bIsAcceptableProperty )
@@ -220,8 +221,6 @@ bool SSingleProperty::GeneratePropertyCustomization()
 
 		TSharedRef< FPropertyEditor > PropertyEditor = FPropertyEditor::Create( ValueNode.ToSharedRef(), TSharedPtr< IPropertyUtilities >( PropertyUtilities ).ToSharedRef() );
 		ValueNode->SetDisplayNameOverride( NameOverride );
-		
-		PropertyHandle = PropertyEditorHelpers::GetPropertyHandle(ValueNode.ToSharedRef(), NotifyHook, PropertyUtilities);
 
 		TSharedPtr<SHorizontalBox> HorizontalBox;
 
@@ -285,12 +284,6 @@ bool SSingleProperty::GeneratePropertyCustomization()
 	}
 	else
 	{
-		if (ValueNode.IsValid())
-		{
-			// Still create a PropertyHandle, though it may not be fully initialized as the ValueNode will not have had children rebuilt
-			PropertyHandle = PropertyEditorHelpers::GetPropertyHandle(ValueNode.ToSharedRef(), NotifyHook, PropertyUtilities);
-		}
-		
 		ChildSlot
 		[
 			SNew(STextBlock)
diff --git a/Engine/Source/Editor/PropertyEditor/Private/Tests/SinglePropertyTests.cpp b/Engine/Source/Editor/PropertyEditor/Private/Tests/SinglePropertyTests.cpp
deleted file mode 100644
diff --git a/Engine/Source/Editor/PropertyEditor/Private/Tests/SinglePropertyTests.h b/Engine/Source/Editor/PropertyEditor/Private/Tests/SinglePropertyTests.h
deleted file mode 100644
