Details View: Fixed issues with property customization expansion
When restorin expansion for property customizations, we need to do 2 passes.
Sometimes the first pass can removed the expansion states needed for the second pass.
Collect the saved expansion state once, and use the same data on both expansion passes so that the first pass cannot mess things up.

diff --git a/Engine/Source/Editor/PropertyEditor/Private/IDetailsViewPrivate.h b/Engine/Source/Editor/PropertyEditor/Private/IDetailsViewPrivate.h
--- a/Engine/Source/Editor/PropertyEditor/Private/IDetailsViewPrivate.h
+++ b/Engine/Source/Editor/PropertyEditor/Private/IDetailsViewPrivate.h
@@ -13,6 +13,7 @@ class FEditConditionParser;
 class FNotifyHook;
 class IDetailPropertyExtensionHandler;
 class IDetailRootObjectCustomization;
+class FStringPrefixTree;
 
 class IDetailsViewPrivate : public IDetailsView
 {
@@ -149,6 +150,11 @@ public:
 	*/
 	virtual void RestoreExpandedItems(TSharedRef<FPropertyNode> StartNode) = 0;
 
+	/**
+	* Restores the expansion state of property nodes for the selected object set, with optionally provided external expansion states.
+	*/
+	virtual void RestoreExpandedItems(TSharedRef<FPropertyNode> StartNode, TMap<UStruct*, FStringPrefixTree>* OptionalExpansionStates) = 0;
+
 	/**
 	* Mark node as animating, useful if animating during behaviors that trigger widget reconstruction.
 	* @param InNode				The slate property node to animate.
diff --git a/Engine/Source/Editor/PropertyEditor/Private/SDetailsView.cpp b/Engine/Source/Editor/PropertyEditor/Private/SDetailsView.cpp
--- a/Engine/Source/Editor/PropertyEditor/Private/SDetailsView.cpp
+++ b/Engine/Source/Editor/PropertyEditor/Private/SDetailsView.cpp
@@ -969,6 +969,18 @@ void SDetailsView::PreSetObject(int32 InNewNumObjects)
 }
 
 
+void SDetailsView::GetExpansionState(UStruct* BaseStruct, FStringPrefixTree& OutPrefixTree) const
+{
+	//while a valid class, and we're either the same as the base class (for multiple actors being selected and base class is AActor) OR we're not down to AActor yet)
+	TArray<FString> DetailPropertyExpansionStrings;
+	for (UStruct* Struct = BaseStruct; Struct && ((BaseStruct == Struct) || (Struct != AActor::StaticClass())); Struct = Struct->GetSuperStruct())
+	{
+		GConfig->GetSingleLineArray(TEXT("DetailPropertyExpansion"), *Struct->GetName(), DetailPropertyExpansionStrings, GEditorPerProjectIni);
+	}
+
+	OutPrefixTree.InsertAll(DetailPropertyExpansionStrings);
+}
+
 /** Called at the end of SetObjectArray after we change the objects being observed */
 void SDetailsView::PostSetObject(const TArray<FDetailsViewObjectRoot>& Roots)
 {
@@ -1044,13 +1056,35 @@ void SDetailsView::PostSetObject(const TArray<FDetailsViewObjectRoot>& Roots)
 	// Restore expanded items is called twice. The first time handles the existing property nodes,
 	// and second time after handling property customization as it might add new items that need to be expanded too.
 
+	// Collect the expansion states here, since the first call to RestoreAllExpandedItems() might overwrite the values and that may prevent some customization expansions to be saved.
+	TMap<UStruct*, FStringPrefixTree> ExpansionStates;
+	for (TSharedPtr<FComplexPropertyNode>& RootPropertyNode : GetRootNodes())
+	{
+		check(RootPropertyNode.IsValid());
+		UStruct* BaseStruct = RootPropertyNode->FindComplexParent()->GetBaseStructure();
+		FStringPrefixTree& PrefixTree = ExpansionStates.Add(BaseStruct);
+		GetExpansionState(BaseStruct, PrefixTree);
+	}
+
+	for (FDetailLayoutData& LayoutData : DetailLayouts)
+	{
+		FRootPropertyNodeList& ExternalRootPropertyNodes = LayoutData.DetailLayout->GetExternalRootPropertyNodes();
+		for (TSharedPtr<FComplexPropertyNode>& ExternalRootPropertyNode : ExternalRootPropertyNodes)
+		{
+			check(ExternalRootPropertyNode.IsValid());
+			UStruct* BaseStruct = ExternalRootPropertyNode->FindComplexParent()->GetBaseStructure();
+			FStringPrefixTree& PrefixTree = ExpansionStates.Add(BaseStruct);
+			GetExpansionState(BaseStruct, PrefixTree);
+		}
+	}
+
 	// Restore existing nodes
-	RestoreAllExpandedItems();
+	RestoreAllExpandedItems(&ExpansionStates);
 
 	UpdatePropertyMaps();
 
 	// Restore nodes generated by property customization
-	RestoreAllExpandedItems();
+	RestoreAllExpandedItems(&ExpansionStates);
 
 	UpdateFilteredDetails();
 }
diff --git a/Engine/Source/Editor/PropertyEditor/Private/SDetailsView.h b/Engine/Source/Editor/PropertyEditor/Private/SDetailsView.h
--- a/Engine/Source/Editor/PropertyEditor/Private/SDetailsView.h
+++ b/Engine/Source/Editor/PropertyEditor/Private/SDetailsView.h
@@ -238,6 +238,8 @@ private:
 	 */
 	void RefreshDisplayManager();
 
+	void GetExpansionState(UStruct* BaseStruct, FStringPrefixTree& OutPrefixTree) const;
+
 	/**
 	* The @code FDetailsViewStyle& @endcode shared pointer which provides style
 	* information for this details view.
diff --git a/Engine/Source/Editor/PropertyEditor/Private/SDetailsViewBase.cpp b/Engine/Source/Editor/PropertyEditor/Private/SDetailsViewBase.cpp
--- a/Engine/Source/Editor/PropertyEditor/Private/SDetailsViewBase.cpp
+++ b/Engine/Source/Editor/PropertyEditor/Private/SDetailsViewBase.cpp
@@ -1537,7 +1537,7 @@ void SDetailsViewBase::SaveExpandedItems(TSharedRef<FPropertyNode> StartNode)
 	}
 }
 
-void SDetailsViewBase::RestoreAllExpandedItems()
+void SDetailsViewBase::RestoreAllExpandedItems(TMap<UStruct*, FStringPrefixTree>* OptionalExpansionStates)
 {
 	TRACE_CPUPROFILER_EVENT_SCOPE(SDetailsViewBase::RestoreAllExpandedItems);
 
@@ -1552,7 +1552,7 @@ void SDetailsViewBase::RestoreAllExpandedItems()
 	for (TSharedPtr<FComplexPropertyNode>& RootPropertyNode : GetRootNodes())
 	{
 		check(RootPropertyNode.IsValid());
-		RestoreExpandedItems(RootPropertyNode.ToSharedRef());
+		RestoreExpandedItems(RootPropertyNode.ToSharedRef(), OptionalExpansionStates);
 	}
 
 	for (FDetailLayoutData& LayoutData : DetailLayouts)
@@ -1561,12 +1561,17 @@ void SDetailsViewBase::RestoreAllExpandedItems()
 		for (TSharedPtr<FComplexPropertyNode>& ExternalRootPropertyNode : ExternalRootPropertyNodes)
 		{
 			check(ExternalRootPropertyNode.IsValid());
-			RestoreExpandedItems(ExternalRootPropertyNode.ToSharedRef());
+			RestoreExpandedItems(ExternalRootPropertyNode.ToSharedRef(), OptionalExpansionStates);
 		}
 	}
 }
 
 void SDetailsViewBase::RestoreExpandedItems(TSharedRef<FPropertyNode> StartNode)
+{
+	RestoreExpandedItems(StartNode, nullptr);
+}
+
+void SDetailsViewBase::RestoreExpandedItems(TSharedRef<FPropertyNode> StartNode, TMap<UStruct*, FStringPrefixTree>* OptionalExpansionStates)
 {
 	if (bRunningDeferredActions)
 	{
@@ -1578,17 +1583,28 @@ void SDetailsViewBase::RestoreExpandedItems(TSharedRef<FPropertyNode> StartNode)
 
 	UStruct* BestBaseStruct = StartNode->FindComplexParent()->GetBaseStructure();
 
-	//while a valid class, and we're either the same as the base class (for multiple actors being selected and base class is AActor) OR we're not down to AActor yet)
-	TArray<FString> DetailPropertyExpansionStrings;
-	for (UStruct* Struct = BestBaseStruct; Struct && ((BestBaseStruct == Struct) || (Struct != AActor::StaticClass())); Struct = Struct->GetSuperStruct())
+	if (OptionalExpansionStates)
 	{
-		GConfig->GetSingleLineArray(TEXT("DetailPropertyExpansion"), *Struct->GetName(), DetailPropertyExpansionStrings, GEditorPerProjectIni);
+		const FStringPrefixTree* PrefixTree = OptionalExpansionStates->Find(BestBaseStruct);
+		if (PrefixTree)
+		{
+			SetExpandedItems(StartNode, *PrefixTree, false);
+		}
 	}
+	else
+	{
+		//while a valid class, and we're either the same as the base class (for multiple actors being selected and base class is AActor) OR we're not down to AActor yet)
+		TArray<FString> DetailPropertyExpansionStrings;
+		for (UStruct* Struct = BestBaseStruct; Struct && ((BestBaseStruct == Struct) || (Struct != AActor::StaticClass())); Struct = Struct->GetSuperStruct())
+		{
+			GConfig->GetSingleLineArray(TEXT("DetailPropertyExpansion"), *Struct->GetName(), DetailPropertyExpansionStrings, GEditorPerProjectIni);
+		}
 
-	FStringPrefixTree PrefixTree;
-	PrefixTree.InsertAll(DetailPropertyExpansionStrings);
+		FStringPrefixTree PrefixTree;
+		PrefixTree.InsertAll(DetailPropertyExpansionStrings);
 
-	SetExpandedItems(StartNode, PrefixTree, false);
+		SetExpandedItems(StartNode, PrefixTree, false);
+	}
 
 	if (BestBaseStruct)
 	{
diff --git a/Engine/Source/Editor/PropertyEditor/Private/SDetailsViewBase.h b/Engine/Source/Editor/PropertyEditor/Private/SDetailsViewBase.h
--- a/Engine/Source/Editor/PropertyEditor/Private/SDetailsViewBase.h
+++ b/Engine/Source/Editor/PropertyEditor/Private/SDetailsViewBase.h
@@ -148,6 +148,7 @@ public:
 	virtual const FCustomPropertyTypeLayoutMap& GetCustomPropertyTypeLayoutMap() const { return InstancedTypeToLayoutMap; }
 	virtual void SaveExpandedItems( TSharedRef<FPropertyNode> StartNode ) override;
 	virtual void RestoreExpandedItems(TSharedRef<FPropertyNode> StartNode) override;
+	virtual void RestoreExpandedItems(TSharedRef<FPropertyNode> StartNode, TMap<UStruct*, FStringPrefixTree>* OptionalExpansionStates);
 	virtual void MarkNodeAnimating(TSharedPtr<FPropertyNode> InNode, float InAnimationDuration, TOptional<FGuid> InAnimationBatchId) override;
 	virtual bool IsNodeAnimating(TSharedPtr<FPropertyNode> InNode) override;
 	virtual FDetailColumnSizeData& GetColumnSizeData() override { return ColumnSizeData; }
@@ -183,7 +184,7 @@ public:
 	virtual void EnqueueDeferredAction(FSimpleDelegate& DeferredAction) override;
 
 	/** Restore all expanded items in root nodes and external root nodes. */
-	void RestoreAllExpandedItems();
+	void RestoreAllExpandedItems(TMap<UStruct*, FStringPrefixTree>* OptionalExpansionStates = nullptr);
 
 	/**
 	 * Returns a @code TSharedPtr @endcode to the @code FDetailsDisplayManager @endcode for this
