Property Customization Helpers: Add a utility function to create an edit button

diff --git a/Engine/Source/Editor/PropertyEditor/Private/PropertyCustomizationHelpers.cpp b/Engine/Source/Editor/PropertyEditor/Private/PropertyCustomizationHelpers.cpp
--- a/Engine/Source/Editor/PropertyEditor/Private/PropertyCustomizationHelpers.cpp
+++ b/Engine/Source/Editor/PropertyEditor/Private/PropertyCustomizationHelpers.cpp
@@ -131,6 +131,17 @@ namespace PropertyCustomizationHelpers
 			.IsFocusable( false );
 	}
 
+	TSharedRef<SWidget> MakeEditButton( FSimpleDelegate OnEditClicked, TAttribute<FText> OptionalToolTipText, TAttribute<bool> IsEnabled )
+	{
+		return
+			SNew( SPropertyEditorButton )
+			.Text( OptionalToolTipText.Get().IsEmpty() ? LOCTEXT( "EditButtonToolTipText", "Edit Element") : OptionalToolTipText )
+			.Image( FAppStyle::GetBrush("Icons.Edit") )
+			.OnClickAction( OnEditClicked )
+			.IsEnabled( IsEnabled )
+			.IsFocusable( false );
+	}
+
 	TSharedRef<SWidget> MakeEmptyButton( FSimpleDelegate OnEmptyClicked, TAttribute<FText> OptionalToolTipText, TAttribute<bool> IsEnabled )
 	{
 		return
diff --git a/Engine/Source/Editor/PropertyEditor/Public/PropertyCustomizationHelpers.h b/Engine/Source/Editor/PropertyEditor/Public/PropertyCustomizationHelpers.h
--- a/Engine/Source/Editor/PropertyEditor/Public/PropertyCustomizationHelpers.h
+++ b/Engine/Source/Editor/PropertyEditor/Public/PropertyCustomizationHelpers.h
@@ -94,6 +94,7 @@ namespace PropertyCustomizationHelpers
 	PROPERTYEDITOR_API TSharedRef<SWidget> MakeResetButton(FSimpleDelegate OnResetClicked, TAttribute<FText> OptionalToolTipText = FText(), TAttribute<bool> IsEnabled = true);
 	PROPERTYEDITOR_API TSharedRef<SWidget> MakeAddButton( FSimpleDelegate OnAddClicked, TAttribute<FText> OptionalToolTipText = FText(), TAttribute<bool> IsEnabled = true );
 	PROPERTYEDITOR_API TSharedRef<SWidget> MakeRemoveButton( FSimpleDelegate OnRemoveClicked, TAttribute<FText> OptionalToolTipText = FText(), TAttribute<bool> IsEnabled = true );
+	PROPERTYEDITOR_API TSharedRef<SWidget> MakeEditButton( FSimpleDelegate OnEditClicked, TAttribute<FText> OptionalToolTipText = FText(), TAttribute<bool> IsEnabled = true );
 	PROPERTYEDITOR_API TSharedRef<SWidget> MakeEmptyButton( FSimpleDelegate OnEmptyClicked, TAttribute<FText> OptionalToolTipText = FText(), TAttribute<bool> IsEnabled = true );
 	PROPERTYEDITOR_API TSharedRef<SWidget> MakeInsertDeleteDuplicateButton( FExecuteAction OnInsertClicked, FExecuteAction OnDeleteClicked, FExecuteAction OnDuplicateClicked );
 	PROPERTYEDITOR_API TSharedRef<SWidget> MakeDeleteButton( FSimpleDelegate OnDeleteClicked, TAttribute<FText> OptionalToolTipText = FText(), TAttribute<bool> IsEnabled = true );
