Skip expensive Asset Registry tag gathering for temporary FAssetData in SPropertyEditorAsset
* Widget does not need asset tags.  For certain assets, tag gathering can be quite expensive during the ::Tick() to update the thumbnail.  This issue was discovered specifically with StaticMeshes that have a large number of collision primitives.

diff --git a/Engine/Source/Editor/PropertyEditor/Private/UserInterface/PropertyEditor/SPropertyEditorAsset.cpp b/Engine/Source/Editor/PropertyEditor/Private/UserInterface/PropertyEditor/SPropertyEditorAsset.cpp
--- a/Engine/Source/Editor/PropertyEditor/Private/UserInterface/PropertyEditor/SPropertyEditorAsset.cpp
+++ b/Engine/Source/Editor/PropertyEditor/Private/UserInterface/PropertyEditor/SPropertyEditorAsset.cpp
@@ -1242,7 +1242,7 @@ SPropertyEditorAsset::FObjectOrAssetData::FObjectOrAssetData(UObject* InObject,
 	}
 	else if(InObject != nullptr)
 	{
-		AssetData = FAssetData(InObject);
+		AssetData = FAssetData(InObject, FAssetData::ECreationFlags::SkipAssetRegistryTagsGathering);
 		ObjectPath = InObject;
 	}
 }
