Skip expensive Asset Registry tag gathering for temporary FAssetData in SPropertyEditorAsset
* Widget only needs asset tags in cases where the IsAssetFiltered is called.  For certain assets, tag gathering can be quite expensive during the ::Tick() to update the thumbnail.  This issue was discovered specifically with StaticMeshes that have a large number of collision primitives.

diff --git a/Engine/Source/Editor/PropertyEditor/Private/UserInterface/PropertyEditor/SPropertyEditorAsset.cpp b/Engine/Source/Editor/PropertyEditor/Private/UserInterface/PropertyEditor/SPropertyEditorAsset.cpp
--- a/Engine/Source/Editor/PropertyEditor/Private/UserInterface/PropertyEditor/SPropertyEditorAsset.cpp
+++ b/Engine/Source/Editor/PropertyEditor/Private/UserInterface/PropertyEditor/SPropertyEditorAsset.cpp
@@ -543,7 +543,7 @@ void SPropertyEditorAsset::Construct(const FArguments& InArgs, const TSharedPtr<
 	if (ShouldDisplayThumbnail(InArgs, ObjectClass))
 	{
 		FObjectOrAssetData Value; 
-		GetValue( Value );
+		GetValue( Value, FObjectOrAssetData::EAssetDataOptions::None );
 
 		AssetThumbnail = MakeShareable( new FAssetThumbnail( Value.AssetData, InArgs._ThumbnailSize.X, InArgs._ThumbnailSize.Y, InArgs._ThumbnailPool ) );
 
@@ -765,7 +765,7 @@ SPropertyEditorAsset::EActorReferenceState SPropertyEditorAsset::GetActorReferen
 	if (bIsActor)
 	{
 		FObjectOrAssetData Value;
-		GetValue(Value);
+		GetValue(Value, FObjectOrAssetData::EAssetDataOptions::SkipAssetRegistryTagsGathering );
 
 		if (Value.Object != nullptr)
 		{
@@ -815,7 +815,7 @@ void SPropertyEditorAsset::Tick( const FGeometry& AllottedGeometry, const double
 	{
 		// Ensure the thumbnail is up to date
 		FObjectOrAssetData Value;
-		GetValue( Value );
+		GetValue( Value, FObjectOrAssetData::EAssetDataOptions::SkipAssetRegistryTagsGathering );
 
 		// If the thumbnail is not the same as the object value set the thumbnail to the new value
 		if( !(AssetThumbnail->GetAssetData() == Value.AssetData) )
@@ -854,7 +854,7 @@ bool SPropertyEditorAsset::Supports(const FProperty* NodeProperty)
 TSharedRef<SWidget> SPropertyEditorAsset::OnGetMenuContent()
 {
 	FObjectOrAssetData Value;
-	GetValue(Value);
+	GetValue(Value, FObjectOrAssetData::EAssetDataOptions::None);
 
 	if (bIsActor)
 	{
@@ -923,7 +923,7 @@ void SPropertyEditorAsset::CloseComboButton()
 FText SPropertyEditorAsset::OnGetAssetName() const
 {
 	FObjectOrAssetData Value; 
-	FPropertyAccess::Result Result = GetValue( Value );
+	FPropertyAccess::Result Result = GetValue( Value, FObjectOrAssetData::EAssetDataOptions::SkipAssetRegistryTagsGathering );
 
 	FText Name = LOCTEXT("None", "None");
 	if( Result == FPropertyAccess::Success )
@@ -982,7 +982,7 @@ FText SPropertyEditorAsset::OnGetAssetClassName() const
 FText SPropertyEditorAsset::OnGetToolTip() const
 {
 	FObjectOrAssetData Value; 
-	FPropertyAccess::Result Result = GetValue( Value );
+	FPropertyAccess::Result Result = GetValue( Value, FObjectOrAssetData::EAssetDataOptions::SkipAssetRegistryTagsGathering );
 
 	FText ToolTipText = FText::GetEmpty();
 
@@ -1063,7 +1063,7 @@ void SPropertyEditorAsset::SetValue( const FAssetData& AssetData )
 	}
 }
 
-FPropertyAccess::Result SPropertyEditorAsset::GetValue( FObjectOrAssetData& OutValue ) const
+FPropertyAccess::Result SPropertyEditorAsset::GetValue( FObjectOrAssetData& OutValue, FObjectOrAssetData::EAssetDataOptions AssetDataOptions ) const
 {
 	// Potentially accessing the value while garbage collecting or saving the package could trigger a crash.
 	// so we fail to get the value when that is occurring.
@@ -1126,7 +1126,7 @@ FPropertyAccess::Result SPropertyEditorAsset::GetValue( FObjectOrAssetData& OutV
 		}
 #endif
 
-		OutValue = FObjectOrAssetData( Object, EditorPathOwner );
+		OutValue = FObjectOrAssetData( Object, EditorPathOwner, AssetDataOptions );
 	}
 	else
 	{
@@ -1156,8 +1156,8 @@ FPropertyAccess::Result SPropertyEditorAsset::GetValue( FObjectOrAssetData& OutV
 				UE_LOG(LogPropertyNode, Fatal, TEXT("Property \"%s\" (%s) contains invalid data."), *Property->GetName(), *Property->GetCPPType());
 			}
 #endif
-
-			OutValue = FObjectOrAssetData(Object);
+			UObject* const InEditorPathOwner = nullptr;
+			OutValue = FObjectOrAssetData( Object, InEditorPathOwner, AssetDataOptions);
 		}
 		else
 		{
@@ -1217,7 +1217,7 @@ FPropertyAccess::Result SPropertyEditorAsset::GetValue( FObjectOrAssetData& OutV
 const UClass* SPropertyEditorAsset::GetDisplayedClass() const
 {
 	FObjectOrAssetData Value;
-	GetValue( Value );
+	GetValue( Value, FObjectOrAssetData::EAssetDataOptions::SkipAssetRegistryTagsGathering );
 	if(Value.Object != nullptr)
 	{
 		return Value.Object->GetClass();
@@ -1233,7 +1233,7 @@ void SPropertyEditorAsset::OnAssetSelected( const struct FAssetData& AssetData )
 	SetValue(AssetData);
 }
 
-SPropertyEditorAsset::FObjectOrAssetData::FObjectOrAssetData(UObject* InObject, UObject* InEditorPathOwner)
+SPropertyEditorAsset::FObjectOrAssetData::FObjectOrAssetData(UObject* InObject, UObject* InEditorPathOwner, EAssetDataOptions AssetDataOptions)
 	: Object(InObject)
 {
 	if (AActor* Actor = Cast<AActor>(InObject))
@@ -1242,7 +1242,12 @@ SPropertyEditorAsset::FObjectOrAssetData::FObjectOrAssetData(UObject* InObject,
 	}
 	else if(InObject != nullptr)
 	{
-		AssetData = FAssetData(InObject);
+		FAssetData::ECreationFlags CreationFlags =
+			(AssetDataOptions == EAssetDataOptions::SkipAssetRegistryTagsGathering)
+			? FAssetData::ECreationFlags::SkipAssetRegistryTagsGathering
+			: FAssetData::ECreationFlags::None;
+		
+		AssetData = FAssetData(InObject, CreationFlags);
 		ObjectPath = InObject;
 	}
 }
@@ -1271,7 +1276,7 @@ void SPropertyEditorAsset::OnGetAllowedClasses(TArray<const UClass*>& AllowedCla
 void SPropertyEditorAsset::OnOpenAssetEditor()
 {
 	FObjectOrAssetData Value;
-	GetValue( Value );
+	GetValue( Value, FObjectOrAssetData::EAssetDataOptions::SkipAssetRegistryTagsGathering );
 
 	UObject* ObjectToEdit = Value.AssetData.GetAsset();
 	if( ObjectToEdit )
@@ -1307,7 +1312,7 @@ void SPropertyEditorAsset::OnOpenAssetEditor()
 void SPropertyEditorAsset::OnBrowse()
 {
 	FObjectOrAssetData Value;
-	GetValue( Value );
+	GetValue( Value, FObjectOrAssetData::EAssetDataOptions::None );
 
 	if (bIsActor)
 	{
@@ -1326,7 +1331,7 @@ void SPropertyEditorAsset::OnBrowse()
 						if (const FWorldPartitionActorDescInstance* ActorDescInstance = World->GetWorldPartition()->GetActorDescInstanceByPath(Value.ObjectPath))
 						{
 							World->GetWorldPartition()->PinActors({ ActorDescInstance->GetGuid() });
-							GetValue(Value);
+							GetValue(Value, FObjectOrAssetData::EAssetDataOptions::None);
 						}
 					}
 				}
@@ -1353,7 +1358,7 @@ void SPropertyEditorAsset::OnBrowse()
 FText SPropertyEditorAsset::GetOnBrowseToolTip() const
 {
 	FObjectOrAssetData Value;
-	GetValue( Value );
+	GetValue( Value, FObjectOrAssetData::EAssetDataOptions::SkipAssetRegistryTagsGathering );
 
 	if (Value.Object)
 	{
@@ -1478,7 +1483,7 @@ void SPropertyEditorAsset::OnAssetDropped( const FDragDropEvent&, TArrayView<FAs
 void SPropertyEditorAsset::OnCopy()
 {
 	FObjectOrAssetData Value;
-	GetValue( Value );
+	GetValue( Value, FObjectOrAssetData::EAssetDataOptions::SkipAssetRegistryTagsGathering );
 
 	if( Value.AssetData.IsValid() )
 	{
diff --git a/Engine/Source/Editor/PropertyEditor/Private/UserInterface/PropertyEditor/SPropertyEditorAsset.h b/Engine/Source/Editor/PropertyEditor/Private/UserInterface/PropertyEditor/SPropertyEditorAsset.h
--- a/Engine/Source/Editor/PropertyEditor/Private/UserInterface/PropertyEditor/SPropertyEditorAsset.h
+++ b/Engine/Source/Editor/PropertyEditor/Private/UserInterface/PropertyEditor/SPropertyEditorAsset.h
@@ -110,12 +110,21 @@ public:
 private:
 	struct FObjectOrAssetData
 	{
-		UObject* Object;
+		enum class EAssetDataOptions
+		{
+			None,
+			// Will not populate AssetRegistry Tags for an AssetData
+			SkipAssetRegistryTagsGathering
+		};
+		
+		UObject* Object = nullptr;
 		FSoftObjectPath ObjectPath;
 		FAssetData AssetData;
-		UObject* EditorPathOwner;
+		UObject* EditorPathOwner = nullptr;
+		
+		FObjectOrAssetData() = default;
 
-		FObjectOrAssetData(UObject* InObject = nullptr, UObject* EditorPathOwner = nullptr);
+		FObjectOrAssetData(UObject* InObject, UObject* EditorPathOwner, EAssetDataOptions AssetDataOptions);
 
 		FObjectOrAssetData( const FSoftObjectPath& InObjectPath )
 			: Object(nullptr)
@@ -213,7 +222,7 @@ private:
 	 * Get the value referenced by this widget.
 	 * @returns the referenced object
 	 */
-	FPropertyAccess::Result GetValue( FObjectOrAssetData& OutValue ) const;
+	FPropertyAccess::Result GetValue( FObjectOrAssetData& OutValue, FObjectOrAssetData::EAssetDataOptions GetValueOptions ) const;
 
 	/** 
 	 * Get the UClass we will display in the UI.
