[PropertyEditor] Added an optional tooltip/description to Categories and Groups

diff --git a/Engine/Source/Editor/PropertyEditor/Private/DetailCategoryBuilderImpl.cpp b/Engine/Source/Editor/PropertyEditor/Private/DetailCategoryBuilderImpl.cpp
--- a/Engine/Source/Editor/PropertyEditor/Private/DetailCategoryBuilderImpl.cpp
+++ b/Engine/Source/Editor/PropertyEditor/Private/DetailCategoryBuilderImpl.cpp
@@ -836,6 +836,16 @@ void FDetailCategoryImpl::SetDisplayName(const FText& InDisplayName)
 	SetDisplayName(CategoryName, InDisplayName);
 }
 
+const TOptional<FText>& FDetailCategoryImpl::GetToolTip() const
+{
+	return ToolTip;
+}
+
+void FDetailCategoryImpl::SetToolTip(const FText& InToolTip)
+{
+	ToolTip = InToolTip;
+}
+
 void FDetailCategoryImpl::SetDisplayName(FName InCategoryName, const FText& LocalizedNameOverride)
 {
 	if (!LocalizedNameOverride.IsEmpty())
@@ -885,7 +895,8 @@ TSharedRef<ITableRow> FDetailCategoryImpl::GenerateWidgetForTableView(const TSha
 	InitializeObjectName();
 	TSharedPtr<FDetailLayoutBuilderImpl> ParentLayout = GetParentLayoutImpl();
 
-	return SNew(SDetailCategoryTableRow, AsShared(), OwnerTable)
+	TSharedRef<SDetailCategoryTableRow> RowWidget =
+		SNew(SDetailCategoryTableRow, AsShared(), OwnerTable)
 		.PasteFromText(OnPasteFromText())
 		.ObjectName( ObjectName )
 		.IsEmpty( bIsEmpty )
@@ -893,6 +904,13 @@ TSharedRef<ITableRow> FDetailCategoryImpl::GenerateWidgetForTableView(const TSha
 		.DisplayName(GetDisplayName())
 		.HeaderContent(HeaderContent)
 		.WholeRowHeaderContent(bHeaderContentWholeRowContent);
+
+	if (ToolTip.IsSet() && !ToolTip->IsEmpty())
+	{
+		RowWidget->SetToolTipText(ToolTip.GetValue());
+	}
+
+	return RowWidget;
 }
 
 void FDetailCategoryImpl::InitializeObjectName()
@@ -932,12 +950,21 @@ bool FDetailCategoryImpl::GenerateStandaloneWidget(FDetailWidgetRow& OutRow) con
 
 	const bool bIsInnerCategory = GetParentLayoutImpl()->IsLayoutForExternalRoot();
 	FTextBlockStyle NameStyle = bIsInnerCategory ? FAppStyle::Get().GetWidgetStyle<FTextBlockStyle>("NormalText") : FAppStyle::Get().GetWidgetStyle<FTextBlockStyle>("DetailsView.CategoryTextStyle");
+
+	TSharedRef<STextBlock> TextBlock =
+    	SNew(STextBlock)
+    	.Text(GetDisplayName())
+    	.TextStyle(&NameStyle)
+    	.ShadowOffset(FVector2D::ZeroVector);
+
+    if (ToolTip.IsSet() && !ToolTip->IsEmpty())
+    {
+    	TextBlock->SetToolTipText(ToolTip.GetValue());
+    }
+
 	OutRow.NameContent()
 	[
-		SNew(STextBlock)
-		.Text(GetDisplayName())
-		.TextStyle(&NameStyle)
-		.ShadowOffset(FVector2D::ZeroVector)
+		TextBlock
 	];
 
 	if(HeaderContentWidget.IsValid())
diff --git a/Engine/Source/Editor/PropertyEditor/Private/DetailCategoryBuilderImpl.h b/Engine/Source/Editor/PropertyEditor/Private/DetailCategoryBuilderImpl.h
--- a/Engine/Source/Editor/PropertyEditor/Private/DetailCategoryBuilderImpl.h
+++ b/Engine/Source/Editor/PropertyEditor/Private/DetailCategoryBuilderImpl.h
@@ -187,6 +187,8 @@ public:
 	virtual void GetDefaultProperties(TArray<TSharedRef<IPropertyHandle> >& OutAllProperties, bool bSimpleProperties = true, bool bAdvancedProperties = true) override;
 	virtual const FText& GetDisplayName() const override { return DisplayName; }
 	virtual void SetDisplayName(const FText& InDisplayName) override;
+	virtual const TOptional<FText>& GetToolTip() const override;
+	virtual void SetToolTip(const FText& ToolTip) override;
 	virtual void SetCategoryVisibility(bool bIsVisible) override;
 	virtual void SetShowAdvanced(bool bShowAdvanced) override;
 	virtual int32 GetSortOrder() const override;
@@ -436,6 +438,8 @@ private:
 	FOnBooleanValueChanged OnExpansionChangedDelegate;
 	/** The display name of the category */
 	FText DisplayName;
+	/** ToolTip for this category */
+    TOptional<FText> ToolTip;
 	/** The path name of the category */
 	FString CategoryPathName;
 	/** Custom header content displayed to the right of the category name */
diff --git a/Engine/Source/Editor/PropertyEditor/Private/DetailGroup.cpp b/Engine/Source/Editor/PropertyEditor/Private/DetailGroup.cpp
--- a/Engine/Source/Editor/PropertyEditor/Private/DetailGroup.cpp
+++ b/Engine/Source/Editor/PropertyEditor/Private/DetailGroup.cpp
@@ -86,6 +86,16 @@ IDetailGroup& FDetailGroup::AddGroup(FName NewGroupName, const FText& InLocalize
 	return *NewCustomization.DetailGroup;
 }
 
+const TOptional<FText>& FDetailGroup::GetToolTip() const
+{
+	return LocalizedToolTip;
+}
+
+void FDetailGroup::SetToolTip(const FText& ToolTip)
+{
+	LocalizedToolTip = ToolTip;
+}
+
 void FDetailGroup::ToggleExpansion( bool bExpand )
 {
 	if( ParentCategory.IsValid() && OwnerTreeNode.IsValid() )
@@ -227,6 +237,15 @@ FReply FDetailGroup::OnNameClicked()
 
 TSharedRef<SWidget> FDetailGroup::MakeNameWidget()
 {
+	TSharedRef<STextBlock> TextBlock = SNew(STextBlock)
+		.Font(IDetailLayoutBuilder::GetDetailFont())
+		.Text(LocalizedDisplayName);
+
+	if (LocalizedToolTip.IsSet() && !LocalizedToolTip->IsEmpty())
+	{
+		TextBlock->SetToolTipText(LocalizedToolTip.GetValue());
+	}
+
 	return
 		SNew( SButton )
 		.ButtonStyle( FAppStyle::Get(), "NoBorder" )
@@ -235,9 +254,7 @@ TSharedRef<SWidget> FDetailGroup::MakeNameWidget()
 		.ForegroundColor( FSlateColor::UseForeground() )
 		.Content()
 		[
-			SNew( STextBlock )
-			.Font( IDetailLayoutBuilder::GetDetailFont() )
-			.Text( LocalizedDisplayName )
+			TextBlock
 		];
 }
 
diff --git a/Engine/Source/Editor/PropertyEditor/Private/DetailGroup.h b/Engine/Source/Editor/PropertyEditor/Private/DetailGroup.h
--- a/Engine/Source/Editor/PropertyEditor/Private/DetailGroup.h
+++ b/Engine/Source/Editor/PropertyEditor/Private/DetailGroup.h
@@ -26,7 +26,8 @@ public:
 	virtual FDetailWidgetRow& AddWidgetRow() override;
 	virtual IDetailPropertyRow& AddPropertyRow( TSharedRef<IPropertyHandle> PropertyHandle ) override;
 	virtual IDetailGroup& AddGroup(FName NewGroupName, const FText& InLocalizedDisplayName, bool bInStartExpanded = false) override;
-
+	virtual const TOptional<FText>& GetToolTip() const override;
+	virtual void SetToolTip(const FText& ToolTip) override;
 	virtual void ToggleExpansion( bool bExpand ) override;
 	virtual bool GetExpansionState() const override;
 	virtual void SetDisplayMode(EDetailGroupDisplayMode Mode) override;
@@ -131,6 +132,8 @@ private:
 	TAttribute<bool> IsParentEnabled;
 	/** Display name of this group */
 	FText LocalizedDisplayName;
+	/** ToolTip for this group */
+	TOptional<FText> LocalizedToolTip;
 	/** Name identifier of this group */
 	FName GroupName;
 	/** Whether the detail group should start expanded or not */
diff --git a/Engine/Source/Editor/PropertyEditor/Public/DetailCategoryBuilder.h b/Engine/Source/Editor/PropertyEditor/Public/DetailCategoryBuilder.h
--- a/Engine/Source/Editor/PropertyEditor/Public/DetailCategoryBuilder.h
+++ b/Engine/Source/Editor/PropertyEditor/Public/DetailCategoryBuilder.h
@@ -201,6 +201,16 @@ public:
 	 */
 	virtual void SetDisplayName(const FText& DisplayName) = 0;
 
+	/**
+	 * Gets the localized tooltip for this category.
+	 */
+	virtual const TOptional<FText>& GetToolTip() const = 0;
+
+	/**
+	 * Sets the localized tooltip for this category.
+	 */
+	virtual void SetToolTip(const FText& ToolTip) = 0;
+
 	/**
 	 * Sets whether or not this category is hidden or shown
 	 * This is designed to be used for dynamic category visibility after construction of the category
diff --git a/Engine/Source/Editor/PropertyEditor/Public/IDetailGroup.h b/Engine/Source/Editor/PropertyEditor/Public/IDetailGroup.h
--- a/Engine/Source/Editor/PropertyEditor/Public/IDetailGroup.h
+++ b/Engine/Source/Editor/PropertyEditor/Public/IDetailGroup.h
@@ -62,6 +62,16 @@ public:
 	 */
 	virtual IDetailGroup& AddGroup(FName NewGroupName, const FText& InLocalizedDisplayName, bool bInStartExpanded = false) = 0;
 
+	/**
+	 * Gets the localized tooltip for this group.
+	 */
+	virtual const TOptional<FText>& GetToolTip() const = 0;
+
+	/**
+	 * Sets the localized tooltip for this group.
+	 */
+	virtual void SetToolTip(const FText& ToolTip) = 0;
+
 	/**
 	 * Toggles expansion on the group
 	 *
