In FPropertyNode::SupportsEditConditionToggle don't show the inline edit condition toggle in the case where the edit condition property is not marked as editable and does not have InlineEditConditionToggle set.

This was a surprising case that was undocumented outside of the source code.  The inline edit condition toggle is now only shown when InlineEditConditionToggle is set ( and HideEditConditionToggle is not set ).

The console variable "PropertyEditor.ShowInlineEditConditionToggleWhenNotSpecifiedAndNotEditable" ( off by default ) can be used to restore the legacy behavior if necessary.

diff --git a/Engine/Source/Editor/PropertyEditor/Private/PropertyNode.cpp b/Engine/Source/Editor/PropertyEditor/Private/PropertyNode.cpp
--- a/Engine/Source/Editor/PropertyEditor/Private/PropertyNode.cpp
+++ b/Engine/Source/Editor/PropertyEditor/Private/PropertyNode.cpp
@@ -32,6 +32,16 @@
 
 #define LOCTEXT_NAMESPACE "PropertyNode"
 
+namespace UE::PropertyEditor::Private
+{
+	static bool bShowInlineEditConditionToggleWhenNotSpecifiedAndNotEditable = false;
+	static FAutoConsoleVariableRef CVarShowInlineEditConditionToggleWhenNotSpecifiedAndNotEditable(
+		TEXT("PropertyEditor.ShowInlineEditConditionToggleWhenNotSpecifiedAndNotEditable"),
+		bShowInlineEditConditionToggleWhenNotSpecifiedAndNotEditable,
+		TEXT("Enables legacy behavior to show the InlineEditConditionToggle when the edit condition property does not have this specifier and is not editable.")
+	);
+}
+
 namespace PropertyEditorPolicy
 {
 	struct FPropertyNodePolicyImpl : public FObjectArchetypeHelper::IObjectArchetypePolicy
@@ -1475,31 +1485,31 @@ bool FPropertyNode::SupportsEditConditionToggle() const
 		const FBoolProperty* ConditionalProperty = EditConditionContext->GetSingleBoolProperty(EditConditionExpression);
 		if (ConditionalProperty != nullptr)
 		{
-			// There are 2 valid states for inline edit conditions:
-			// 1. The property is marked as editable and has InlineEditConditionToggle set. 
-			// 2. The property is not marked as editable and does not have InlineEditConditionToggle set.
-			// In both cases, the original property will be hidden and only show up as a toggle.
-
 			static const FName Name_InlineEditConditionToggle("InlineEditConditionToggle");
 			const bool bIsInlineEditCondition = ConditionalProperty->HasMetaData(Name_InlineEditConditionToggle);
 			const bool bIsEditable = ConditionalProperty->HasAllPropertyFlags(CPF_Edit);
-
-			if (bIsInlineEditCondition == bIsEditable)
+			
+			// Support for legacy behavior ( case 2. in the comment below ) if enabled
+			if (UE::PropertyEditor::Private::bShowInlineEditConditionToggleWhenNotSpecifiedAndNotEditable)
 			{
-				return true;
+				// There are 2 valid states for inline edit conditions:
+				// 1. The property is marked as editable and has InlineEditConditionToggle set. 
+				// 2. The property is not marked as editable and does not have InlineEditConditionToggle set.
+				// In both cases, the original property will be hidden and only show up as a toggle.
+				if (bIsInlineEditCondition == bIsEditable)
+				{
+					return true;
+				}
 			}
-
-			if (bIsInlineEditCondition && !bIsEditable)
+			
+			if (bIsInlineEditCondition)
 			{
-				UE_LOG(LogPropertyNode, Warning, TEXT("Property being used as inline edit condition is not editable, but has redundant InlineEditConditionToggle flag. Field \"%s\" in class \"%s\"."), *ConditionalProperty->GetNameCPP(), *Property->GetOwnerStruct()->GetName());
+				if (!bIsEditable)
+				{
+					UE_LOG(LogPropertyNode, Warning, TEXT("Property being used as inline edit condition is not editable, but has InlineEditConditionToggle flag. Field \"%s\" in class \"%s\"."), *ConditionalProperty->GetNameCPP(), *Property->GetOwnerStruct()->GetName());
+				}
 				return true;
 			}
-
-			// The property is already shown, and not marked as inline edit condition.
-			if (!bIsInlineEditCondition && bIsEditable)
-			{
-				return false;
-			}
 		}
 	}
 
