[CrashFix] Made details splitter slot support null details panel (and display nothing)

diff --git a/Engine/Source/Editor/PropertyEditor/Private/SDetailsSplitter.cpp b/Engine/Source/Editor/PropertyEditor/Private/SDetailsSplitter.cpp
--- a/Engine/Source/Editor/PropertyEditor/Private/SDetailsSplitter.cpp
+++ b/Engine/Source/Editor/PropertyEditor/Private/SDetailsSplitter.cpp
@@ -432,7 +432,7 @@ void SDetailsSplitter::AddSlot(const FSlot::FSlotArguments& SlotArgs, int32 Inde
 	[
 		SNew(SBox).Padding(15.f,0.f, 15.f,0.f)
 		[
-			SlotArgs._DetailsView.ToSharedRef()
+			SlotArgs._DetailsView ? SlotArgs._DetailsView.ToSharedRef() : SNullWidget::NullWidget
 		]
 	];
 	Panels.Insert({
