Property Row: Allow property extensions on structs

diff --git a/Engine/Source/Editor/PropertyEditor/Private/SDetailSingleItemRow.cpp b/Engine/Source/Editor/PropertyEditor/Private/SDetailSingleItemRow.cpp
--- a/Engine/Source/Editor/PropertyEditor/Private/SDetailSingleItemRow.cpp
+++ b/Engine/Source/Editor/PropertyEditor/Private/SDetailSingleItemRow.cpp
@@ -1653,7 +1653,7 @@ void SDetailSingleItemRow::PopulateExtensionWidget()
 		{
 			TSharedPtr<IPropertyHandle> Handle = PropertyEditorHelpers::GetPropertyHandle(Customization->GetPropertyNode().ToSharedRef(), nullptr, nullptr);
 			const UClass* ObjectClass = Handle->GetOuterBaseClass();
-			if (Handle->IsValidHandle() && ObjectClass && ExtensionHandler->IsPropertyExtendable(ObjectClass, *Handle))
+			if (Handle->IsValidHandle() && ExtensionHandler->IsPropertyExtendable(ObjectClass, *Handle))
 			{
 				IDetailLayoutBuilder& DetailLayout = OwnerTreeNodePinned->GetParentCategory()->GetParentLayout();
 				ExtensionHandler->ExtendWidgetRow(WidgetRow, DetailLayout, ObjectClass, Handle);
