Implement InvalidateCachedState for SStructureDetailsView

diff --git a/Engine/Source/Editor/PropertyEditor/Private/SStructureDetailsView.cpp b/Engine/Source/Editor/PropertyEditor/Private/SStructureDetailsView.cpp
--- a/Engine/Source/Editor/PropertyEditor/Private/SStructureDetailsView.cpp
+++ b/Engine/Source/Editor/PropertyEditor/Private/SStructureDetailsView.cpp
@@ -276,6 +276,14 @@ void SStructureDetailsView::ForceRefresh()
 	SetStructureProvider(StructProvider);
 }
 
+void SStructureDetailsView::InvalidateCachedState()
+{
+	for (const TSharedPtr<FComplexPropertyNode>& RootNode : RootNodes)
+	{
+		RootNode->InvalidateCachedState();
+	}
+}
+
 void SStructureDetailsView::ClearSearch()
 {
 	CurrentFilter.FilterStrings.Empty();
@@ -283,7 +291,6 @@ void SStructureDetailsView::ClearSearch()
 	RerunCurrentFilter();
 }
 
-
 const TArray< TWeakObjectPtr<UObject> >& SStructureDetailsView::GetSelectedObjects() const
 {
 	static const TArray< TWeakObjectPtr<UObject> > DummyRef;
diff --git a/Engine/Source/Editor/PropertyEditor/Private/SStructureDetailsView.h b/Engine/Source/Editor/PropertyEditor/Private/SStructureDetailsView.h
--- a/Engine/Source/Editor/PropertyEditor/Private/SStructureDetailsView.h
+++ b/Engine/Source/Editor/PropertyEditor/Private/SStructureDetailsView.h
@@ -74,7 +74,10 @@ public:
 	}
 
 public:
+	/** Causes the details view to be refreshed (new widgets generated) with the current set of objects */
 	virtual void ForceRefresh() override;
+	/** Invalidates cached state such as the "revert to default" arrow and edit conditions, without rebuilding the entire panel. */
+	virtual void InvalidateCachedState() override;
 	virtual void MoveScrollOffset(int32 DeltaOffset) override {}
 	virtual void ClearSearch() override;
 public:
