Fixed crash when there are additional children nodes attached to Vector property

#rnx

diff --git a/Engine/Source/Editor/PropertyEditor/Private/PropertyHandleImpl.cpp b/Engine/Source/Editor/PropertyEditor/Private/PropertyHandleImpl.cpp
--- a/Engine/Source/Editor/PropertyEditor/Private/PropertyHandleImpl.cpp
+++ b/Engine/Source/Editor/PropertyEditor/Private/PropertyHandleImpl.cpp
@@ -4604,29 +4604,25 @@ bool FPropertyHandleVector::Supports( TSharedRef<FPropertyNode> PropertyNode )
 FPropertyHandleVector::FPropertyHandleVector( TSharedRef<class FPropertyNode> PropertyNode, class FNotifyHook* NotifyHook, TSharedPtr<IPropertyUtilities> PropertyUtilities )
 	: FPropertyHandleStruct( PropertyNode, NotifyHook, PropertyUtilities ) 
 {
-	if( Implementation->GetNumChildren() > 0 )
+	const bool bRecurse = false;
+	if( TSharedPtr<FPropertyNode> XComponentPropertyNode = Implementation->GetChildNode("X", bRecurse) )
 	{
-		const bool bRecurse = false;
-		// A vector is a struct property that has multiple children.  We get/set the values from the children
-		VectorComponents.Add( MakeShareable( new FPropertyHandleMixed( Implementation->GetChildNode("X", bRecurse).ToSharedRef(), NotifyHook, PropertyUtilities ) ) );
+		VectorComponents.Add( MakeShareable( new FPropertyHandleMixed( XComponentPropertyNode.ToSharedRef(), NotifyHook, PropertyUtilities ) ) );
+	}
 
-		if( Implementation->GetNumChildren() > 1 )
-		{
-			// at least a 2 component vector
-			VectorComponents.Add( MakeShareable( new FPropertyHandleMixed( Implementation->GetChildNode("Y", bRecurse).ToSharedRef(), NotifyHook, PropertyUtilities ) ) );
-		}
+	if( TSharedPtr<FPropertyNode> YComponentPropertyNode = Implementation->GetChildNode("Y", bRecurse) )
+	{
+		VectorComponents.Add( MakeShareable( new FPropertyHandleMixed( YComponentPropertyNode.ToSharedRef(), NotifyHook, PropertyUtilities ) ) );
+	}
 		
-		if( Implementation->GetNumChildren() > 2 )
-		{
-			// at least a 3 component vector
-			VectorComponents.Add( MakeShareable( new FPropertyHandleMixed( Implementation->GetChildNode("Z",bRecurse).ToSharedRef(), NotifyHook, PropertyUtilities ) ) );
-		}
+	if( TSharedPtr<FPropertyNode> ZComponentPropertyNode = Implementation->GetChildNode("Z", bRecurse) )
+	{
+		VectorComponents.Add( MakeShareable( new FPropertyHandleMixed( ZComponentPropertyNode.ToSharedRef(), NotifyHook, PropertyUtilities ) ) );
+	}
 		
-		if( Implementation->GetNumChildren() > 3 )
-		{
-			// a 4 component vector
-			VectorComponents.Add( MakeShareable( new FPropertyHandleMixed( Implementation->GetChildNode("W",bRecurse).ToSharedRef(), NotifyHook, PropertyUtilities ) ) );
-		}
+	if( TSharedPtr<FPropertyNode> WComponentPropertyNode = Implementation->GetChildNode("W", bRecurse) )
+	{
+		VectorComponents.Add( MakeShareable( new FPropertyHandleMixed( WComponentPropertyNode.ToSharedRef(), NotifyHook, PropertyUtilities ) ) );
 	}
 }
 
