Fix loctext namespace conflict. This was introduced by 34368520. The conflict was between the text in SDetailSingleItemRow.cpp and SColorGradingDetailTreeRow.cpp, as they used the same namespace and "copy to clipboard" text.

#rnx

diff --git a/Engine/Source/Editor/PropertyEditor/Private/SDetailSingleItemRow.cpp b/Engine/Source/Editor/PropertyEditor/Private/SDetailSingleItemRow.cpp
--- a/Engine/Source/Editor/PropertyEditor/Private/SDetailSingleItemRow.cpp
+++ b/Engine/Source/Editor/PropertyEditor/Private/SDetailSingleItemRow.cpp
@@ -1135,7 +1135,7 @@ void SDetailSingleItemRow::PopulateContextMenu(UToolMenu* ToolMenu)
 			FUIAction CopyDisplayNameAction = FExecuteAction::CreateSP(this, &SDetailSingleItemRow::OnCopyPropertyDisplayName);
 			CopyDisplayNameAction.CanExecuteAction = FCanExecuteAction::CreateSP(this, &SDetailSingleItemRow::CanCopyPropertyDisplayName);
 
-			static const FTextFormat TooltipFormat = NSLOCTEXT("PropertyView", "CopyPropertyDisplayName_ToolTip", "Copy the display name of this property to the system clipboard:\n{0}");
+			static const FTextFormat TooltipFormat = NSLOCTEXT("PropertyView_Single", "CopyPropertyDisplayName_ToolTip", "Copy the display name of this property to the system clipboard:\n{0}");
 
 			EditSection.AddMenuEntry(
 				TEXT("CopyDisplayName"),
@@ -1150,7 +1150,7 @@ void SDetailSingleItemRow::PopulateContextMenu(UToolMenu* ToolMenu)
 			FUIAction CopyInternalNameAction = FExecuteAction::CreateSP(this, &SDetailSingleItemRow::OnCopyPropertyInternalName);
 			CopyInternalNameAction.CanExecuteAction = FCanExecuteAction::CreateSP(this, &SDetailSingleItemRow::CanCopyPropertyInternalName);
 
-			static const FTextFormat TooltipFormat = NSLOCTEXT("PropertyView", "CopyPropertyInternalName_ToolTip", "Copy the internal name of this property to the system clipboard:\n{0}");
+			static const FTextFormat TooltipFormat = NSLOCTEXT("PropertyView_Single", "CopyPropertyInternalName_ToolTip", "Copy the internal name of this property to the system clipboard:\n{0}");
 
 			EditSection.AddMenuEntry(
 				TEXT("CopyInternalName"),
