Allow drag & dropping assets on a matching TScriptInterface picker

diff --git a/Engine/Source/Editor/PropertyEditor/Private/UserInterface/PropertyEditor/SPropertyEditorAsset.cpp b/Engine/Source/Editor/PropertyEditor/Private/UserInterface/PropertyEditor/SPropertyEditorAsset.cpp
--- a/Engine/Source/Editor/PropertyEditor/Private/UserInterface/PropertyEditor/SPropertyEditorAsset.cpp
+++ b/Engine/Source/Editor/PropertyEditor/Private/UserInterface/PropertyEditor/SPropertyEditorAsset.cpp
@@ -1454,7 +1454,7 @@ FSlateColor SPropertyEditorAsset::GetAssetClassColor()
 bool SPropertyEditorAsset::OnAssetDraggedOver( TArrayView<FAssetData> InAssets, FText& OutReason ) const
 {
 	UObject* AssetObject = InAssets[0].GetAsset();
-	if (CanEdit() && (AssetObject != nullptr) && AssetObject->IsA(ObjectClass))
+	if (CanEdit() && (AssetObject != nullptr) && (AssetObject->IsA(ObjectClass) || AssetObject->GetClass()->ImplementsInterface(ObjectClass)))
 	{
 		FAssetData AssetData(InAssets[0]);
 		// Check against custom asset filter
