Fixed const correctness of TSoftObjectPtr usage in permission lists.

diff --git a/Engine/Source/Editor/PropertyEditor/Private/PropertyPermissionList.cpp b/Engine/Source/Editor/PropertyEditor/Private/PropertyPermissionList.cpp
--- a/Engine/Source/Editor/PropertyEditor/Private/PropertyPermissionList.cpp
+++ b/Engine/Source/Editor/PropertyEditor/Private/PropertyPermissionList.cpp
@@ -71,7 +71,7 @@ void FPropertyPermissionList::RegisterOnBlueprintCompiled()
 	}
 }
 
-void FPropertyPermissionList::ClearCacheAndBroadcast(TSoftObjectPtr<UStruct> ObjectStruct, FName OwnerName)
+void FPropertyPermissionList::ClearCacheAndBroadcast(TSoftObjectPtr<const UStruct> ObjectStruct, FName OwnerName)
 {
 	// The cache isn't too expensive to recompute, so it is cleared
 	// and lazily repopulated any time the raw PermissionList changes.
@@ -83,7 +83,7 @@ void FPropertyPermissionList::ClearCacheAndBroadcast(TSoftObjectPtr<UStruct> Obj
 	}
 }
 
-void FPropertyPermissionList::AddPermissionList(TSoftObjectPtr<UStruct> Struct, const FNamePermissionList& PermissionList, const EPropertyPermissionListRules Rules, const TConstArrayView<FName> InAdditionalOwnerNames)
+void FPropertyPermissionList::AddPermissionList(TSoftObjectPtr<const UStruct> Struct, const FNamePermissionList& PermissionList, const EPropertyPermissionListRules Rules, const TConstArrayView<FName> InAdditionalOwnerNames)
 {
 	FPropertyPermissionListEntry& Entry = RawPropertyPermissionList.FindOrAdd(Struct);
 	Entry.PermissionList = PermissionList;
@@ -101,12 +101,12 @@ void FPropertyPermissionList::AddPermissionList(TSoftObjectPtr<UStruct> Struct,
 	ClearCacheAndBroadcast(Struct);
 }
 
-void FPropertyPermissionList::AddPermissionList(TSoftObjectPtr<UStruct> Struct, const FNamePermissionList& PermissionList, const EPropertyPermissionListRules Rules)
+void FPropertyPermissionList::AddPermissionList(TSoftObjectPtr<const UStruct> Struct, const FNamePermissionList& PermissionList, const EPropertyPermissionListRules Rules)
 {
 	AddPermissionList(Struct, PermissionList, Rules, {});
 }
 
-void FPropertyPermissionList::RemovePermissionList(TSoftObjectPtr<UStruct> Struct)
+void FPropertyPermissionList::RemovePermissionList(TSoftObjectPtr<const UStruct> Struct)
 {
 	if (RawPropertyPermissionList.Remove(Struct) > 0)
 	{
@@ -123,9 +123,9 @@ void FPropertyPermissionList::ClearPermissionList()
 
 void FPropertyPermissionList::UnregisterOwner(const FName Owner)
 {
-	TArray<TSoftObjectPtr<UStruct>> StructsToRemove;
+	TArray<TSoftObjectPtr<const UStruct>> StructsToRemove;
 
-	for (TPair<TSoftObjectPtr<UStruct>, FPropertyPermissionListEntry>& Pair : RawPropertyPermissionList)
+	for (TPair<TSoftObjectPtr<const UStruct>, FPropertyPermissionListEntry>& Pair : RawPropertyPermissionList)
 	{
 		Pair.Value.PermissionList.UnregisterOwner(Owner);
 		Pair.Value.AdditionalOwnerNames.Remove(Owner);
@@ -137,7 +137,7 @@ void FPropertyPermissionList::UnregisterOwner(const FName Owner)
 
 	{
 		TGuardValue<bool> SuppressGuard(bSuppressUpdateDelegate, true);
-		for (TSoftObjectPtr<UStruct>& StructToRemove : StructsToRemove)
+		for (TSoftObjectPtr<const UStruct>& StructToRemove : StructsToRemove)
 		{
 			RemovePermissionList(StructToRemove);
 		}
@@ -146,7 +146,7 @@ void FPropertyPermissionList::UnregisterOwner(const FName Owner)
 	ClearCacheAndBroadcast(nullptr, Owner);
 }
 
-void FPropertyPermissionList::AddToAllowList(TSoftObjectPtr<UStruct> Struct, const FName PropertyName, const FName Owner)
+void FPropertyPermissionList::AddToAllowList(TSoftObjectPtr<const UStruct> Struct, const FName PropertyName, const FName Owner)
 {
 	FPropertyPermissionListEntry& Entry = RawPropertyPermissionList.FindOrAdd(Struct);
 	if (Entry.PermissionList.AddAllowListItem(Owner, PropertyName))
@@ -155,7 +155,7 @@ void FPropertyPermissionList::AddToAllowList(TSoftObjectPtr<UStruct> Struct, con
 	}
 }
 
-void FPropertyPermissionList::AddToAllowList(TSoftObjectPtr<UStruct> Struct, const TArray<FName>& PropertyNames, const FName Owner)
+void FPropertyPermissionList::AddToAllowList(TSoftObjectPtr<const UStruct> Struct, const TArray<FName>& PropertyNames, const FName Owner)
 {
 	FPropertyPermissionListEntry& Entry = RawPropertyPermissionList.FindOrAdd(Struct);
 	bool bAddedItem = false;
@@ -173,7 +173,7 @@ void FPropertyPermissionList::AddToAllowList(TSoftObjectPtr<UStruct> Struct, con
 	}
 }
 
-void FPropertyPermissionList::RemoveFromAllowList(TSoftObjectPtr<UStruct> Struct, const FName PropertyName, const FName Owner)
+void FPropertyPermissionList::RemoveFromAllowList(TSoftObjectPtr<const UStruct> Struct, const FName PropertyName, const FName Owner)
 {
 	FPropertyPermissionListEntry& Entry = RawPropertyPermissionList.FindOrAdd(Struct);
 	if (Entry.PermissionList.RemoveAllowListItem(Owner, PropertyName))
@@ -182,7 +182,7 @@ void FPropertyPermissionList::RemoveFromAllowList(TSoftObjectPtr<UStruct> Struct
 	}
 }
 
-void FPropertyPermissionList::AddToDenyList(TSoftObjectPtr<UStruct> Struct, const FName PropertyName, const FName Owner)
+void FPropertyPermissionList::AddToDenyList(TSoftObjectPtr<const UStruct> Struct, const FName PropertyName, const FName Owner)
 {
 	FPropertyPermissionListEntry& Entry = RawPropertyPermissionList.FindOrAdd(Struct);
 	if (Entry.PermissionList.AddDenyListItem(Owner, PropertyName))
@@ -191,7 +191,7 @@ void FPropertyPermissionList::AddToDenyList(TSoftObjectPtr<UStruct> Struct, cons
 	}
 }
 
-void FPropertyPermissionList::RemoveFromDenyList(TSoftObjectPtr<UStruct> Struct, const FName PropertyName, const FName Owner)
+void FPropertyPermissionList::RemoveFromDenyList(TSoftObjectPtr<const UStruct> Struct, const FName PropertyName, const FName Owner)
 {
 	FPropertyPermissionListEntry& Entry = RawPropertyPermissionList.FindOrAdd(Struct);
 	if (Entry.PermissionList.RemoveDenyListItem(Owner, PropertyName))
diff --git a/Engine/Source/Editor/PropertyEditor/Private/SDetailsViewBase.cpp b/Engine/Source/Editor/PropertyEditor/Private/SDetailsViewBase.cpp
--- a/Engine/Source/Editor/PropertyEditor/Private/SDetailsViewBase.cpp
+++ b/Engine/Source/Editor/PropertyEditor/Private/SDetailsViewBase.cpp
@@ -45,7 +45,7 @@ SDetailsViewBase::SDetailsViewBase() :
 	}
 
 	// RequestForceRefresh is deferred until next tick to avoid the editor locking up for several minutes in one frame when refreshing multiple times
-	PropertyPermissionListChangedDelegate = FPropertyEditorPermissionList::Get().PermissionListUpdatedDelegate.AddLambda([this](TSoftObjectPtr<UStruct> Struct, FName Owner) { RequestForceRefresh(); });
+	PropertyPermissionListChangedDelegate = FPropertyEditorPermissionList::Get().PermissionListUpdatedDelegate.AddLambda([this](TSoftObjectPtr<const UStruct> Struct, FName Owner) { RequestForceRefresh(); });
 	PropertyPermissionListEnabledDelegate = FPropertyEditorPermissionList::Get().PermissionListEnabledDelegate.AddRaw(this, &SDetailsViewBase::RequestForceRefresh);
 }
 
diff --git a/Engine/Source/Editor/PropertyEditor/Private/SPropertyTreeViewImpl.cpp b/Engine/Source/Editor/PropertyEditor/Private/SPropertyTreeViewImpl.cpp
--- a/Engine/Source/Editor/PropertyEditor/Private/SPropertyTreeViewImpl.cpp
+++ b/Engine/Source/Editor/PropertyEditor/Private/SPropertyTreeViewImpl.cpp
@@ -170,7 +170,7 @@ void SPropertyTreeViewImpl::Construct(const FArguments& InArgs)
 	FPropertyEditorPermissionList::Get().PermissionListEnabledDelegate.AddSP(this, &SPropertyTreeViewImpl::RequestRefresh);
 }
 
-void SPropertyTreeViewImpl::OnPermissionListUpdated(TSoftObjectPtr<UStruct>, FName)
+void SPropertyTreeViewImpl::OnPermissionListUpdated(TSoftObjectPtr<const UStruct>, FName)
 {
 	RequestRefresh();
 }
diff --git a/Engine/Source/Editor/PropertyEditor/Private/SPropertyTreeViewImpl.h b/Engine/Source/Editor/PropertyEditor/Private/SPropertyTreeViewImpl.h
--- a/Engine/Source/Editor/PropertyEditor/Private/SPropertyTreeViewImpl.h
+++ b/Engine/Source/Editor/PropertyEditor/Private/SPropertyTreeViewImpl.h
@@ -183,7 +183,7 @@ public:
 private:
 
 	/** Delegate called when the property editor permission list is updated */
-	void OnPermissionListUpdated(TSoftObjectPtr<UStruct>, FName);
+	void OnPermissionListUpdated(TSoftObjectPtr<const UStruct>, FName);
 
 	/** Restores the expansion state in the tree for a specific object set */
 	void RestoreExpandedItems();
diff --git a/Engine/Source/Editor/PropertyEditor/Public/PropertyPermissionList.h b/Engine/Source/Editor/PropertyEditor/Public/PropertyPermissionList.h
--- a/Engine/Source/Editor/PropertyEditor/Public/PropertyPermissionList.h
+++ b/Engine/Source/Editor/PropertyEditor/Public/PropertyPermissionList.h
@@ -11,7 +11,7 @@
 DECLARE_LOG_CATEGORY_EXTERN(LogPropertyEditorPermissionList, Log, All);
 
 /** Struct, OwnerName */
-DECLARE_MULTICAST_DELEGATE_TwoParams(FPermissionListUpdated, TSoftObjectPtr<UStruct>, FName);
+DECLARE_MULTICAST_DELEGATE_TwoParams(FPermissionListUpdated, TSoftObjectPtr<const UStruct>, FName);
 
 /**
  * A hierarchical set of rules that can be used to PermissionList all properties of specific Structs without
@@ -40,7 +40,7 @@ enum class EPropertyPermissionListRules : uint8
 
 struct FPropertyPermissionListEntry
 {
-    FNamePermissionList PermissionList;
+	FNamePermissionList PermissionList;
 	// When the permission list does not contain any AllowList, DenyList or DenyListAll it needs an AdditionalOwnerName to be tracked here so stays alive without being removed
 	TArray<FName> AdditionalOwnerNames;
 	EPropertyPermissionListRules Rules = EPropertyPermissionListRules::UseExistingPermissionList;
@@ -51,11 +51,11 @@ class PROPERTYEDITOR_API FPropertyPermissionList
 public:
 	/** Add a set of rules for a specific base UStruct to determine which properties are visible in all details panels */
 	UE_DEPRECATED(5.5, "Call AddPermissionList with additional required arguments instead.")
-	void AddPermissionList(TSoftObjectPtr<UStruct> Struct, const FNamePermissionList& PermissionList, const EPropertyPermissionListRules Rules = EPropertyPermissionListRules::UseExistingPermissionList);
+	void AddPermissionList(TSoftObjectPtr<const UStruct> Struct, const FNamePermissionList& PermissionList, const EPropertyPermissionListRules Rules = EPropertyPermissionListRules::UseExistingPermissionList);
 	/** Add a set of rules for a specific base UStruct to determine which properties are visible in all details panels */
-	void AddPermissionList(TSoftObjectPtr<UStruct> Struct, const FNamePermissionList& PermissionList, const EPropertyPermissionListRules Rules, const TConstArrayView<FName> InAdditionalOwnerNames);
+	void AddPermissionList(TSoftObjectPtr<const UStruct> Struct, const FNamePermissionList& PermissionList, const EPropertyPermissionListRules Rules, const TConstArrayView<FName> InAdditionalOwnerNames);
 	/** Remove a set of rules for a specific base UStruct to determine which properties are visible in all details panels */
-	void RemovePermissionList(TSoftObjectPtr<UStruct> Struct);
+	void RemovePermissionList(TSoftObjectPtr<const UStruct> Struct);
 	/** Remove all rules */
 	void ClearPermissionList();
 
@@ -63,19 +63,19 @@ public:
 	void UnregisterOwner(const FName Owner);
 
 	/** Add a specific property to a UStruct's AllowList */
-	void AddToAllowList(TSoftObjectPtr<UStruct> Struct, const FName PropertyName, const FName Owner);
+	void AddToAllowList(TSoftObjectPtr<const UStruct> Struct, const FName PropertyName, const FName Owner);
 	/** Add a list of properties to a UStruct's AllowList */
-	void AddToAllowList(TSoftObjectPtr<UStruct> Struct, const TArray<FName>& PropertyNames, const FName Owner);
+	void AddToAllowList(TSoftObjectPtr<const UStruct> Struct, const TArray<FName>& PropertyNames, const FName Owner);
 	/** Remove a specific property from a UStruct's AllowList */
-	void RemoveFromAllowList(TSoftObjectPtr<UStruct> Struct, const FName PropertyName, const FName Owner);
+	void RemoveFromAllowList(TSoftObjectPtr<const UStruct> Struct, const FName PropertyName, const FName Owner);
 	/** Add a specific property to a UStruct's DenyList */
-	void AddToDenyList(TSoftObjectPtr<UStruct> Struct, const FName PropertyName, const FName Owner);
+	void AddToDenyList(TSoftObjectPtr<const UStruct> Struct, const FName PropertyName, const FName Owner);
 	/** Remove a specific property from a UStruct's DenyList */
-    void RemoveFromDenyList(TSoftObjectPtr<UStruct> Struct, const FName PropertyName, const FName Owner);
+	void RemoveFromDenyList(TSoftObjectPtr<const UStruct> Struct, const FName PropertyName, const FName Owner);
 
 	/** When the PermissionList or DenyList for any struct was added to or removed from. */
-    FPermissionListUpdated PermissionListUpdatedDelegate;
-    
+	FPermissionListUpdated PermissionListUpdatedDelegate;
+
 	/** When the entire PermissionList is enabled or disabled */
 	FSimpleMulticastDelegate PermissionListEnabledDelegate;
 
@@ -106,7 +106,7 @@ public:
 	bool IsSpecificPropertyDenyListed(const UStruct* ObjectStruct, FName PropertyName) const;
 
 	/** Gets a read-only copy of the original, un-flattened PermissionList. */
-	const TMap<TSoftObjectPtr<UStruct>, FPropertyPermissionListEntry>& GetRawPermissionList() const { return RawPropertyPermissionList; }
+	const TMap<TSoftObjectPtr<const UStruct>, FPropertyPermissionListEntry>& GetRawPermissionList() const { return RawPropertyPermissionList; }
 
 	/** Clear CachedPropertyPermissionList to cause the PermissionListed property list to be regenerated next time it's queried */
 	void ClearCache();
@@ -121,20 +121,20 @@ protected:
 private:
 	bool Tick(float DeltaTime);
 	void RegisterOnBlueprintCompiled();
-	void ClearCacheAndBroadcast(TSoftObjectPtr<UStruct> ObjectStruct = nullptr, FName OwnerName = NAME_None);
-	
+	void ClearCacheAndBroadcast(TSoftObjectPtr<const UStruct> ObjectStruct = nullptr, FName OwnerName = NAME_None);
+
 	/** Whether DoesPropertyPassFilter should perform its PermissionList check or always return true */
 	bool bEnablePermissionList = false;
-	
+
 	/** Stores assigned PermissionLists from AddPermissionList(), which are later flattened and stored in CachedPropertyPermissionList. */
-	TMap<TSoftObjectPtr<UStruct>, FPropertyPermissionListEntry> RawPropertyPermissionList;
+	TMap<TSoftObjectPtr<const UStruct>, FPropertyPermissionListEntry> RawPropertyPermissionList;
 
 	/** Handle for our tick function */
 	FTSTicker::FDelegateHandle OnTickHandle;
 
 	struct FPermissionListUpdate
 	{
-		TSoftObjectPtr<UStruct> ObjectStruct;
+		TSoftObjectPtr<const UStruct> ObjectStruct;
 		FName OwnerName;
 	};
 	friend bool operator==(const FPermissionListUpdate& A, const FPermissionListUpdate& B);
