Fix "Multiple Values" feature was broken on some struct properties.

CL 25214019 has introduced an optimization(?) which would not fill up OutAddresses in FStructurePropertyNode::GetReadAddressUncached if not all values match.
However, GetSingleReadAddress needs that array to be filled up to be able to differentiate between MultipleValues and Fail states.
FObjectPropertyNode::GetReadAddressUncached does not have such optimization, so this change also makes struct and object behavior more consistent (once again).

diff --git a/Engine/Source/Editor/PropertyEditor/Private/StructurePropertyNode.cpp b/Engine/Source/Editor/PropertyEditor/Private/StructurePropertyNode.cpp
--- a/Engine/Source/Editor/PropertyEditor/Private/StructurePropertyNode.cpp
+++ b/Engine/Source/Editor/PropertyEditor/Private/StructurePropertyNode.cpp
@@ -206,7 +206,7 @@ bool FStructurePropertyNode::GetReadAddressUncached(const FPropertyNode& InPrope
 		}
 	}
 
-	if (bAllTheSame && OutAddresses)
+	if (OutAddresses)
 	{
 		for (TSharedPtr<FStructOnScope>& Instance : Instances)
 		{
