Allow enum properties to be display as segmented button widget. Any enum property that is annotated with the "SegmentedDisplay" meta tag will use this new widget instead of a combo box dropdown in the details panel.

diff --git a/Engine/Source/Editor/PropertyEditor/Private/UserInterface/PropertyEditor/SPropertyEditorCombo.cpp b/Engine/Source/Editor/PropertyEditor/Private/UserInterface/PropertyEditor/SPropertyEditorCombo.cpp
--- a/Engine/Source/Editor/PropertyEditor/Private/UserInterface/PropertyEditor/SPropertyEditorCombo.cpp
+++ b/Engine/Source/Editor/PropertyEditor/Private/UserInterface/PropertyEditor/SPropertyEditorCombo.cpp
@@ -2,7 +2,7 @@
 
 #include "UserInterface/PropertyEditor/SPropertyEditorCombo.h"
 #include "IDocumentation.h"
-
+#include "Widgets/SToolTip.h"
 #include "PropertyEditorHelpers.h"
 #include "UserInterface/PropertyEditor/SPropertyComboBox.h"
 
@@ -39,12 +39,14 @@ void SPropertyEditorCombo::Construct( const FArguments& InArgs, const TSharedPtr
 	PropertyEditor = InPropertyEditor;
 	ComboArgs = InArgs._ComboArgs;
 
+	bool bSegmentedDisplay = false;
 	if (PropertyEditor.IsValid())
 	{
 		ComboArgs.PropertyHandle = PropertyEditor->GetPropertyHandle();
 		if (ComboArgs.PropertyHandle.IsValid())
 		{
 			ComboArgs.PropertyHandle->SetOnPropertyResetToDefault(FSimpleDelegate::CreateSP(this, &SPropertyEditorCombo::OnResetToDefault));
+			bSegmentedDisplay = ComboArgs.PropertyHandle->GetBoolMetaData(TEXT("SegmentedDisplay"));
 		}
 	}
 
@@ -61,21 +63,59 @@ void SPropertyEditorCombo::Construct( const FArguments& InArgs, const TSharedPtr
 
 	GenerateComboBoxStrings(ComboItems, RichToolTips, Restrictions);
 
-	SAssignNew(ComboBox, SPropertyComboBox)
-		.Font( ComboArgs.Font )
-		.RichToolTipList( RichToolTips )
-		.ComboItemList( ComboItems )
-		.RestrictedList( Restrictions )
-		.OnSelectionChanged( this, &SPropertyEditorCombo::OnComboSelectionChanged )
-		.OnComboBoxOpening( this, &SPropertyEditorCombo::OnComboOpening )
-		.VisibleText( this, &SPropertyEditorCombo::GetDisplayValueAsString )
-		.ToolTipText( this, &SPropertyEditorCombo::GetValueToolTip )
-		.ShowSearchForItemCount( ComboArgs.ShowSearchForItemCount );
-
-	ChildSlot
-	[
-		ComboBox.ToSharedRef()
-	];
+	if (bSegmentedDisplay)
+	{
+		ParameterTextStyle = FTextBlockStyle(FAppStyle::GetWidgetStyle<FTextBlockStyle>("NormalText"))
+			.SetFont(ComboArgs.Font);
+		
+		SAssignNew(SegmentControl, SSegmentedControl<FString>)
+			.UniformPadding(FMargin(10, 5))
+			.TextStyle(&ParameterTextStyle)
+			.Value(this, &SPropertyEditorCombo::GetDisplayValueAsString)
+			.OnValueChanged(this, &SPropertyEditorCombo::OnSegmentedControlSelectionChanged);
+
+		for (int ItemIndex = 0; ItemIndex < ComboItems.Num(); ItemIndex++)
+		{
+			const FString& ComboItem = *ComboItems[ItemIndex];
+			
+			FText DisplayName = FText::FromString(ComboItem);
+			FText TooltipText = DisplayName;
+			if (RichToolTips.IsValidIndex(ItemIndex) && !RichToolTips[ItemIndex]->IsEmpty())
+			{
+				TooltipText = RichToolTips[ItemIndex]->GetTextTooltip();
+			}
+
+			SSegmentedControl<FString>::FScopedWidgetSlotArguments Slot = SegmentControl->AddSlot(ComboItem);
+			Slot
+			  .HAlign(HAlign_Center)
+			  .VAlign(VAlign_Center)
+			  .Text(DisplayName)
+			  .ToolTip(TooltipText);
+		}
+		
+		ChildSlot
+		[
+			SegmentControl.ToSharedRef()
+		];
+	}
+	else
+	{
+		SAssignNew(ComboBox, SPropertyComboBox)
+			.Font( ComboArgs.Font )
+			.RichToolTipList( RichToolTips )
+			.ComboItemList( ComboItems )
+			.RestrictedList( Restrictions )
+			.OnSelectionChanged( this, &SPropertyEditorCombo::OnComboSelectionChanged )
+			.OnComboBoxOpening( this, &SPropertyEditorCombo::OnComboOpening )
+			.VisibleText( this, &SPropertyEditorCombo::GetDisplayValueAsString )
+			.ToolTipText( this, &SPropertyEditorCombo::GetValueToolTip )
+			.ShowSearchForItemCount( ComboArgs.ShowSearchForItemCount );
+
+		ChildSlot
+		[
+			ComboBox.ToSharedRef()
+		];
+	}
 
 	SetEnabled( TAttribute<bool>( this, &SPropertyEditorCombo::CanEdit ) );
 	SetToolTipText( TAttribute<FText>( this, &SPropertyEditorCombo::GetValueToolTip) );
@@ -271,10 +311,22 @@ void SPropertyEditorCombo::OnComboSelectionChanged( TSharedPtr<FString> NewValue
 	}
 }
 
+void SPropertyEditorCombo::OnSegmentedControlSelectionChanged(FString NewValue)
+{
+	SendToObjects(NewValue);
+}
+
 void SPropertyEditorCombo::OnResetToDefault()
 {
 	FString CurrentDisplayValue = GetDisplayValueAsString();
-	ComboBox->SetSelectedItem(CurrentDisplayValue);
+	if (ComboBox)
+	{
+		ComboBox->SetSelectedItem(CurrentDisplayValue);
+	}
+	if (SegmentControl)
+	{
+		SegmentControl->SetValue(CurrentDisplayValue);
+	}
 }
 
 void SPropertyEditorCombo::OnComboOpening()
diff --git a/Engine/Source/Editor/PropertyEditor/Private/UserInterface/PropertyEditor/SPropertyEditorCombo.h b/Engine/Source/Editor/PropertyEditor/Private/UserInterface/PropertyEditor/SPropertyEditorCombo.h
--- a/Engine/Source/Editor/PropertyEditor/Private/UserInterface/PropertyEditor/SPropertyEditorCombo.h
+++ b/Engine/Source/Editor/PropertyEditor/Private/UserInterface/PropertyEditor/SPropertyEditorCombo.h
@@ -9,6 +9,7 @@
 #include "Presentation/PropertyEditor/PropertyEditor.h"
 #include "UserInterface/PropertyEditor/PropertyEditorConstants.h"
 #include "PropertyCustomizationHelpers.h"
+#include "Widgets/Input/SSegmentedControl.h"
 
 class SPropertyComboBox;
 
@@ -30,6 +31,7 @@ public:
 private:
 	void GenerateComboBoxStrings( TArray< TSharedPtr<FString> >& OutComboBoxStrings, TArray<TSharedPtr<class SToolTip>>& OutToolTips, TArray<bool>& OutRestrictedItems );
 	void OnComboSelectionChanged( TSharedPtr<FString> NewValue, ESelectInfo::Type SelectInfo );
+	void OnSegmentedControlSelectionChanged( FString NewValue );
 	void OnResetToDefault();
 	void OnComboOpening();
 
@@ -55,6 +57,10 @@ private:
 	/** Fills out with generated strings. */
 	TSharedPtr<class SPropertyComboBox> ComboBox;
 
+	/** Alternative display widget. */
+	TSharedPtr<SSegmentedControl<FString>> SegmentControl;
+	FTextBlockStyle ParameterTextStyle;
+
 	/** Arguments used to construct the combo box */
 	FPropertyComboBoxArgs ComboArgs;
 
