[PropertyEditor] Moved CallInEditor functionality to shared location

diff --git a/Engine/Source/Editor/DetailCustomizations/Private/ObjectDetails.cpp b/Engine/Source/Editor/DetailCustomizations/Private/ObjectDetails.cpp
--- a/Engine/Source/Editor/DetailCustomizations/Private/ObjectDetails.cpp
+++ b/Engine/Source/Editor/DetailCustomizations/Private/ObjectDetails.cpp
@@ -17,16 +17,11 @@
 #include "HAL/PlatformCrt.h"
 #include "Internationalization/Internationalization.h"
 #include "Internationalization/Text.h"
-#include "Kismet2/BlueprintEditorUtils.h"
 #include "Layout/Margin.h"
 #include "Math/NumericLimits.h"
 #include "Misc/Attribute.h"
 #include "Misc/CString.h"
 #include "ObjectEditorUtils.h"
-#include "ObjectTools.h"
-#include "PropertyPermissionList.h"
-#include "Reflection/FunctionUtils.h"
-#include "Settings/BlueprintEditorProjectSettings.h"
 #include "SWarningOrErrorBox.h"
 #include "ScopedTransaction.h"
 #include "SlotBase.h"
@@ -42,7 +37,8 @@
 #include "Widgets/DeclarativeSyntaxSupport.h"
 #include "Widgets/Input/SButton.h"
 #include "Widgets/Layout/SBox.h"
-#include "Widgets/Layout/SWrapBox.h"
+#include "Modules/ModuleManager.h"
+#include "PropertyCustomizationHelpers.h"
 
 #define LOCTEXT_NAMESPACE "ObjectDetails"
 
@@ -94,232 +90,24 @@ void FObjectDetails::AddExperimentalWarningCategory(IDetailLayoutBuilder& Detail
 	}
 }
 
-static bool CanCallFunctionBasedOnParams(const UFunction* TestFunction)
-{
-	bool bCanCall = TestFunction->GetBoolMetaData(FBlueprintMetadata::MD_CallInEditor) && (TestFunction->ParmsSize == 0); // no params required, we can call it!
-
-	// else - if the function only takes a world context object we can use the editor's
-	// world context - but only if the blueprint is editor only:
-	if (UClass* TestFunctionOwnerClass = TestFunction->GetOwnerClass())
-	{
-		if (UBlueprint* Blueprint = Cast<UBlueprint>(TestFunctionOwnerClass->ClassGeneratedBy))
-		{
-			if (FBlueprintEditorUtils::IsEditorUtilityBlueprint(Blueprint) && Blueprint->BlueprintType == BPTYPE_FunctionLibrary)
-			{
-				using namespace UE::Reflection;
-				return TestFunction->HasMetaData(FBlueprintMetadata::MD_WorldContext) &&
-					DoesStaticFunctionSignatureMatch<void(TObjectPtr<UObject>)>(TestFunction);
-			}
-		}
-	}
-
-	return bCanCall;
-}
-
 void FObjectDetails::AddCallInEditorMethods(IDetailLayoutBuilder& DetailBuilder)
 {
-	// metadata tag for defining sort order of function buttons within a Category
-	static const FName NAME_DisplayPriority("DisplayPriority");
-
-	const bool bDisallowEditorUtilityBlueprintFunctions = GetDefault<UBlueprintEditorProjectSettings>()->bDisallowEditorUtilityBlueprintFunctionsInDetailsView;
-
-	// Get all of the functions we need to display (done ahead of time so we can sort them)
-	TArray<UFunction*, TInlineAllocator<8>> CallInEditorFunctions;
-	for (TFieldIterator<UFunction> FunctionIter(DetailBuilder.GetBaseClass(), EFieldIteratorFlags::IncludeSuper); FunctionIter; ++FunctionIter)
-	{
-		UFunction* TestFunction = *FunctionIter;
-
-		if (CanCallFunctionBasedOnParams(TestFunction))
-		{
-			bool bAllowFunction = true;
-			if (UClass* TestFunctionOwnerClass = TestFunction->GetOwnerClass())
-			{
-				if (UBlueprint* Blueprint = Cast<UBlueprint>(TestFunctionOwnerClass->ClassGeneratedBy))
-				{
-					if (FBlueprintEditorUtils::IsEditorUtilityBlueprint(Blueprint))
-					{
-						// Skip Blutilities if disabled via project settings
-						bAllowFunction = !bDisallowEditorUtilityBlueprintFunctions;
-					}
-				}
-			}
-
-			if (bAllowFunction)
-			{
-				const FName FunctionName = TestFunction->GetFName();
-				if (FPropertyEditorPermissionList::Get().DoesPropertyPassFilter(TestFunction->GetOwnerClass(), FunctionName))
-				{
-					if (!CallInEditorFunctions.FindByPredicate([&FunctionName](const UFunction* Func) { return Func->GetFName() == FunctionName; }))
-					{
-						CallInEditorFunctions.Add(*FunctionIter);
-					}
-				}
-			}
-		}
-	}
-
-	if (CallInEditorFunctions.Num() > 0)
-	{
-		// Copy off the objects being customized so we can invoke a function on them later, removing any that are a CDO
-		DetailBuilder.GetObjectsBeingCustomized(/*out*/ SelectedObjectsList);
-		SelectedObjectsList.RemoveAllSwap([](TWeakObjectPtr<UObject> ObjPtr) { UObject* Obj = ObjPtr.Get(); return (Obj == nullptr) || Obj->HasAnyFlags(RF_ArchetypeObject); });
-
-		if (SelectedObjectsList.Num() == 0)
-		{
-			// remove all non-static functions - no objects to call them on
-			CallInEditorFunctions.RemoveAllSwap([](const UFunction* Function) { return !Function->HasAnyFunctionFlags(FUNC_Static);});
-
-			if (CallInEditorFunctions.Num() == 0)
-			{
-				return;
-			}
-		}
-
-		// Sort the functions by category and then by DisplayPriority meta tag, and then by name
-		CallInEditorFunctions.Sort([](UFunction& A, UFunction& B)
-		{
-			const int32 CategorySort = A.GetMetaData(FBlueprintMetadata::MD_FunctionCategory).Compare(B.GetMetaData(FBlueprintMetadata::MD_FunctionCategory));
-			if (CategorySort != 0)
-			{
-				return (CategorySort <= 0);
-			}
-			else 
-			{
-				FString DisplayPriorityAStr = A.GetMetaData(NAME_DisplayPriority);
-				int32 DisplayPriorityA = (DisplayPriorityAStr.IsEmpty() ? MAX_int32 : FCString::Atoi(*DisplayPriorityAStr));
-				if (DisplayPriorityA == 0 && !FCString::IsNumeric(*DisplayPriorityAStr))
-				{
-					DisplayPriorityA = MAX_int32;
-				}
-
-				FString DisplayPriorityBStr = B.GetMetaData(NAME_DisplayPriority);
-				int32 DisplayPriorityB = (DisplayPriorityBStr.IsEmpty() ? MAX_int32 : FCString::Atoi(*DisplayPriorityBStr));
-				if (DisplayPriorityB == 0 && !FCString::IsNumeric(*DisplayPriorityBStr))
-				{
-					DisplayPriorityB = MAX_int32;
-				}
-
-				return (DisplayPriorityA == DisplayPriorityB) ? (A.GetName() <= B.GetName()) : (DisplayPriorityA <= DisplayPriorityB);
-			}
-		});
-
-		struct FCategoryEntry
-		{
-			FName CategoryName;
-			FName RowTag;
-			TSharedPtr<SWrapBox> WrapBox;
-			FTextBuilder FunctionSearchText;
-
-			FCategoryEntry(FName InCategoryName)
-				: CategoryName(InCategoryName)
-			{
-				WrapBox = SNew(SWrapBox)
-					// Setting the preferred size here (despite using UseAllottedSize) is a workaround for an issue
-					// when contained in a scroll box: prior to the first tick, the wrap box will use preferred size
-					// instead of allotted, and if preferred size is set small, it will cause the box to wrap a lot and
-					// request too much space from the scroll box. On next tick, SWrapBox is updated but the scroll box
-					// does not realize that it needs to show more elements, until it is scrolled.
-					// Setting a large value here means that the SWrapBox will request too little space prior to tick,
-					// which will cause the scroll box to virtualize more elements at the start, but this is less broken.
-					.PreferredSize(2000)
-					.UseAllottedSize(true);
-			}
-		};
-
-		// Build up a set of functions for each category, accumulating search text and buttons in a wrap box
-		FName ActiveCategory;
-		TArray<FCategoryEntry, TInlineAllocator<8>> CategoryList;
-		for (UFunction* Function : CallInEditorFunctions)
-		{
-			FName FunctionCategoryName(NAME_Default);
-			if (Function->HasMetaData(FBlueprintMetadata::FBlueprintMetadata::MD_FunctionCategory))
-			{
-				FunctionCategoryName = FName(*Function->GetMetaData(FBlueprintMetadata::MD_FunctionCategory));
-			}
-
-			if (FunctionCategoryName != ActiveCategory)
-			{
-				ActiveCategory = FunctionCategoryName;
-				CategoryList.Emplace(FunctionCategoryName);
-			}
-			FCategoryEntry& CategoryEntry = CategoryList.Last();
-
-			const FText ButtonCaption = ObjectTools::GetUserFacingFunctionName(Function);
-			FText FunctionTooltip = Function->GetToolTipText();
-			if (FunctionTooltip.IsEmpty())
-			{
-				FunctionTooltip = ButtonCaption;
-			}
-
-			TWeakObjectPtr<UFunction> WeakFunctionPtr(Function);
-			CategoryEntry.WrapBox->AddSlot()
-			.Padding(0.0f, 0.0f, 5.0f, 3.0f)
-			[
-				SNew(SButton)
-				.Text(ButtonCaption)
-				.OnClicked(FOnClicked::CreateSP(this, &FObjectDetails::OnExecuteCallInEditorFunction, WeakFunctionPtr))
-				.ToolTipText(FunctionTooltip.IsEmptyOrWhitespace() ? LOCTEXT("CallInEditorTooltip", "Call an event on the selected object(s)") : FunctionTooltip)
-			];
-
-			CategoryEntry.RowTag = Function->GetFName();
-			CategoryEntry.FunctionSearchText.AppendLine(ButtonCaption);
-			CategoryEntry.FunctionSearchText.AppendLine(FunctionTooltip);
-
-			if (ButtonCaption.ToString() != Function->GetName())
-			{
-				CategoryEntry.FunctionSearchText.AppendLine(FText::FromString(Function->GetName()));
-			}
-		}
-		
-		// Now edit the categories, adding the button strips to the details panel
-		for (FCategoryEntry& CategoryEntry : CategoryList)
-		{
-			IDetailCategoryBuilder& CategoryBuilder = DetailBuilder.EditCategory(CategoryEntry.CategoryName);
-			CategoryBuilder.AddCustomRow(CategoryEntry.FunctionSearchText.ToText())
-			.RowTag(CategoryEntry.RowTag)
-			[
-				CategoryEntry.WrapBox.ToSharedRef()
-			];
-		}
-	}
+	TArray<UFunction*> CallInEditorFunctions;
+	PropertyCustomizationHelpers::GetCallInEditorFunctionsForClass(
+		DetailBuilder.GetBaseClass(),
+		CallInEditorFunctions);
+
+	PropertyCustomizationHelpers::AddFunctionCallWidgets(
+		DetailBuilder,
+		CallInEditorFunctions,
+		FPropertyFunctionCallDelegates(
+			FPropertyFunctionCallDelegates::FOnGetExecutionContext::CreateSP(this, &FObjectDetails::GetFunctionCallExecutionContext)
+		));
 }
 
-FReply FObjectDetails::OnExecuteCallInEditorFunction(TWeakObjectPtr<UFunction> WeakFunctionPtr)
+TArray<TWeakObjectPtr<UObject>> FObjectDetails::GetFunctionCallExecutionContext(TWeakObjectPtr<UFunction> InWeakFunction) const
 {
-	using namespace UE::Reflection;
-	if (UFunction* Function = WeakFunctionPtr.Get())
-	{
-		//@TODO: Consider naming the transaction scope after the fully qualified function name for better UX
-		FScopedTransaction Transaction(LOCTEXT("ExecuteCallInEditorMethod", "Call In Editor Action"));
-		TStrongObjectPtr<UFunction> CallingFunction(Function);
-
-		if (Function->HasMetaData(FBlueprintMetadata::MD_WorldContext) &&
-			DoesStaticFunctionSignatureMatch<void(TObjectPtr<UObject>)>(Function))
-		{
-			FEditorScriptExecutionGuard ScriptGuard;
-			extern ENGINE_API class UEngine* GEngine;
-			UEditorEngine* EditorEngine = Cast<UEditorEngine>(GEngine);
-			UObject* WorldContextObject = EditorEngine->GetEditorWorldContext().World();
-			TStrongObjectPtr<UObject> CDO(Function->GetOwnerClass()->ClassDefaultObject);
-			CDO->ProcessEvent(Function, &WorldContextObject);
-		}
-		else
-		{
-			FEditorScriptExecutionGuard ScriptGuard;
-			for (TWeakObjectPtr<UObject> SelectedObjectPtr : SelectedObjectsList)
-			{
-				if (UObject* Object = SelectedObjectPtr.Get())
-				{
-					ensure(Function->ParmsSize == 0);
-					TStrongObjectPtr<UObject> ObjectStrong(Object);
-					Object->ProcessEvent(Function, nullptr);
-				}
-			}
-		}
-
-	}
-
-	return FReply::Handled();
+	return SelectedObjectsList;
 }
 
 #undef LOCTEXT_NAMESPACE
diff --git a/Engine/Source/Editor/DetailCustomizations/Private/ObjectDetails.h b/Engine/Source/Editor/DetailCustomizations/Private/ObjectDetails.h
--- a/Engine/Source/Editor/DetailCustomizations/Private/ObjectDetails.h
+++ b/Engine/Source/Editor/DetailCustomizations/Private/ObjectDetails.h
@@ -30,8 +30,7 @@ private:
 	// Creates a button strip in each category that contains reflected functions marked as CallInEditor
 	void AddCallInEditorMethods(IDetailLayoutBuilder& DetailBuilder);
 
-	// Executes the specified method on the selected objects
-	FReply OnExecuteCallInEditorFunction(TWeakObjectPtr<UFunction> WeakFunctionPtr);
+	TArray<TWeakObjectPtr<UObject>> GetFunctionCallExecutionContext(TWeakObjectPtr<UFunction> InWeakFunction) const;
 
 private:
 	// The list of selected objects, used when invoking a CallInEditor method
diff --git a/Engine/Source/Editor/PropertyEditor/Private/PropertyCustomizationHelpers.cpp b/Engine/Source/Editor/PropertyEditor/Private/PropertyCustomizationHelpers.cpp
--- a/Engine/Source/Editor/PropertyEditor/Private/PropertyCustomizationHelpers.cpp
+++ b/Engine/Source/Editor/PropertyEditor/Private/PropertyCustomizationHelpers.cpp
@@ -29,13 +29,19 @@
 #include "UserInterface/PropertyEditor/SPropertyEditorInteractiveActorPicker.h"
 #include "UserInterface/PropertyEditor/SPropertyEditorSceneDepthPicker.h"
 #include "Widgets/Input/SHyperlink.h"
-#include "Widgets/Layout/SWidgetSwitcher.h"
 #include "IDocumentation.h"
-#include "SResetToDefaultPropertyEditor.h"
 #include "EditorFontGlyphs.h"
 #include "DetailCategoryBuilder.h"
 #include "IDetailGroup.h"
 #include "AssetToolsModule.h"
+#include "Kismet2/BlueprintEditorUtils.h"
+#include "ObjectTools.h"
+#include "PropertyPermissionList.h"
+#include "Reflection/FunctionUtils.h"
+#include "ScopedTransaction.h"
+#include "Settings/BlueprintEditorProjectSettings.h"
+#include "Settings/EditorStyleSettings.h"
+#include "Widgets/Layout/SWrapBox.h"
 
 #define LOCTEXT_NAMESPACE "PropertyCustomizationHelpers"
 
@@ -221,6 +227,273 @@ namespace PropertyCustomizationHelpers
 			.IsFocusable(false);
 	}
 
+	TSharedRef<SWidget> MakeFunctionCallButton(const FPropertyFunctionCallArgs& InArgs)
+	{
+		TSharedRef<SWidget> Widget = SNullWidget::NullWidget;
+		if (const UFunction* Function = InArgs.Function.Get())
+		{
+			const FText Label = InArgs.LabelOverride.Get(Function->GetDisplayNameText());
+			FText ToolTipText = InArgs.ToolTipTextOverride.Get(Function->GetToolTipText());
+			if (ToolTipText.IsEmpty())
+			{
+				ToolTipText = Label;
+			}
+
+			TAttribute<bool> IsEnabled;
+			if (InArgs.OnCanExecute.IsBound())
+			{
+				IsEnabled.Bind(
+					TAttribute<bool>::FGetter::CreateLambda(
+					[WeakFunctionPtr = InArgs.Function, CanExecute = InArgs.OnCanExecute]()
+					{
+						return CanExecute.Execute(WeakFunctionPtr);
+					})
+				);
+			}
+			else
+			{
+				IsEnabled = true;
+			}
+
+			Widget = SNew(SButton)
+				.Text(Label)
+				.OnClicked_Lambda([WeakFunctionPtr = InArgs.Function, OnExecute = InArgs.OnExecute]()
+				{
+					return OnExecute.Execute(WeakFunctionPtr);
+				})
+				.IsEnabled(IsEnabled)
+				.ToolTipText(ToolTipText.IsEmptyOrWhitespace() ? LOCTEXT("CallInEditorTooltip", "Call an event on the selected object(s)") : ToolTipText);
+
+			if (InArgs.SearchText != nullptr)
+			{
+				InArgs.SearchText->AppendLine(Label);
+				InArgs.SearchText->AppendLine(ToolTipText);
+
+				if (Label.ToString() != Function->GetName())
+				{
+					InArgs.SearchText->AppendLine(FText::FromString(Function->GetName()));
+				}
+			}
+		}
+
+		return Widget;
+	}
+
+	struct FCategorizedFunctionCallEntry
+	{
+		FName CategoryName;
+		FName RowTag;
+		TSharedPtr<SWrapBox> WrapBox;
+		FTextBuilder FunctionSearchText;
+
+		FCategorizedFunctionCallEntry(FName InCategoryName)
+			: CategoryName(InCategoryName)
+		{
+			WrapBox = SNew(SWrapBox)
+				// Setting the preferred size here (despite using UseAllottedSize) is a workaround for an issue
+				// when contained in a scroll box: prior to the first tick, the wrap box will use preferred size
+				// instead of allotted, and if preferred size is set small, it will cause the box to wrap a lot and
+				// request too much space from the scroll box. On next tick, SWrapBox is updated but the scroll box
+				// does not realize that it needs to show more elements, until it is scrolled.
+				// Setting a large value here means that the SWrapBox will request too little space prior to tick,
+				// which will cause the scroll box to virtualize more elements at the start, but this is less broken.
+				.PreferredSize(2000)
+				.UseAllottedSize(true);
+		}
+	};
+
+	void GetFunctionCallWidgets(const TArrayView<UFunction*>& InCallInEditorFunctions, const FPropertyFunctionCallDelegates& InArgs, TArray<FCategorizedFunctionCallEntry, TInlineAllocator<8>>& OutCategorizedEntries)
+	{
+		if (InCallInEditorFunctions.IsEmpty())
+		{
+			return;
+		}
+
+		// Build up a set of functions for each category, accumulating search text and buttons in a wrap box
+		FName ActiveCategory;
+
+		const bool bUseDisplayNames = GEditor && GetDefault<UEditorStyleSettings>()->bShowFriendlyNames;
+
+		// FBlueprintMetadata::MD_FunctionCategory
+		static const FName NAME_FunctionCategory(TEXT("Category"));
+
+		FPropertyFunctionCallDelegates::FOnExecute OnExecute = InArgs.OnExecute;
+		FPropertyFunctionCallDelegates::FOnCanExecute OnCanExecute = InArgs.OnCanExecute;
+
+		if (!OnExecute.IsBound() && InArgs.OnGetExecutionContext.IsBound())
+		{
+			// FBlueprintMetadata::MD_WorldContext
+			static const FName NAME_WorldContext(TEXT("WorldContext"));
+
+			OnExecute = FPropertyFunctionCallDelegates::FOnExecute::CreateLambda(
+				[OnGetExecutionContext = InArgs.OnGetExecutionContext](const TWeakObjectPtr<UFunction>& InWeakFunction)
+				{
+					TArray<TWeakObjectPtr<UObject>> WeakExecutionObjects = OnGetExecutionContext.Execute(InWeakFunction);
+					if (!WeakExecutionObjects.IsEmpty())
+					{
+						using namespace UE::Reflection;
+						if (UFunction* Function = InWeakFunction.Get())
+						{
+							// @todo: Consider naming the transaction scope after the fully qualified function name for better UX
+							FScopedTransaction Transaction(LOCTEXT("ExecuteCallInEditorMethod", "Call In Editor Action"));
+							TStrongObjectPtr<UFunction> CallingFunction(Function);
+
+							if (Function->HasMetaData(NAME_WorldContext)
+								&& DoesStaticFunctionSignatureMatch<void(TObjectPtr<UObject>)>(Function))
+							{
+								FEditorScriptExecutionGuard ScriptGuard;
+								UEditorEngine* EditorEngine = Cast<UEditorEngine>(GEngine);
+								UObject* WorldContextObject = EditorEngine->GetEditorWorldContext().World();
+								TStrongObjectPtr<UObject> CDO(Function->GetOwnerClass()->ClassDefaultObject);
+								CDO->ProcessEvent(Function, &WorldContextObject);
+							}
+							else
+							{
+								FEditorScriptExecutionGuard ScriptGuard;
+								for (const TWeakObjectPtr<UObject>& WeakExecutionObject : WeakExecutionObjects)
+								{
+									if (UObject* ExecutionObject = WeakExecutionObject.Get())
+									{
+										ensure(Function->ParmsSize == 0);
+										TStrongObjectPtr<UObject> StrongExecutionObject(ExecutionObject); // Prevent GC during call
+										ExecutionObject->ProcessEvent(Function, nullptr);
+									}
+								}
+							}
+						}
+					}
+
+					return FReply::Handled();
+				});
+		}
+
+		for (UFunction* Function : InCallInEditorFunctions)
+		{
+			if (!Function)
+			{
+				continue;
+			}
+
+			FName FunctionCategoryName(NAME_Default);
+			if (Function->HasMetaData(NAME_FunctionCategory))
+			{
+				FunctionCategoryName = FName(*Function->GetMetaData(NAME_FunctionCategory));
+			}
+
+			if (FunctionCategoryName != ActiveCategory)
+			{
+				ActiveCategory = FunctionCategoryName;
+				OutCategorizedEntries.Emplace(FunctionCategoryName);
+			}
+
+			FCategorizedFunctionCallEntry& CategoryEntry = OutCategorizedEntries.Last();
+
+			FText ButtonLabel = ObjectTools::GetUserFacingFunctionName(Function);
+			FText ButtonToolTip = Function->GetToolTipText();
+			if (ButtonToolTip.IsEmpty())
+			{
+				ButtonToolTip = ButtonLabel;
+			}
+
+			CategoryEntry.WrapBox->AddSlot()
+			.Padding(0.0f, 0.0f, 5.0f, 3.0f)
+			[
+				PropertyCustomizationHelpers::MakeFunctionCallButton(
+					FPropertyFunctionCallArgs(
+						Function,
+						OnExecute,
+						OnCanExecute,
+						ButtonLabel,
+						ButtonToolTip,
+						&CategoryEntry.FunctionSearchText)
+					)
+			];
+
+			CategoryEntry.RowTag = Function->GetFName();
+		}
+	}
+
+	void AddFunctionCallWidgets(IDetailGroup& RootGroup, const TArrayView<UFunction*>& InCallInEditorFunctions, const FPropertyFunctionCallDelegates& InArgs)
+	{
+		TArray<FCategorizedFunctionCallEntry, TInlineAllocator<8>> CategorizedEntries;
+		GetFunctionCallWidgets(InCallInEditorFunctions, InArgs, CategorizedEntries);
+
+		TMap<FName, IDetailGroup*> Groups;
+
+		// Now edit the categories, adding the button strips to the details panel
+		for (FCategorizedFunctionCallEntry& CategoryEntry : CategorizedEntries)
+		{
+			IDetailGroup* Group = nullptr;
+			if (CategoryEntry.CategoryName == NAME_Default)
+			{
+				Group = &RootGroup;
+			}
+			else if (IDetailGroup** ExistingGroup = Groups.Find(CategoryEntry.CategoryName))
+			{
+				Group = *ExistingGroup;
+			}
+			else
+			{
+				Group = Groups.Emplace(
+					CategoryEntry.CategoryName,
+					&RootGroup.AddGroup(
+						CategoryEntry.CategoryName,
+						FText::FromName(CategoryEntry.CategoryName)));
+			}
+
+			Group->AddWidgetRow()
+			.FilterString(CategoryEntry.FunctionSearchText.ToText())
+			.ShouldAutoExpand(true)
+			.RowTag(CategoryEntry.RowTag)
+			[
+				CategoryEntry.WrapBox.ToSharedRef()
+			];
+		}
+	}
+
+	void AddFunctionCallWidgets(IDetailLayoutBuilder& DetailBuilder, const TArrayView<UFunction*>& InCallInEditorFunctions, const FPropertyFunctionCallDelegates& InArgs)
+	{
+		TArray<FCategorizedFunctionCallEntry, TInlineAllocator<8>> CategorizedEntries;
+		GetFunctionCallWidgets(InCallInEditorFunctions, InArgs, CategorizedEntries);
+
+		// Now edit the categories, adding the button strips to the details panel
+		for (FCategorizedFunctionCallEntry& CategoryEntry : CategorizedEntries)
+		{
+			IDetailCategoryBuilder& CategoryBuilder = DetailBuilder.EditCategory(CategoryEntry.CategoryName);
+			CategoryBuilder.AddCustomRow(CategoryEntry.FunctionSearchText.ToText())
+			.RowTag(CategoryEntry.RowTag)
+			[
+				CategoryEntry.WrapBox.ToSharedRef()
+			];
+		}
+	}
+
+	void AddCallInEditorFunctionCallWidgetsForClass(IDetailGroup& RootGroup, const UClass* Class, const FPropertyFunctionCallDelegates& InArgs)
+	{
+		TArray<UFunction*> CallInEditorFunctions;
+		PropertyCustomizationHelpers::GetCallInEditorFunctionsForClass(
+			Class,
+			CallInEditorFunctions);
+
+		if (!CallInEditorFunctions.IsEmpty())
+		{
+			AddFunctionCallWidgets(RootGroup, CallInEditorFunctions, InArgs);
+		}
+	}
+
+	void AddCallInEditorFunctionCallWidgetsForClass(IDetailLayoutBuilder& DetailBuilder, const UClass* Class, const FPropertyFunctionCallDelegates& InArgs)
+	{
+		TArray<UFunction*> CallInEditorFunctions;
+		PropertyCustomizationHelpers::GetCallInEditorFunctionsForClass(
+			Class,
+			CallInEditorFunctions);
+
+		if (!CallInEditorFunctions.IsEmpty())
+		{
+			AddFunctionCallWidgets(DetailBuilder, CallInEditorFunctions, InArgs);
+		}
+	}
+
 	FText GetVisibilityDisplay(TAttribute<bool> bEnabled)
 	{
 		return bEnabled.Get() ? FEditorFontGlyphs::Eye : FEditorFontGlyphs::Eye_Slash;
@@ -923,6 +1196,129 @@ TArray<const UClass*> PropertyCustomizationHelpers::GetClassesFromMetadataString
 	return Classes;
 }
 
+namespace PropertyCustomizationHelpers
+{
+	namespace Private
+	{
+		void GetCallInEditorFunctionsForClassInternal(const UClass* InClass, TOptional<TFunctionRef<bool(const UFunction*)>> InFunctionFilter, TArray<UFunction*>& OutCallInEditorFunctions, EFieldIterationFlags InIterationFlags)
+		{
+			// metadata tag for defining sort order of function buttons within a Category
+			static const FName NAME_DisplayPriority("DisplayPriority");
+
+			// FBlueprintMetadata::MD_CallInEditor
+			static const FName NAME_CallInEditor(TEXT("CallInEditor"));
+
+			const bool bDisallowEditorUtilityBlueprintFunctions = GetDefault<UBlueprintEditorProjectSettings>()->bDisallowEditorUtilityBlueprintFunctionsInDetailsView;
+
+			TFunction<bool(const UFunction*)> CanDisplayAndCallFunction =
+				[bDisallowEditorUtilityBlueprintFunctions](const UFunction* TestFunction)
+				{
+					bool bCanCall = TestFunction->GetBoolMetaData(NAME_CallInEditor)
+								&& TestFunction->ParmsSize == 0; // Params not supported
+
+					if (bCanCall)
+					{
+						if (const UClass* TestFunctionOwnerClass = TestFunction->GetOwnerClass())
+						{
+							if (const UBlueprint* Blueprint = Cast<UBlueprint>(TestFunctionOwnerClass->ClassGeneratedBy))
+							{
+								if (FBlueprintEditorUtils::IsEditorUtilityBlueprint(Blueprint))
+								{
+									// Skip Blutilities if disabled via project settings
+									bCanCall = !bDisallowEditorUtilityBlueprintFunctions;
+								}
+							}
+						}
+					}
+
+					return bCanCall;
+				};
+
+			if (InFunctionFilter.IsSet())
+			{
+				// If the user provided a filter, we call it if the default (OuterFunc) passes - it should always be CallInEditor and have no params
+				CanDisplayAndCallFunction = [OuterFunc = CanDisplayAndCallFunction, InnerFunc = InFunctionFilter.GetValue()](const UFunction* TestFunction)
+				{
+					return OuterFunc(TestFunction)
+						? InnerFunc(TestFunction)
+						: false;
+				};
+			}
+
+			// Get all of the functions we need to display (done ahead of time so we can sort them)
+			for (TFieldIterator<UFunction> FunctionIter(InClass, InIterationFlags); FunctionIter; ++FunctionIter)
+			{
+				const UFunction* TestFunction = *FunctionIter;
+				if (CanDisplayAndCallFunction(TestFunction))
+				{
+					const FName FunctionName = TestFunction->GetFName();
+
+					if (const bool bFunctionIsPermissible = FPropertyEditorPermissionList::Get().DoesPropertyPassFilter(TestFunction->GetOwnerClass(), FunctionName);
+						!bFunctionIsPermissible)
+					{
+						continue;
+					}
+
+					if (const bool bFunctionAlreadyAdded = OutCallInEditorFunctions.ContainsByPredicate([&FunctionName](UFunction*& Func) { return Func->GetFName() == FunctionName; });
+						bFunctionAlreadyAdded)
+					{
+						continue;
+					}
+
+					OutCallInEditorFunctions.Add(*FunctionIter);
+				}
+			}
+
+			if (OutCallInEditorFunctions.IsEmpty())
+			{
+				return;
+			}
+
+
+			// FBlueprintMetadata::MD_FunctionCategory
+			static const FName NAME_FunctionCategory(TEXT("Category"));
+
+			// Sort the functions by category and then by DisplayPriority meta tag, and then by name
+			OutCallInEditorFunctions.Sort([](const UFunction& A, const UFunction& B)
+			{
+				const int32 CategorySort = A.GetMetaData(NAME_FunctionCategory).Compare(B.GetMetaData(NAME_FunctionCategory));
+				if (CategorySort != 0)
+				{
+					return (CategorySort <= 0);
+				}
+				else
+				{
+					const FString DisplayPriorityAStr = A.GetMetaData(NAME_DisplayPriority);
+					int32 DisplayPriorityA = (DisplayPriorityAStr.IsEmpty() ? MAX_int32 : FCString::Atoi(*DisplayPriorityAStr));
+					if (DisplayPriorityA == 0 && !FCString::IsNumeric(*DisplayPriorityAStr))
+					{
+						DisplayPriorityA = MAX_int32;
+					}
+
+					const FString DisplayPriorityBStr = B.GetMetaData(NAME_DisplayPriority);
+					int32 DisplayPriorityB = (DisplayPriorityBStr.IsEmpty() ? MAX_int32 : FCString::Atoi(*DisplayPriorityBStr));
+					if (DisplayPriorityB == 0 && !FCString::IsNumeric(*DisplayPriorityBStr))
+					{
+						DisplayPriorityB = MAX_int32;
+					}
+
+					return (DisplayPriorityA == DisplayPriorityB) ? (A.GetName() <= B.GetName()) : (DisplayPriorityA <= DisplayPriorityB);
+				}
+			});
+		}
+	}
+
+	void GetCallInEditorFunctionsForClass(const UClass* InClass, TArray<UFunction*>& OutCallInEditorFunctions, EFieldIterationFlags InIterationFlags)
+	{
+		return Private::GetCallInEditorFunctionsForClassInternal(InClass, {}, OutCallInEditorFunctions, InIterationFlags);
+	}
+
+	void GetCallInEditorFunctionsForClass(const UClass* InClass, const TFunctionRef<bool(const UFunction*)>& InFunctionFilter,	TArray<UFunction*>& OutCallInEditorFunctions, EFieldIterationFlags InIterationFlags)
+	{
+		return Private::GetCallInEditorFunctionsForClassInternal(InClass, InFunctionFilter, OutCallInEditorFunctions, InIterationFlags);
+	}
+}
+
 //////////////////////////////////////////////////////////////////////////
 //
 // Sections list
diff --git a/Engine/Source/Editor/PropertyEditor/Public/PropertyCustomizationHelpers.h b/Engine/Source/Editor/PropertyEditor/Public/PropertyCustomizationHelpers.h
--- a/Engine/Source/Editor/PropertyEditor/Public/PropertyCustomizationHelpers.h
+++ b/Engine/Source/Editor/PropertyEditor/Public/PropertyCustomizationHelpers.h
@@ -89,6 +89,65 @@ struct FPropertyComboBoxArgs
 	{}
 };
 
+struct FPropertyFunctionCallArgs
+{
+	TWeakObjectPtr<UFunction> Function;
+
+	TOptional<FText> LabelOverride;
+
+	TOptional<FText> ToolTipTextOverride;
+
+	using FOnExecute = TDelegate<FReply(TWeakObjectPtr<UFunction>)>;
+	FOnExecute OnExecute;
+
+	using FOnCanExecute = TDelegate<bool(TWeakObjectPtr<UFunction>)>;
+	FOnCanExecute OnCanExecute;
+
+	FTextBuilder* SearchText = nullptr;
+
+	FPropertyFunctionCallArgs(
+		UFunction* InFunction,
+		const FOnExecute& InOnExecute,
+		const FOnCanExecute& InOnCanExecute = {},
+		const TOptional<FText>& InLabelOverride = {},
+		const TOptional<FText>& InToolTipTextOverride = {},
+		FTextBuilder* InSearchText = nullptr)
+		: Function(InFunction)
+		, LabelOverride(InLabelOverride)
+		, ToolTipTextOverride(InToolTipTextOverride)
+		, OnExecute(InOnExecute)
+		, OnCanExecute(InOnCanExecute)
+		, SearchText(InSearchText)
+	{
+	}
+};
+
+/** The callbacks, if specified, are used when invoking function calls. */
+struct FPropertyFunctionCallDelegates
+{
+	using FOnGetExecutionContext = TDelegate<TArray<TWeakObjectPtr<UObject>>(TWeakObjectPtr<UFunction>)>;
+	FOnGetExecutionContext OnGetExecutionContext;
+
+	using FOnExecute = TDelegate<FReply(TWeakObjectPtr<UFunction>)>;
+	FOnExecute OnExecute;
+
+	using FOnCanExecute = TDelegate<bool(TWeakObjectPtr<UFunction>)>;
+	FOnCanExecute OnCanExecute;
+
+	FPropertyFunctionCallDelegates(
+		const FOnExecute& InOnExecute,
+		const FOnCanExecute& InOnCanExecute = {})
+		: OnExecute(InOnExecute)
+		, OnCanExecute(InOnCanExecute)
+	{
+	}
+
+	FPropertyFunctionCallDelegates(const FOnGetExecutionContext& InOnGetExecutionContext)
+		: OnGetExecutionContext(InOnGetExecutionContext)
+	{
+	}
+};
+
 namespace PropertyCustomizationHelpers
 {
 	PROPERTYEDITOR_API TSharedRef<SWidget> MakeResetButton(FSimpleDelegate OnResetClicked, TAttribute<FText> OptionalToolTipText = FText(), TAttribute<bool> IsEnabled = true);
@@ -119,6 +178,11 @@ namespace PropertyCustomizationHelpers
 	PROPERTYEDITOR_API TSharedRef<SWidget> MakeSaveButton(FSimpleDelegate OnSaveClicked, TAttribute<FText> OptionalToolTipText = FText(), TAttribute<bool> IsEnabled = true);
 	PROPERTYEDITOR_API TSharedRef<SWidget> MakeSetOptionalButton(FSimpleDelegate OnSetOptionalClicked, TAttribute<FText> OptionalToolTipText = FText(), TAttribute<bool> IsEnabled = true);
 	PROPERTYEDITOR_API TSharedRef<SWidget> MakeClearOptionalButton(FSimpleDelegate OnClearOptionalClicked, TAttribute<FText> OptionalToolTipText = FText(), TAttribute<bool> IsEnabled = true);
+	PROPERTYEDITOR_API TSharedRef<SWidget> MakeFunctionCallButton(const FPropertyFunctionCallArgs& InArgs);
+	PROPERTYEDITOR_API void AddFunctionCallWidgets(IDetailGroup& RootGroup, const TArrayView<UFunction*>& InCallInEditorFunctions, const FPropertyFunctionCallDelegates& InArgs);
+	PROPERTYEDITOR_API void AddFunctionCallWidgets(IDetailLayoutBuilder& DetailBuilder, const TArrayView<UFunction*>& InCallInEditorFunctions, const FPropertyFunctionCallDelegates& InArgs);
+	PROPERTYEDITOR_API void AddCallInEditorFunctionCallWidgetsForClass(IDetailGroup& RootGroup, const UClass* Class, const FPropertyFunctionCallDelegates& InArgs);
+	PROPERTYEDITOR_API void AddCallInEditorFunctionCallWidgetsForClass(IDetailLayoutBuilder& DetailBuilder, const UClass* Class, const FPropertyFunctionCallDelegates& InArgs);
 
 	/** @return the FBoolProperty edit condition property if one exists. */
 	PROPERTYEDITOR_API FBoolProperty* GetEditConditionProperty(const FProperty* InProperty, bool& bNegate);
@@ -161,6 +225,16 @@ namespace PropertyCustomizationHelpers
 	 * The metadata string is likely from something like AllowedClasses or DisallowedClasses.
 	 */
 	PROPERTYEDITOR_API TArray<const UClass*> GetClassesFromMetadataString(const FString& MetadataString);
+
+	/**
+	 *
+	 */
+	PROPERTYEDITOR_API void GetCallInEditorFunctionsForClass(const UClass* InClass, TArray<UFunction*>& OutCallInEditorFunctions, EFieldIterationFlags InIterationFlags = EFieldIterationFlags::IncludeSuper);
+
+	/**
+	 *	@param InFunctionFilter	A filter to select candidate UFunctions
+	 */
+	PROPERTYEDITOR_API void GetCallInEditorFunctionsForClass(const UClass* InClass, const TFunctionRef<bool(const UFunction*)>& InFunctionFilter, TArray<UFunction*>& OutCallInEditorFunctions, EFieldIterationFlags InIterationFlags = EFieldIterationFlags::IncludeSuper);
 }
 
 
