Engine (array, map, and set add item);

Fixing an issue with the previous TMap and TSet auto increment implementaiton, where if you remove an element in the middle of the array then adding a new entry would incorrectly grab the last index of the container to adjust it. When adding an entry into a container, it will use the first free index (head of the linked list) instead of just adding to the end. Now we'll return the added logical index (meaning, the valid element index and not the internal index) into the container so calling code can get the correct entry if needed.
NOTE: There is still a side effect of how we use the first free index of the container, which is really the last removed element of the container. This means that if we then add an element from an enum, it could be in a position that isn't sequential visually. For example, you add all entries with an enum key from 0-5, those keys will match the container index [0] [1] [2] [3] [4] [5]. You remove index 2 and 4 from the container [0] [1] [3] [5]. Now when you add, it'll add the new element to index 4 but the enum would be 2 [0] [1] [3] [2] [5].

-- Code --
PropertyHandleImpl:
FPropertyValueImpl::AddChild() - Added an int32 return type so we can return the logical index of the added child.
FPropertyHandleSet::AddItem() - Using the child index from AddChild() instead of grabbing the last element on the container.
FPropertyHandleMap::AddItem() - Using the child index from AddChild() instead of grabbing the last element on the container.

#tests Editor, used test container properties with an enum as the key. Added, removed, and added, to make sure entries were added as expected to the start/end of the container or middle of container.

diff --git a/Engine/Source/Editor/PropertyEditor/Private/PropertyHandleImpl.cpp b/Engine/Source/Editor/PropertyEditor/Private/PropertyHandleImpl.cpp
--- a/Engine/Source/Editor/PropertyEditor/Private/PropertyHandleImpl.cpp
+++ b/Engine/Source/Editor/PropertyEditor/Private/PropertyHandleImpl.cpp
@@ -1047,8 +1047,10 @@ FText FPropertyValueImpl::GetResetToDefaultLabel() const
 	return FText::GetEmpty();
 }
 
-void FPropertyValueImpl::AddChild()
+int32 FPropertyValueImpl::AddChild()
 {
+	int32 ReturnLogicalIndex = INDEX_NONE;
+
 	TSharedPtr<FPropertyNode> PropertyNodePin = PropertyNode.Pin();
 	if ( PropertyNodePin.IsValid() )
 	{
@@ -1135,10 +1137,11 @@ void FPropertyValueImpl::AddChild()
 
 						if (Array)
 						{
-							Array->PerformOperationWithSetter(Obj, Addr, [Obj, Array, &Index](void* DirectAddress)
-							{								
+							Array->PerformOperationWithSetter(Obj, Addr, [Obj, Array, &Index, &ReturnLogicalIndex](void* DirectAddress)
+							{
 								FScriptArrayHelper	ArrayHelper(Array, DirectAddress);
 								Index = ArrayHelper.AddValue();
+								ReturnLogicalIndex = Index;
 
 								// check whether the inner type is flagged as a non-nullable. if so, create it.
 								FObjectProperty* InnerObjectProperty = CastField<FObjectProperty>(Array->Inner);
@@ -1151,10 +1154,11 @@ void FPropertyValueImpl::AddChild()
 						}
 						else if (Set)
 						{
-							Set->PerformOperationWithSetter(Obj, Addr, [Obj, Set, &Index](void* DirectAddress)
+							Set->PerformOperationWithSetter(Obj, Addr, [Obj, Set, &Index, &ReturnLogicalIndex](void* DirectAddress)
 							{
 								FScriptSetHelper	SetHelper(Set, DirectAddress);
 								Index = SetHelper.AddDefaultValue_Invalid_NeedsRehash();
+								ReturnLogicalIndex = SetHelper.FindLogicalIndex(Index);
 
 								// check whether the element type is flagged as a non-nullable. if so, create it.
 								FObjectProperty* ElementObjectProperty = CastField<FObjectProperty>(Set->ElementProp);
@@ -1169,10 +1173,11 @@ void FPropertyValueImpl::AddChild()
 						}
 						else if (Map)
 						{
-							Map->PerformOperationWithSetter(Obj, Addr, [Obj, Map, &Index, &bAddedMapEntry](void* DirectAddress)
+							Map->PerformOperationWithSetter(Obj, Addr, [Obj, Map, &Index, &bAddedMapEntry, &ReturnLogicalIndex](void* DirectAddress)
 							{
 								FScriptMapHelper	MapHelper(Map, DirectAddress);
 								Index = MapHelper.AddDefaultValue_Invalid_NeedsRehash();
+								ReturnLogicalIndex = MapHelper.FindLogicalIndex(Index);
 
 								// check whether the key or value type is flagged as a non-nullable. if so, create it.
 								{
@@ -1222,6 +1227,8 @@ void FPropertyValueImpl::AddChild()
 			}
 		}
 	}
+
+	return ReturnLogicalIndex;
 }
 
 void FPropertyValueImpl::ClearChildren()
@@ -5585,11 +5592,8 @@ FPropertyAccess::Result FPropertyHandleSet::AddItem()
 							}
 
 							// If we don't have this element then add an entry and set it to this element value.
-							Implementation->AddChild();
+							const int32 ChildNodeIndex = Implementation->AddChild();
 							Implementation->GetPropertyNode()->RebuildChildren();
-
-							// Grab the last entry since we just added it.
-							const int32 ChildNodeIndex = PropNode->GetNumChildNodes() - 1;
 							if (ChildNodeIndex >= 0)
 							{
 								TSharedPtr<FPropertyNode> ChildNode = Implementation->GetChildNode(ChildNodeIndex);
@@ -5836,10 +5840,7 @@ FPropertyAccess::Result FPropertyHandleMap::AddItem()
 							}
 
 							// If we don't have this key then add an entry and set it to this key value.
-							Implementation->AddChild();
-
-							// Grab the last entry since we just added it.
-							const int32 ChildNodeIndex = PropNode->GetNumChildNodes() - 1;
+							const int32 ChildNodeIndex = Implementation->AddChild();
 							if (ChildNodeIndex >= 0)
 							{
 								TSharedPtr<FPropertyNode> ChildNode = Implementation->GetChildNode(ChildNodeIndex);
diff --git a/Engine/Source/Editor/PropertyEditor/Private/PropertyHandleImpl.h b/Engine/Source/Editor/PropertyEditor/Private/PropertyHandleImpl.h
--- a/Engine/Source/Editor/PropertyEditor/Private/PropertyHandleImpl.h
+++ b/Engine/Source/Editor/PropertyEditor/Private/PropertyHandleImpl.h
@@ -287,8 +287,9 @@ public:
 
 	/**
 	 * Adds a child to the property node (container properties only)
+	 * @return The logical index of the added child.
 	 */
-	void AddChild();
+	int32 AddChild();
 
 	/**
 	 * Removes all children from the property node (container properties only)
