Fixed mistake spotted by static analysis

diff --git a/Engine/Source/Editor/PropertyEditor/Private/DetailLayoutBuilderImpl.cpp b/Engine/Source/Editor/PropertyEditor/Private/DetailLayoutBuilderImpl.cpp
--- a/Engine/Source/Editor/PropertyEditor/Private/DetailLayoutBuilderImpl.cpp
+++ b/Engine/Source/Editor/PropertyEditor/Private/DetailLayoutBuilderImpl.cpp
@@ -207,7 +207,7 @@ IDetailPropertyRow* FDetailLayoutBuilderImpl::EditDefaultProperty(TSharedPtr<IPr
 			{
 				Category = CustomCategoryMap.FindRef(CategoryFName);
 			}
-			else if (!Category.IsValid())
+			if (!Category.IsValid())
 			{
 				Category = SubCategoryMap.FindRef(CategoryFName);
 			}
