Engine change to address an issue with copying array of string values. When you copy an array of strings that pasted strings each include "" when pasted.

diff --git a/Engine/Source/Editor/PropertyEditor/Private/PropertyHandleImpl.cpp b/Engine/Source/Editor/PropertyEditor/Private/PropertyHandleImpl.cpp
--- a/Engine/Source/Editor/PropertyEditor/Private/PropertyHandleImpl.cpp
+++ b/Engine/Source/Editor/PropertyEditor/Private/PropertyHandleImpl.cpp
@@ -5171,7 +5171,11 @@ FPropertyAccess::Result FPropertyHandleArray::SetValueFromFormattedString(const
 			
 			const TSharedRef<IPropertyHandle> Property = GetElement( NumElements - 1 );
 			ensure(!Property->AsStruct());
-			if (Property->IsValidHandle() && Property->SetValueFromFormattedString(Value, Flags) == FPropertyAccess::Fail)
+
+			// Trim ""
+			FString FormattedString = Value.TrimQuotes();
+			
+			if (Property->IsValidHandle() && Property->SetValueFromFormattedString(FormattedString, Flags) == FPropertyAccess::Fail)
 			{
 				Result = FPropertyAccess::Fail;
 			}
