Property Editor: Fix Struct Detail View's scrollbar overlapped the Revert to Default buttons

diff --git a/Engine/Source/Editor/PropertyEditor/Private/SStructureDetailsView.cpp b/Engine/Source/Editor/PropertyEditor/Private/SStructureDetailsView.cpp
--- a/Engine/Source/Editor/PropertyEditor/Private/SStructureDetailsView.cpp
+++ b/Engine/Source/Editor/PropertyEditor/Private/SStructureDetailsView.cpp
@@ -172,6 +172,8 @@ void SStructureDetailsView::Construct(const FArguments& InArgs)
 		.SelectionMode(ESelectionMode::None)
 		.ExternalScrollbar(ExternalScrollbar);
 
+	constexpr float ScrollbarWidth = 16.0f;
+
 	ChildSlot
 	[
 		SNew( SBox )
@@ -190,6 +192,15 @@ void SStructureDetailsView::Construct(const FArguments& InArgs)
 			[
 				SNew( SOverlay )
 				+ SOverlay::Slot()
+				.Padding(TAttribute<FMargin>::CreateLambda([ExternalScrollbar]()
+				{
+					if (ExternalScrollbar->GetVisibility().IsVisible())
+					{
+						return FMargin(0.0f, 0.0f, ScrollbarWidth, 0.0f);
+					}
+
+					return FMargin();
+				}))
 				[
 					DetailTree.ToSharedRef()
 				]
@@ -197,7 +208,7 @@ void SStructureDetailsView::Construct(const FArguments& InArgs)
 				.HAlign(HAlign_Right)
 				[
 					SNew( SBox )
-					.WidthOverride( 16.0f )
+					.WidthOverride( ScrollbarWidth )
 					[
 						ExternalScrollbar
 					]
