SPropertyEditorInline: Fix refresh issue with editing inline/instanced properties.

diff --git a/Engine/Source/Editor/PropertyEditor/Private/UserInterface/PropertyEditor/SPropertyEditorEditInline.cpp b/Engine/Source/Editor/PropertyEditor/Private/UserInterface/PropertyEditor/SPropertyEditorEditInline.cpp
--- a/Engine/Source/Editor/PropertyEditor/Private/UserInterface/PropertyEditor/SPropertyEditorEditInline.cpp
+++ b/Engine/Source/Editor/PropertyEditor/Private/UserInterface/PropertyEditor/SPropertyEditorEditInline.cpp
@@ -221,6 +221,7 @@ const FSlateBrush* SPropertyEditorEditInline::GetDisplayValueIcon() const
 void SPropertyEditorEditInline::OnClassPickedInternal(UClass* InClass, TSharedRef<IPropertyHandle> PropertyHandle)
 {
 	SPropertyEditorEditInline::OnClassPicked(InClass, PropertyHandle);
+	PropertyEditor->ForceRefresh();
 	ComboButton->SetIsOpen(false);
 }
 
