Fix categories not working with Verse Devices

#rnx

diff --git a/Engine/Source/Editor/PropertyEditor/Private/DetailLayoutBuilderImpl.cpp b/Engine/Source/Editor/PropertyEditor/Private/DetailLayoutBuilderImpl.cpp
--- a/Engine/Source/Editor/PropertyEditor/Private/DetailLayoutBuilderImpl.cpp
+++ b/Engine/Source/Editor/PropertyEditor/Private/DetailLayoutBuilderImpl.cpp
@@ -1119,6 +1119,15 @@ bool FDetailLayoutBuilderImpl::IsPropertyPathAllowed(const FString& InPath) cons
 	return false;
 }
 
+void FDetailLayoutBuilderImpl::DisableInstancedReference(TSharedRef<IPropertyHandle> PropertyHandle) const
+{
+	TSharedPtr<FPropertyNode> PropertyNode = GetPropertyNode(PropertyHandle);
+	if (PropertyNode.IsValid())
+	{
+		PropertyNode->SetIgnoreInstancedReference();
+	}
+}
+
 bool FDetailLayoutBuilderImpl::AddEmptyCategoryIfNeeded(TSharedPtr<FComplexPropertyNode> Node)
 {
 	const bool bHasNoValidCategories = DefaultCategoryMap.IsEmpty();
diff --git a/Engine/Source/Editor/PropertyEditor/Private/DetailLayoutBuilderImpl.h b/Engine/Source/Editor/PropertyEditor/Private/DetailLayoutBuilderImpl.h
--- a/Engine/Source/Editor/PropertyEditor/Private/DetailLayoutBuilderImpl.h
+++ b/Engine/Source/Editor/PropertyEditor/Private/DetailLayoutBuilderImpl.h
@@ -58,6 +58,7 @@ public:
 	virtual void SortCategories(const FOnCategorySortOrderFunction& SortFunction) override;
 	virtual void SetPropertyGenerationAllowListPaths(const TSet<FString>& InPropertyGenerationAllowListPaths) override;
 	virtual bool IsPropertyPathAllowed(const FString& InPath) const override;
+	virtual void DisableInstancedReference(TSharedRef<IPropertyHandle> PropertyHandle) const override;
 
 	/**
 	 * Creates an empty category row if there currently are no categories and one is required for the
diff --git a/Engine/Source/Editor/PropertyEditor/Public/DetailLayoutBuilder.h b/Engine/Source/Editor/PropertyEditor/Public/DetailLayoutBuilder.h
--- a/Engine/Source/Editor/PropertyEditor/Public/DetailLayoutBuilder.h
+++ b/Engine/Source/Editor/PropertyEditor/Public/DetailLayoutBuilder.h
@@ -309,6 +309,11 @@ public:
 	 * @return True if the property path is contained within our allowed paths
 	 */
 	virtual bool IsPropertyPathAllowed(const FString& InPath) const = 0;
+
+	/**
+	 * Force a property to behave as a normal, peer reference regardless of CPF_InstancedReference
+	 */
+	virtual void DisableInstancedReference(TSharedRef<IPropertyHandle> PropertyHandle) const = 0;
 };
 
 template <typename ObjectType>
