Property Editor: Add support for StructProvider-driven Detail Views

#rnx

diff --git a/Engine/Source/Editor/PropertyEditor/Private/PropertyEditorModule.cpp b/Engine/Source/Editor/PropertyEditor/Private/PropertyEditorModule.cpp
--- a/Engine/Source/Editor/PropertyEditor/Private/PropertyEditorModule.cpp
+++ b/Engine/Source/Editor/PropertyEditor/Private/PropertyEditorModule.cpp
@@ -1104,7 +1104,7 @@ FPropertyTypeLayoutCallback FPropertyEditorModule::FindPropertyTypeLayoutCallbac
 	return FPropertyTypeLayoutCallback();
 }
 
-TSharedRef<class IStructureDetailsView> FPropertyEditorModule::CreateStructureDetailView(const FDetailsViewArgs& DetailsViewArgs, const FStructureDetailsViewArgs& StructureDetailsViewArgs, TSharedPtr<FStructOnScope> StructData, const FText& CustomName)
+TSharedRef<class IStructureDetailsView> FPropertyEditorModule::CreateStructureDetailView(const FDetailsViewArgs& DetailsViewArgs, const FStructureDetailsViewArgs& StructureDetailsViewArgs, const FText& CustomName)
 {
 	TSharedRef<SStructureDetailsView> DetailView =
 		SNew(SStructureDetailsView)
@@ -1185,8 +1185,25 @@ TSharedRef<class IStructureDetailsView> FPropertyEditorModule::CreateStructureDe
 	{
 		DetailView->SetIsPropertyVisibleDelegate(FIsPropertyVisible::CreateStatic(&FStructureDetailsViewFilter::PassesFilter, StructureDetailsViewArgs));
 	}
+
+	return DetailView;
+}
+
+TSharedRef<IStructureDetailsView> FPropertyEditorModule::CreateStructureDetailView(
+	const FDetailsViewArgs& DetailsViewArgs, const FStructureDetailsViewArgs& StructureDetailsViewArgs,
+	TSharedPtr<FStructOnScope> StructData, const FText& CustomName)
+{
+	TSharedRef<IStructureDetailsView> DetailView = CreateStructureDetailView(DetailsViewArgs, StructureDetailsViewArgs, CustomName);
 	DetailView->SetStructureData(StructData);
+	return DetailView;
+}
 
+TSharedRef<IStructureDetailsView> FPropertyEditorModule::CreateStructureProviderDetailView(
+	const FDetailsViewArgs& DetailsViewArgs, const FStructureDetailsViewArgs& StructureDetailsViewArgs,
+	TSharedPtr<IStructureDataProvider> StructProvider, const FText& CustomName)
+{
+	TSharedRef<IStructureDetailsView> DetailView = CreateStructureDetailView(DetailsViewArgs, StructureDetailsViewArgs, CustomName);
+	DetailView->SetStructureProvider(StructProvider);
 	return DetailView;
 }
 
diff --git a/Engine/Source/Editor/PropertyEditor/Public/PropertyEditorModule.h b/Engine/Source/Editor/PropertyEditor/Public/PropertyEditorModule.h
--- a/Engine/Source/Editor/PropertyEditor/Public/PropertyEditorModule.h
+++ b/Engine/Source/Editor/PropertyEditor/Public/PropertyEditorModule.h
@@ -395,6 +395,7 @@ public:
 	virtual TSharedPtr<class ISinglePropertyView> CreateSingleProperty(const TSharedPtr<class IStructureDataProvider>& InStruct, FName InPropertyName, const struct FSinglePropertyParams& InitParams);
 
 	virtual TSharedRef<class IStructureDetailsView> CreateStructureDetailView(const struct FDetailsViewArgs& DetailsViewArgs, const FStructureDetailsViewArgs& StructureDetailsViewArgs, TSharedPtr<class FStructOnScope> StructData, const FText& CustomName = FText::GetEmpty());
+	virtual TSharedRef<class IStructureDetailsView> CreateStructureProviderDetailView(const FDetailsViewArgs& DetailsViewArgs, const FStructureDetailsViewArgs& StructureDetailsViewArgs, TSharedPtr<IStructureDataProvider> StructProvider, const FText& CustomName = FText::GetEmpty());
 
 	virtual TSharedRef<class IPropertyRowGenerator> CreatePropertyRowGenerator(const struct FPropertyRowGeneratorArgs& InArgs);
 
@@ -466,6 +467,8 @@ private:
 	 */
 	virtual TSharedRef<SPropertyTreeViewImpl> CreatePropertyView( UObject* InObject, bool bAllowFavorites, bool bIsLockable, bool bHiddenPropertyVisibility, bool bAllowSearch, bool ShowTopLevelNodes, FNotifyHook* InNotifyHook, float InNameColumnWidth, FOnPropertySelectionChanged OnPropertySelectionChanged, FOnPropertyClicked OnPropertyMiddleClicked, FConstructExternalColumnHeaders ConstructExternalColumnHeaders, FConstructExternalColumnCell ConstructExternalColumnCell );
 
+	virtual TSharedRef<class IStructureDetailsView> CreateStructureDetailView(const struct FDetailsViewArgs& DetailsViewArgs, const FStructureDetailsViewArgs& StructureDetailsViewArgs, const FText& CustomName = FText::GetEmpty());
+
 	TSharedPtr<FAssetThumbnailPool> GetThumbnailPool();
 
 	void GetAllSectionsHelper(const UStruct* Struct, TArray<TSharedPtr<FPropertySection>>& OutSections, TSet<const UStruct*>& ProcessedStructs) const;
