Add Expand to IPropertyHandle
#rnx

diff --git a/Engine/Source/Editor/PropertyEditor/Private/PropertyHandleImpl.cpp b/Engine/Source/Editor/PropertyEditor/Private/PropertyHandleImpl.cpp
--- a/Engine/Source/Editor/PropertyEditor/Private/PropertyHandleImpl.cpp
+++ b/Engine/Source/Editor/PropertyEditor/Private/PropertyHandleImpl.cpp
@@ -1036,6 +1036,28 @@ bool FPropertyValueImpl::IsEditConst() const
 	return false;
 }
 
+bool FPropertyValueImpl::IsExpanded() const
+{
+	TSharedPtr<FPropertyNode> PropertyNodePin = PropertyNode.Pin();
+	if( PropertyNodePin.IsValid()  )
+	{
+		return PropertyNodePin->HasNodeFlags(EPropertyNodeFlags::Expanded);
+	} 
+
+	return false;
+}
+
+void FPropertyValueImpl::SetExpanded(bool bExpanded)
+{
+	TSharedPtr<FPropertyNode> PropertyNodePin = PropertyNode.Pin();
+	if (PropertyNodePin.IsValid() && PropertyNodePin->HasNodeFlags(EPropertyNodeFlags::CanBeExpanded))
+	{
+		PropertyNodePin->SetNodeFlags(EPropertyNodeFlags::Expanded, bExpanded);
+		// if the node is built, then we need to rebuild the child to expand it.
+		PropertyNodePin->RequestRebuildChildren();
+	}
+}
+
 FText FPropertyValueImpl::GetResetToDefaultLabel() const
 {
 	TSharedPtr<FPropertyNode> PropertyNodePin = PropertyNode.Pin();
@@ -2527,6 +2549,16 @@ bool FPropertyHandleBase::IsEditable() const
 	return !IsEditConst();
 }
 
+bool FPropertyHandleBase::IsExpanded() const
+{
+	return Implementation->IsExpanded();
+}
+
+void FPropertyHandleBase::SetExpanded(bool bExpanded)
+{
+	Implementation->SetExpanded(bExpanded);
+}
+
 FPropertyAccess::Result FPropertyHandleBase::GetValueAsFormattedString( FString& OutValue, EPropertyPortFlags PortFlags ) const
 {
 	return Implementation->GetValueAsString(OutValue, PortFlags);
diff --git a/Engine/Source/Editor/PropertyEditor/Private/PropertyHandleImpl.h b/Engine/Source/Editor/PropertyEditor/Private/PropertyHandleImpl.h
--- a/Engine/Source/Editor/PropertyEditor/Private/PropertyHandleImpl.h
+++ b/Engine/Source/Editor/PropertyEditor/Private/PropertyHandleImpl.h
@@ -279,6 +279,16 @@ public:
 	 * @return true if the property  is edit const and cannot be changed
 	 */
 	bool IsEditConst() const;
+	
+	/**
+	 * @return Whether or not the property is expanded
+	 */
+	bool IsExpanded() const;
+
+	/**
+	* Set the expanded flag
+	*/
+	void SetExpanded(bool bExpanded);
 
 	/**
 	 * @return The label to use for displaying reset to default values
@@ -436,6 +446,8 @@ public:
 	virtual void CreateDefaultPropertyCopyPasteActions(FUIAction& OutCopyAction, FUIAction& OutPasteAction) const override;
 	virtual bool IsEditConst() const override;
 	virtual bool IsEditable() const override;
+	virtual bool IsExpanded() const override;
+	virtual void SetExpanded(bool bExpanded) override;
 	virtual void SetOnPropertyValueChanged( const FSimpleDelegate& InOnPropertyValueChanged ) override;
 	virtual void SetOnPropertyValueChangedWithData(const TDelegate<void(const FPropertyChangedEvent&)>& InOnPropertyValueChanged) override;
 	virtual void SetOnChildPropertyValueChanged( const FSimpleDelegate& InOnPropertyValueChanged ) override;
diff --git a/Engine/Source/Editor/PropertyEditor/Public/PropertyHandle.h b/Engine/Source/Editor/PropertyEditor/Public/PropertyHandle.h
--- a/Engine/Source/Editor/PropertyEditor/Public/PropertyHandle.h
+++ b/Engine/Source/Editor/PropertyEditor/Public/PropertyHandle.h
@@ -62,6 +62,16 @@ public:
 	*/
 	virtual bool IsEditable() const = 0;
 
+	/**
+	 * @return Whether or not the property is expanded
+	 */
+	virtual bool IsExpanded() const = 0;
+
+	/**
+	 * Sets the expanded flag for this property.
+	 */
+	virtual void SetExpanded(bool bExpanded) = 0;
+
 	/**
 	 * Gets the class of the property being edited
 	 */
