Property Editor: Added support custom reset to default to IDetailPropertyRow/FDetailWidgetRow

Allows details customizations to override the widget on a per-property basis

[RN] Added support custom reset to default to IDetailPropertyRow/FDetailWidgetRow

diff --git a/Engine/Source/Editor/PropertyEditor/Private/DetailPropertyRow.cpp b/Engine/Source/Editor/PropertyEditor/Private/DetailPropertyRow.cpp
--- a/Engine/Source/Editor/PropertyEditor/Private/DetailPropertyRow.cpp
+++ b/Engine/Source/Editor/PropertyEditor/Private/DetailPropertyRow.cpp
@@ -265,6 +265,11 @@ FDetailWidgetDecl* FDetailPropertyRow::CustomValueWidget()
 	return CustomPropertyWidget.IsValid() ? &CustomPropertyWidget->ValueContent() : nullptr;
 }
 
+FDetailWidgetDecl* FDetailPropertyRow::CustomResetToDefaultWidget()
+{
+	return CustomPropertyWidget.IsValid() ? &CustomPropertyWidget->ResetToDefaultContent() : nullptr;
+}
+
 TSharedPtr<FAssetThumbnailPool> FDetailPropertyRow::GetThumbnailPool() const
 {
 	TSharedPtr<FDetailCategoryImpl> ParentCategoryPinned = ParentCategory.Pin();
@@ -971,6 +976,14 @@ void FDetailPropertyRow::SetWidgetRowProperties(FDetailWidgetRow& Row) const
 			ensureMsgf(!CustomResetToDefault.IsSet(), TEXT("Duplicate reset to default handlers set on both FDetailPropertyRow and CustomWidget()!"));
 			Row.CustomResetToDefault = CustomPropertyWidget->CustomResetToDefault;
 		}
+
+		if(CustomPropertyWidget->HasResetToDefaultContent())
+		{
+			Row.ResetToDefaultContent()
+			[
+				CustomPropertyWidget->ResetToDefaultWidget.Widget
+			];
+		}
 	}
 }
 
diff --git a/Engine/Source/Editor/PropertyEditor/Private/DetailPropertyRow.h b/Engine/Source/Editor/PropertyEditor/Private/DetailPropertyRow.h
--- a/Engine/Source/Editor/PropertyEditor/Private/DetailPropertyRow.h
+++ b/Engine/Source/Editor/PropertyEditor/Private/DetailPropertyRow.h
@@ -41,6 +41,7 @@ public:
 	virtual FDetailWidgetRow& CustomWidget( bool bShowChildren = false ) override;
 	virtual FDetailWidgetDecl* CustomNameWidget() override;
 	virtual FDetailWidgetDecl* CustomValueWidget() override;
+	virtual FDetailWidgetDecl* CustomResetToDefaultWidget() override;
 	virtual void GetDefaultWidgets( TSharedPtr<SWidget>& OutNameWidget, TSharedPtr<SWidget>& OutValueWidget, bool bAddWidgetDecoration = false) override;
 	virtual void GetDefaultWidgets( TSharedPtr<SWidget>& OutNameWidget, TSharedPtr<SWidget>& OutValueWidget, FDetailWidgetRow& Row, bool bAddWidgetDecoration = false) override;
 
diff --git a/Engine/Source/Editor/PropertyEditor/Private/SDetailSingleItemRow.cpp b/Engine/Source/Editor/PropertyEditor/Private/SDetailSingleItemRow.cpp
--- a/Engine/Source/Editor/PropertyEditor/Private/SDetailSingleItemRow.cpp
+++ b/Engine/Source/Editor/PropertyEditor/Private/SDetailSingleItemRow.cpp
@@ -493,11 +493,11 @@ void SDetailSingleItemRow::Construct( const FArguments& InArgs, FDetailLayoutCus
 				if (PropertyNode->IsReorderable() || 
 					(CastField<FArrayProperty>(PropertyNode->GetProperty()) != nullptr && 
 					CastField<FObjectProperty>(CastField<FArrayProperty>(PropertyNode->GetProperty())->Inner) != nullptr)) // Is an object array
-						{
+				{
 					DragLeaveDelegate = FOnTableRowDragLeave::CreateSP(this, &SDetailSingleItemRow::OnArrayOrCustomDragLeave);
 					AcceptDropDelegate = FOnAcceptDrop::CreateSP(this, PropertyNode->IsReorderable() ? &SDetailSingleItemRow::OnArrayAcceptDrop : &SDetailSingleItemRow::OnArrayHeaderAcceptDrop);
 					CanAcceptDropDelegate = FOnCanAcceptDrop::CreateSP(this, &SDetailSingleItemRow::OnArrayCanAcceptDrop);
-						}
+				}
 			}
 
 			NameColumnBox->AddSlot()
@@ -586,55 +586,10 @@ void SDetailSingleItemRow::Construct( const FArguments& InArgs, FDetailLayoutCus
 					];
 			}
 
-			TArray<FPropertyRowExtensionButton> ExtensionButtons;
-
-			UpdateResetToDefault();
-			FPropertyRowExtensionButton& ResetToDefault = ExtensionButtons.AddDefaulted_GetRef();
-			ResetToDefault.Label = NSLOCTEXT("PropertyEditor", "ResetToDefault", "Reset to Default");
-			ResetToDefault.UIAction = FUIAction(
-				FExecuteAction::CreateSP(this, &SDetailSingleItemRow::OnResetToDefaultClicked),
-				FCanExecuteAction::CreateLambda([this, IsValueEnabledAttribute]()
-					{
-						return IsResetToDefaultVisible() && IsValueEnabledAttribute.Get(true);
-					})
-			);
-
-			// We could just collapse the Reset to Default button by setting the FIsActionButtonVisible delegate,
-			// but this would cause the reset to defaults not to reserve space in the toolbar and not be aligned across all rows.
-			// Instead, we show an empty icon and tooltip and disable the button.
-			static FSlateIcon EnabledResetToDefaultIcon(FAppStyle::Get().GetStyleSetName(), "PropertyWindow.DiffersFromDefault");
-			static FSlateIcon DisabledResetToDefaultIcon(FAppStyle::Get().GetStyleSetName(), "NoBrush");
-			ResetToDefault.Icon = TAttribute<FSlateIcon>::Create([this]()
-			{
-				return IsResetToDefaultVisible() ?
-					EnabledResetToDefaultIcon :
-					DisabledResetToDefaultIcon;
-			});
-
-			ResetToDefault.ToolTip = TAttribute<FText>::Create([this]() 
-			{
-				return IsResetToDefaultVisible() ?
-					NSLOCTEXT("PropertyEditor", "ResetToDefaultPropertyValueToolTip", "Reset this property to its default value.") :
-					FText::GetEmpty();
-			});
-
-			CreateGlobalExtensionWidgets(ExtensionButtons);
-
-			FSlimHorizontalToolBarBuilder ToolbarBuilder(TSharedPtr<FUICommandList>(), FMultiBoxCustomization::None);
-			ToolbarBuilder.SetLabelVisibility(EVisibility::Collapsed);
-			ToolbarBuilder.SetStyle(&FAppStyle::Get(), "DetailsView.ExtensionToolBar");
-			ToolbarBuilder.SetIsFocusable(false);
-
-			for (const FPropertyRowExtensionButton& Extension : ExtensionButtons)
-			{
-				ToolbarBuilder.AddToolBarButton(Extension.UIAction, NAME_None, Extension.Label, Extension.ToolTip, Extension.Icon);
-			}
-
-			FProperty* Property = nullptr;
-			
+			TSharedPtr<SWidget> RightmostWidget;
 			if (GetPropertyNode().IsValid() &&
-				DetailsView &&
-				DetailsView->GetDisplayManager().IsValid())
+					DetailsView &&
+					DetailsView->GetDisplayManager().IsValid())
 			{
 				TSharedPtr<FDetailsDisplayManager> DisplayManagerLocal = DetailsView->GetDisplayManager();
 
@@ -642,12 +597,12 @@ void SDetailSingleItemRow::Construct( const FArguments& InArgs, FDetailLayoutCus
 				if (DisplayManagerLocal->CanConstructPropertyUpdatedWidgetBuilder())
 				{
 					DisplayManager = MoveTemp(DisplayManagerLocal);
-				
+
 					FConstructPropertyUpdatedWidgetBuilderArgs Args;
 					Args.ResetToDefaultAction = FExecuteAction::CreateSP(this, &SDetailSingleItemRow::OnResetToDefaultClicked);
 					Args.PropertyPath = MakeShared<FPropertyPath>(InOwnerTreeNode->GetPropertyPath());
 					Args.CategoryObjectName = Category->GetObjectName();
-					
+						
 					PropertyUpdatedWidgetBuilder = DisplayManager->ConstructPropertyUpdatedWidgetBuilder(Args);
 
 					if (PropertyUpdatedWidgetBuilder.IsValid())
@@ -659,6 +614,63 @@ void SDetailSingleItemRow::Construct( const FArguments& InArgs, FDetailLayoutCus
 				PRAGMA_ENABLE_DEPRECATION_WARNINGS
 			}
 
+			if(PropertyUpdatedWidgetBuilder.IsValid())
+			{
+				RightmostWidget = PropertyUpdatedWidgetBuilder->GenerateWidget();
+			}
+			else if(WidgetRow.HasResetToDefaultContent())
+			{
+				RightmostWidget = WidgetRow.ResetToDefaultWidget.Widget;
+			}
+			else
+			{
+				TArray<FPropertyRowExtensionButton> ExtensionButtons;
+
+				UpdateResetToDefault();
+				FPropertyRowExtensionButton& ResetToDefault = ExtensionButtons.AddDefaulted_GetRef();
+				ResetToDefault.Label = NSLOCTEXT("PropertyEditor", "ResetToDefault", "Reset to Default");
+				ResetToDefault.UIAction = FUIAction(
+					FExecuteAction::CreateSP(this, &SDetailSingleItemRow::OnResetToDefaultClicked),
+					FCanExecuteAction::CreateLambda([this, IsValueEnabledAttribute]()
+						{
+							return IsResetToDefaultVisible() && IsValueEnabledAttribute.Get(true);
+						})
+				);
+
+				// We could just collapse the Reset to Default button by setting the FIsActionButtonVisible delegate,
+				// but this would cause the reset to defaults not to reserve space in the toolbar and not be aligned across all rows.
+				// Instead, we show an empty icon and tooltip and disable the button.
+				static FSlateIcon EnabledResetToDefaultIcon(FAppStyle::Get().GetStyleSetName(), "PropertyWindow.DiffersFromDefault");
+				static FSlateIcon DisabledResetToDefaultIcon(FAppStyle::Get().GetStyleSetName(), "NoBrush");
+				ResetToDefault.Icon = TAttribute<FSlateIcon>::Create([this]()
+				{
+					return IsResetToDefaultVisible() ?
+						EnabledResetToDefaultIcon :
+						DisabledResetToDefaultIcon;
+				});
+
+				ResetToDefault.ToolTip = TAttribute<FText>::Create([this]() 
+				{
+					return IsResetToDefaultVisible() ?
+						NSLOCTEXT("PropertyEditor", "ResetToDefaultPropertyValueToolTip", "Reset this property to its default value.") :
+						FText::GetEmpty();
+				});
+
+				CreateGlobalExtensionWidgets(ExtensionButtons);
+
+				FSlimHorizontalToolBarBuilder ToolbarBuilder(TSharedPtr<FUICommandList>(), FMultiBoxCustomization::None);
+				ToolbarBuilder.SetLabelVisibility(EVisibility::Collapsed);
+				ToolbarBuilder.SetStyle(&FAppStyle::Get(), "DetailsView.ExtensionToolBar");
+				ToolbarBuilder.SetIsFocusable(false);
+
+				for (const FPropertyRowExtensionButton& Extension : ExtensionButtons)
+				{
+					ToolbarBuilder.AddToolBarButton(Extension.UIAction, NAME_None, Extension.Label, Extension.ToolTip, Extension.Icon);
+				}
+
+				RightmostWidget = ToolbarBuilder.MakeWidget();
+			}
+
 			Splitter->AddSlot()
 				.Value(ColumnSizeData.GetRightColumnWidth())
 				.OnSlotResized(ColumnSizeData.GetOnRightColumnResized())
@@ -671,9 +683,7 @@ void SDetailSingleItemRow::Construct( const FArguments& InArgs, FDetailLayoutCus
 				.VAlign(VAlign_Center)
 				.Padding(0.0f)
 				[
-					PropertyUpdatedWidgetBuilder.IsValid() ?
-					         PropertyUpdatedWidgetBuilder->GenerateWidget().ToSharedRef()  :
-					         ToolbarBuilder.MakeWidget()
+					RightmostWidget.ToSharedRef()
 				]
 			];
 		}
diff --git a/Engine/Source/Editor/PropertyEditor/Public/DetailWidgetRow.h b/Engine/Source/Editor/PropertyEditor/Public/DetailWidgetRow.h
--- a/Engine/Source/Editor/PropertyEditor/Public/DetailWidgetRow.h
+++ b/Engine/Source/Editor/PropertyEditor/Public/DetailWidgetRow.h
@@ -112,6 +112,7 @@ public:
 		: NameWidget( *this, 0.0f, 0.0f, HAlign_Left, VAlign_Center )
 		, ValueWidget( *this, DefaultValueMinWidth, DefaultValueMaxWidth, HAlign_Left, VAlign_Center )
 		, ExtensionWidget( *this, 0.0f, 0.0f, HAlign_Right, VAlign_Center)
+		, ResetToDefaultWidget( *this, 0.0f, 0.0f, HAlign_Center, VAlign_Center)
 		, WholeRowWidget( *this, 0.0f, 0.0f, HAlign_Fill, VAlign_Fill )
 		, VisibilityAttr( EVisibility::Visible )
 		, FilterTextString()
@@ -126,6 +127,7 @@ public:
 		NameWidget = FDetailWidgetDecl(*this, Other.NameWidget);
 		ValueWidget = FDetailWidgetDecl(*this, Other.ValueWidget);
 		ExtensionWidget = FDetailWidgetDecl(*this, Other.ExtensionWidget);
+		ResetToDefaultWidget = FDetailWidgetDecl(*this, Other.ResetToDefaultWidget);
 		WholeRowWidget = FDetailWidgetDecl(*this, Other.WholeRowWidget);
 		VisibilityAttr = Other.VisibilityAttr;
 		IsEnabledAttr = Other.IsEnabledAttr;
@@ -190,6 +192,14 @@ public:
 		return ExtensionWidget;
 	}
 
+	/**
+	 * Assigns content to the reset to default (right-most) slot
+	 */
+	FDetailWidgetDecl& ResetToDefaultContent()
+	{
+		return ResetToDefaultWidget;
+	}
+
 	/**
 	 * Sets a string which should be used to filter the content when a user searches
 	 */
@@ -268,6 +278,11 @@ public:
 		return ExtensionWidget.Widget->GetType() != InvalidDetailWidgetName;
 	}
 
+	bool HasResetToDefaultContent() const
+	{
+		return ResetToDefaultWidget.Widget->GetType() != InvalidDetailWidgetName;
+	}
+	
 	/**
 	 * @return true if the row has columns, false if it spans the entire row
 	 */
@@ -281,7 +296,7 @@ public:
 	 */
 	bool HasAnyContent() const
 	{
-		return WholeRowWidget.Widget->GetType() != InvalidDetailWidgetName || HasColumns();
+		return WholeRowWidget.Widget->GetType() != InvalidDetailWidgetName || HasColumns() || HasResetToDefaultContent();
 	}
 
 	/** @return true if a custom copy/paste is bound on this row */
@@ -365,6 +380,8 @@ public:
 	FDetailWidgetDecl ValueWidget;
 	/** Extension (right) column content */
 	FDetailWidgetDecl ExtensionWidget;
+	/** Reset to default (right-most) column content */
+	FDetailWidgetDecl ResetToDefaultWidget;
 	/** Whole row content */
 	FDetailWidgetDecl WholeRowWidget;
 	/** Visibility of the row */
diff --git a/Engine/Source/Editor/PropertyEditor/Public/IDetailPropertyRow.h b/Engine/Source/Editor/PropertyEditor/Public/IDetailPropertyRow.h
--- a/Engine/Source/Editor/PropertyEditor/Public/IDetailPropertyRow.h
+++ b/Engine/Source/Editor/PropertyEditor/Public/IDetailPropertyRow.h
@@ -295,4 +295,11 @@ public:
 	 * @return	Pointer to value widget if custom property widget already exists, nullptr otherwise
 	 */
 	virtual FDetailWidgetDecl* CustomValueWidget() = 0;
+
+	/**
+	 * Gives a non-owning pointer to reset to default widget on existing custom property widget if it exists.
+	 *
+	 * @return	Pointer to reset to default widget if custom property widget already exists, nullptr otherwise
+	 */
+	virtual FDetailWidgetDecl* CustomResetToDefaultWidget() = 0;
 };
