[BugFix] FPropertyNode::NotifyPostChange wasn't constructing array indices lookup table for chain events

diff --git a/Engine/Source/Editor/PropertyEditor/Private/PropertyNode.cpp b/Engine/Source/Editor/PropertyEditor/Private/PropertyNode.cpp
--- a/Engine/Source/Editor/PropertyEditor/Private/PropertyNode.cpp
+++ b/Engine/Source/Editor/PropertyEditor/Private/PropertyNode.cpp
@@ -3146,6 +3146,7 @@ void FPropertyNode::NotifyPreChangeInternal( TSharedRef<FEditPropertyChain> Prop
 void FPropertyNode::NotifyPostChange( FPropertyChangedEvent& InPropertyChangedEvent, class FNotifyHook* InNotifyHook )
 {
 	TSharedRef<FEditPropertyChain> PropertyChain = BuildPropertyChain( InPropertyChangedEvent.Property );
+	TArray<TMap<FString, int32>> ArrayIndices = BuildArrayIndices();
 	
 	// remember the property that was the chain's original active property; this will correspond to the outermost property of struct/array that was modified
 	FProperty* const OriginalActiveProperty = PropertyChain->GetActiveMemberNode() ? PropertyChain->GetActiveMemberNode()->GetValue() : nullptr;
@@ -3205,7 +3206,7 @@ void FPropertyNode::NotifyPostChange( FPropertyChangedEvent& InPropertyChangedEv
 
 				// Use a scope to ensure that only local variable are use in the loop.
 				// Since this object can be destroyed in this loop.
-				auto ScopePostEditChange = [&PropertyChain, &InPropertyChangedEvent, &CurProperty, CurrentObjectIndex](UObject* Object)
+				auto ScopePostEditChange = [&PropertyChain, &InPropertyChangedEvent, &CurProperty, CurrentObjectIndex, &ArrayIndices](UObject* Object)
 				{
 					// copy the property changed event
 					FPropertyChangedEvent ChangedEvent = InPropertyChangedEvent;
@@ -3225,6 +3226,7 @@ void FPropertyNode::NotifyPostChange( FPropertyChangedEvent& InPropertyChangedEv
 					else
 					{
 						FPropertyChangedChainEvent ChainEvent(*PropertyChain, ChangedEvent);
+						ChainEvent.SetArrayIndexPerObject(ArrayIndices);
 						ChainEvent.ObjectIteratorIndex = CurrentObjectIndex;
 
 						Object->PostEditChangeChainProperty(ChainEvent);
@@ -3559,6 +3561,22 @@ TSharedRef<FEditPropertyChain> FPropertyNode::BuildPropertyChain( FProperty* InP
 	return PropertyChain;
 }
 
+TArray<TMap<FString, int32>> FPropertyNode::BuildArrayIndices() const
+{
+	TArray<TMap<FString, int32>> ArrayIndices;
+	ArrayIndices.SetNum(1);
+
+	for (const FPropertyNode* ItemNode = this; ItemNode; ItemNode = ItemNode->GetParentNode())
+	{
+		if (int32 Index = ItemNode->GetArrayIndex(); Index != INDEX_NONE)
+		{
+			ArrayIndices[0].Add(ItemNode->Property->GetName(), Index);
+		}
+	} 
+
+	return ArrayIndices;
+}
+
 TSharedRef<FEditPropertyChain> FPropertyNode::BuildPropertyChain( FProperty* InProperty, const TSet<UObject*>& InAffectedArchetypeInstances ) const
 {
 	TSharedRef<FEditPropertyChain> PropertyChain = BuildPropertyChain(InProperty);
diff --git a/Engine/Source/Editor/PropertyEditor/Private/PropertyNode.h b/Engine/Source/Editor/PropertyEditor/Private/PropertyNode.h
--- a/Engine/Source/Editor/PropertyEditor/Private/PropertyNode.h
+++ b/Engine/Source/Editor/PropertyEditor/Private/PropertyNode.h
@@ -1118,6 +1118,7 @@ public:
 	}
 
 	TSharedRef<FEditPropertyChain> BuildPropertyChain( FProperty* PropertyAboutToChange ) const;
+	TArray<TMap<FString, int32>> BuildArrayIndices() const;
 	
 protected:
 	TSharedRef<FEditPropertyChain> BuildPropertyChain( FProperty* PropertyAboutToChange, const TSet<UObject*>& InAffectedArchetypeInstances ) const;
