Slate- When using IDetailCategoryBuilder::AddExternalObjects, if sub-categories are disabled (via CollapseCategories meta for example), and if no property row is specifically added for a property, but instead a property customization for that property adds custom rows instead, then these rows do not show up.

This is because when subcategories aren't supported, the code in FDetailLayoutBuilderImpl::GenerateDetailLayout() generates a detail layout and manually tries to copy the property nodes from generated child nodes to the parent category. However, custom rows will not have property nodes, and in this case the parent property that has been customized may not itself add a detail node.

The solution I have here which I'd like reviewed is in the case no PropertyNode is found for a child node, to look to the parent node and if the parent node has a valid property node, then add that instead to the parent category. This will force the same customization to run for the parent category. I considered trying to copy over the widgets from the custom property rows added, but this resulted in disconnected widgets.

I'm pretty sure the case I'm trying to solve for here is rare, and so far I haven't encountered any other property panes that even fall inside the if (!PropertyNode.IsValid()) case, but if there are ways to make this more robust, or if I need to take a different approach, please let me know.

Relevant slack thread: https://epicgames.slack.com/archives/C045H4DFN5P/p1724841764932099

diff --git a/Engine/Source/Editor/PropertyEditor/Private/DetailLayoutBuilderImpl.cpp b/Engine/Source/Editor/PropertyEditor/Private/DetailLayoutBuilderImpl.cpp
--- a/Engine/Source/Editor/PropertyEditor/Private/DetailLayoutBuilderImpl.cpp
+++ b/Engine/Source/Editor/PropertyEditor/Private/DetailLayoutBuilderImpl.cpp
@@ -459,13 +459,28 @@ void FDetailLayoutBuilderImpl::GenerateDetailLayout()
 				FDetailNodeList ChildNodes;
 				DetailCategory->GenerateLayout();
 				DetailCategory->GetGeneratedChildren(ChildNodes, /*bIgnoreVisibility*/true, /*bIgnoreAdvancedDropdown*/true);
-
+				TArray<TSharedPtr<FPropertyNode>> AddedNodes;
 				for (const TSharedRef<FDetailTreeNode>& ChildNode : ChildNodes)
 				{
 					TSharedPtr<FPropertyNode> PropertyNode = ChildNode->GetPropertyNode();
 					if (!PropertyNode.IsValid())
 					{
-						continue;
+						// If we don't have a property node, then we're likely a custom generated node from a parent node. 
+						// If so, add the parent node to the parent category if it has a valid property node, but ensure we don't add more than once, 
+						// as there may be more than one custom row added.
+						TWeakPtr<FDetailTreeNode> ParentNode = ChildNode->GetParentNode();
+						if (ParentNode.IsValid())
+						{
+							PropertyNode = ParentNode.Pin()->GetPropertyNode();
+							if (AddedNodes.Contains(PropertyNode))
+							{
+								continue;
+							}
+						}
+						if (!PropertyNode.IsValid())
+						{
+							continue;
+						}
 					}
 
 					// If there is no outer object then the class is the object root and there is only one instance
@@ -476,6 +491,7 @@ void FDetailLayoutBuilderImpl::GenerateDetailLayout()
 						InstanceName = ParentNode->GetProperty()->GetFName();
 					}
 					ParentDetailCategory->AddPropertyNode(PropertyNode.ToSharedRef(), InstanceName);
+					AddedNodes.Add(PropertyNode);
 				}
 			}
 
