[SWIN] [MetaSounds] [Editor] Add ability to Audition SoundWaves From Input Literal/Array Literal
A way to add Optional Custom Widget at SPropertyEditorAsset depending on the selected AssetDefinition. Retrieves the Custom widget through GetCustomAssetPickerWidget function that can be overriden in AssetDefinitions.
#rnx

diff --git a/Engine/Source/Editor/AssetDefinition/Public/AssetDefinition.h b/Engine/Source/Editor/AssetDefinition/Public/AssetDefinition.h
--- a/Engine/Source/Editor/AssetDefinition/Public/AssetDefinition.h
+++ b/Engine/Source/Editor/AssetDefinition/Public/AssetDefinition.h
@@ -6,6 +6,7 @@
 #include "UObject/Object.h"
 #include "UObject/SoftObjectPtr.h"
 #include "AssetRegistry/AssetData.h"
+#include "Framework/SlateDelegates.h"
 #include "Misc/ScopedSlowTask.h"
 #include "Toolkits/IToolkit.h"
 #include "Misc/AssetFilterData.h"
@@ -336,6 +337,14 @@ public:
 	TSharedPtr<SWidget> ActionButtonWidget;
 };
 
+struct FAssetButtonActionExtension
+{
+public:
+	TAttribute<FText> PickTooltipAttribute;
+	TAttribute<const FSlateBrush*> PickBrushAttribute;
+	FOnClicked OnClicked;
+};
+
 /**
  * Asset Definitions represent top level assets that are known to the editor.
  *
@@ -597,6 +606,9 @@ public:
 	}
 #endif
 
+	/**Optionally adds custom Asset Picker buttons to the Property Editor for the selected AssetDefinition, using the provided extension data this function returns.*/
+	virtual void GetAssetActionButtonExtensions(const FAssetData& InAssetData, TArray<FAssetButtonActionExtension>& OutExtensions) const { };
+
 	/** Whether this asset has external packages associated with it (impacts saving and dirty-state behavior) */
 	virtual bool ShouldSaveExternalPackages() const
 	{
diff --git a/Engine/Source/Editor/AudioEditor/Private/AssetTypeActions/AssetDefinition_SoundBase.cpp b/Engine/Source/Editor/AudioEditor/Private/AssetTypeActions/AssetDefinition_SoundBase.cpp
--- a/Engine/Source/Editor/AudioEditor/Private/AssetTypeActions/AssetDefinition_SoundBase.cpp
+++ b/Engine/Source/Editor/AudioEditor/Private/AssetTypeActions/AssetDefinition_SoundBase.cpp
@@ -177,6 +177,46 @@ EAssetCommandResult UAssetDefinition_SoundBase::ActivateAssets(const FAssetActiv
 	return Super::ActivateAssets(ActivateArgs);
 }
 
+void UAssetDefinition_SoundBase::GetAssetActionButtonExtensions(const FAssetData& InAssetData, TArray<FAssetButtonActionExtension>& OutExtensions) const
+{
+	FAssetButtonActionExtension AssetButtonActionExtension
+	{	
+		.PickTooltipAttribute = TAttribute<FText>::CreateLambda([InAssetData]() -> const FText
+		{
+			if (UE::AudioEditor::IsSoundPlaying(InAssetData))
+			{
+				return LOCTEXT("SoundAudition_PlaySoundToolTip", "Stop selected sound");
+			}
+
+			return LOCTEXT("SoundAudition_StopSoundToolTip", "Play selected sound");
+		}),
+		.PickBrushAttribute = TAttribute<const FSlateBrush*>::CreateLambda([InAssetData]() -> const FSlateBrush* {
+		
+			if (UE::AudioEditor::IsSoundPlaying(InAssetData))
+			{
+				return FAppStyle::GetBrush("MediaAsset.AssetActions.Stop.Small");
+			}
+
+			return FAppStyle::GetBrush("MediaAsset.AssetActions.Play.Small");
+		}),
+		.OnClicked = FOnClicked::CreateLambda([InAssetData]() -> FReply 
+			{
+				if (UE::AudioEditor::IsSoundPlaying(InAssetData))
+				{
+					UE::AudioEditor::StopSound();
+				}
+				else
+				{
+					// Load and play sound
+					UE::AudioEditor::PlaySound(Cast<USoundBase>(InAssetData.GetAsset()));
+				}
+			return FReply::Handled(); 
+		})
+	};
+	
+	OutExtensions.Add(AssetButtonActionExtension);
+}
+
 TSharedPtr<SWidget> UAssetDefinition_SoundBase::GetThumbnailOverlay(const FAssetData& InAssetData) const
 {
 	auto OnClickedLambda = [InAssetData]() -> FReply
diff --git a/Engine/Source/Editor/AudioEditor/Public/AssetTypeActions/AssetDefinition_SoundBase.h b/Engine/Source/Editor/AudioEditor/Public/AssetTypeActions/AssetDefinition_SoundBase.h
--- a/Engine/Source/Editor/AudioEditor/Public/AssetTypeActions/AssetDefinition_SoundBase.h
+++ b/Engine/Source/Editor/AudioEditor/Public/AssetTypeActions/AssetDefinition_SoundBase.h
@@ -37,6 +37,7 @@ public:
 	virtual TSharedPtr<SWidget> GetThumbnailOverlay(const FAssetData& InAssetData) const override;
 	virtual bool GetThumbnailActionOverlay(const FAssetData& InAssetData, FAssetActionThumbnailOverlayInfo& OutActionOverlayInfo) const override;
 	virtual EAssetCommandResult ActivateAssets(const FAssetActivateArgs& ActivateArgs) const override;
+	virtual void GetAssetActionButtonExtensions(const FAssetData& InAssetData, TArray<FAssetButtonActionExtension>& OutExtensions) const override;
 	// UAssetDefinition End
 
 	// Menu Extension statics
diff --git a/Engine/Source/Editor/PropertyEditor/Private/UserInterface/PropertyEditor/SPropertyEditorAsset.cpp b/Engine/Source/Editor/PropertyEditor/Private/UserInterface/PropertyEditor/SPropertyEditorAsset.cpp
--- a/Engine/Source/Editor/PropertyEditor/Private/UserInterface/PropertyEditor/SPropertyEditorAsset.cpp
+++ b/Engine/Source/Editor/PropertyEditor/Private/UserInterface/PropertyEditor/SPropertyEditorAsset.cpp
@@ -1,6 +1,8 @@
 // Copyright Epic Games, Inc. All Rights Reserved.
 
 #include "UserInterface/PropertyEditor/SPropertyEditorAsset.h"
+#include "AssetDefinitionRegistry.h"
+#include "AssetDefinition.h"
 #include "Engine/Texture.h"
 #include "Engine/SkeletalMesh.h"
 #include "Components/StaticMeshComponent.h"
@@ -12,6 +14,7 @@
 #include "UObject/UObjectIterator.h"
 #include "Widgets/Layout/SBox.h"
 #include "Widgets/Images/SImage.h"
+#include "Widgets/Input/SButton.h"
 #include "Components/SkeletalMeshComponent.h"
 #include "Particles/ParticleSystem.h"
 #include "UserInterface/PropertyEditor/SPropertyEditorAsset.h"
@@ -298,6 +301,45 @@ bool SPropertyEditorAsset::IsAssetFiltered(const FAssetData& InAssetData)
 	return false;
 }
 
+void SPropertyEditorAsset::GenerateCustomAssetPickerButtons(const FAssetData& InAssetData, const TArray<FAssetButtonActionExtension>& InExtensions)
+{
+	if (!CustomAssetPickerButtonBox.IsValid())
+	{
+		return;
+	}
+
+	CustomAssetPickerButtonBox->ClearChildren();
+
+	for (FAssetButtonActionExtension Extension : InExtensions)
+	{
+		CustomAssetPickerButtonBox->AddSlot()
+			.Padding(2.0f, 0.0f)
+			.AutoWidth()
+			.VAlign(VAlign_Center)
+			[
+				SNew(SBox)
+					.HAlign(HAlign_Center)
+					.VAlign(VAlign_Center)
+					.WidthOverride(22.0f)
+					.HeightOverride(22.0f)
+					.IsEnabled(true)
+					.ToolTipText(Extension.PickTooltipAttribute)
+					[
+						SNew(SButton)
+							.ButtonStyle(FAppStyle::Get(), "SimpleButton")
+							.OnClicked_Lambda([Extension]()-> FReply { return Extension.OnClicked.Execute(); })
+							.ContentPadding(0.0f)
+							.IsFocusable(false)
+							[
+								SNew(SImage)
+									.Image(Extension.PickBrushAttribute)
+									.ColorAndOpacity(FSlateColor::UseForeground())
+							]
+					]
+			];
+	}
+}
+
 // Awful hack to deal with UClass::FindCommonBase taking an array of non-const classes...
 static TArray<UClass*> ConstCastClassArray(TArray<const UClass*>& Classes)
 {
@@ -718,6 +760,26 @@ void SPropertyEditorAsset::Construct(const FArguments& InArgs, const TSharedPtr<
 			ActorPicker
 		];
 	}
+		
+	FObjectOrAssetData Value;
+	GetValue(Value, FObjectOrAssetData::EAssetDataOptions::SkipAssetRegistryTagsGathering);
+
+	if (const UAssetDefinition* AssetDefinition = UAssetDefinitionRegistry::Get()->GetAssetDefinitionForAsset(Value.AssetData))
+	{
+		TArray<FAssetButtonActionExtension> AssetButtonActionExtension;
+		AssetDefinition->GetAssetActionButtonExtensions(Value.AssetData, AssetButtonActionExtension);
+
+		if (!AssetButtonActionExtension.IsEmpty())
+		{
+			CustomAssetPickerButtonBox = SNew(SHorizontalBox);
+			ButtonBox->AddSlot()
+			[
+				CustomAssetPickerButtonBox.ToSharedRef()
+			];
+
+			GenerateCustomAssetPickerButtons(Value.AssetData, AssetButtonActionExtension);
+		}
+	}	
 
 	NumButtons = ButtonBox->NumSlots();
 	
@@ -1050,6 +1112,13 @@ void SPropertyEditorAsset::SetValue( const FAssetData& AssetData )
 			if (PropertyEditor.IsValid())
 			{
 				PropertyEditor->GetPropertyHandle()->SetValue(AssetData);
+
+				if (const UAssetDefinition* AssetDefinition = UAssetDefinitionRegistry::Get()->GetAssetDefinitionForAsset(AssetData))
+				{
+					TArray<FAssetButtonActionExtension> AssetButtonActionExtension;
+					AssetDefinition->GetAssetActionButtonExtensions(AssetData, AssetButtonActionExtension);
+					GenerateCustomAssetPickerButtons(AssetData, AssetButtonActionExtension);
+				}
 			}
 
 			OnSetObject.ExecuteIfBound(AssetData);
diff --git a/Engine/Source/Editor/PropertyEditor/Private/UserInterface/PropertyEditor/SPropertyEditorAsset.h b/Engine/Source/Editor/PropertyEditor/Private/UserInterface/PropertyEditor/SPropertyEditorAsset.h
--- a/Engine/Source/Editor/PropertyEditor/Private/UserInterface/PropertyEditor/SPropertyEditorAsset.h
+++ b/Engine/Source/Editor/PropertyEditor/Private/UserInterface/PropertyEditor/SPropertyEditorAsset.h
@@ -26,6 +26,7 @@ class SBorder;
 class SComboButton;
 class UFactory;
 class FDetailWidgetRow;
+struct FAssetButtonActionExtension;
 
 /**
  * A widget used to edit Asset-type properties (UObject-derived properties).
@@ -368,6 +369,11 @@ private:
 	/** @return Returns true if the asset is excluded for this property*/
 	bool IsAssetFiltered(const FAssetData& InAssetData);
 
+	/**
+	 *Generates custom Asset Picker Buttons using the provided AssetActionButtonExtensions data.
+	 */
+	void GenerateCustomAssetPickerButtons(const FAssetData& InAssetData, const TArray<FAssetButtonActionExtension>& InExtensions);
+
 private:
 
 	/** Main combobutton */
@@ -379,6 +385,9 @@ private:
 	/** The property editor, if any */
 	TSharedPtr<FPropertyEditor> PropertyEditor;
 
+	/** Optional Custom Asset Picker Buttons*/
+	TSharedPtr<SHorizontalBox> CustomAssetPickerButtonBox;
+
 	/** Path to the object being edited instead of accessing the value directly with a property handle */
 	TAttribute<FString> ObjectPath;
 
diff --git a/Engine/Source/Editor/PropertyEditor/PropertyEditor.Build.cs b/Engine/Source/Editor/PropertyEditor/PropertyEditor.Build.cs
--- a/Engine/Source/Editor/PropertyEditor/PropertyEditor.Build.cs
+++ b/Engine/Source/Editor/PropertyEditor/PropertyEditor.Build.cs
@@ -30,6 +30,7 @@ public class PropertyEditor : ModuleRules
 		PrivateDependencyModuleNames.AddRange(
 			new string[] {
                 "AppFramework",
+				"AssetDefinition",
 				"Core",
 				"CoreUObject",
 				"ApplicationCore",
