Unify Struct Detail View header style with Object Detail View's.

diff --git a/Engine/Source/Editor/PropertyEditor/Private/SStructureDetailsView.cpp b/Engine/Source/Editor/PropertyEditor/Private/SStructureDetailsView.cpp
--- a/Engine/Source/Editor/PropertyEditor/Private/SStructureDetailsView.cpp
+++ b/Engine/Source/Editor/PropertyEditor/Private/SStructureDetailsView.cpp
@@ -3,6 +3,7 @@
 #include "SStructureDetailsView.h"
 #include "Framework/MultiBox/MultiBoxBuilder.h"
 #include "Widgets/Images/SImage.h"
+#include "Widgets/Images/SLayeredImage.h"
 #include "Widgets/Input/SComboButton.h"
 #include "AssetSelection.h"
 #include "DetailCategoryBuilderImpl.h"
@@ -70,7 +71,7 @@ void SStructureDetailsView::Construct(const FArguments& InArgs)
 			FSlateIcon(),
 			ShowOnlyModifiedAction,
 			NAME_None,
-			EUserInterfaceActionType::ToggleButton 
+			EUserInterfaceActionType::Check
 		);
 	}
 	if (DetailsViewArgs.bShowKeyablePropertiesOption)
@@ -85,7 +86,7 @@ void SStructureDetailsView::Construct(const FArguments& InArgs)
 				FIsActionChecked::CreateSP(this, &SStructureDetailsView::IsShowKeyableChecked)
 			),
 			NAME_None,
-			EUserInterfaceActionType::ToggleButton
+			EUserInterfaceActionType::Check
 		);
 	}
 	if (DetailsViewArgs.bShowAnimatedPropertiesOption)
@@ -100,7 +101,7 @@ void SStructureDetailsView::Construct(const FArguments& InArgs)
 				FIsActionChecked::CreateSP(this, &SStructureDetailsView::IsShowAnimatedChecked)
 			),
 			NAME_None,
-			EUserInterfaceActionType::ToggleButton
+			EUserInterfaceActionType::Check
 		);
 	}
 	FUIAction ShowAllAdvancedAction( 
@@ -115,7 +116,7 @@ void SStructureDetailsView::Construct(const FArguments& InArgs)
 		FSlateIcon(),
 		ShowAllAdvancedAction,
 		NAME_None,
-		EUserInterfaceActionType::ToggleButton 
+		EUserInterfaceActionType::Check
 		);
 
 	DetailViewOptions.AddMenuEntry(
@@ -134,6 +135,7 @@ void SStructureDetailsView::Construct(const FArguments& InArgs)
 		+SHorizontalBox::Slot()
 		.FillWidth( 1 )
 		.VAlign( VAlign_Center )
+		.Padding( 6.0f )
 		[
 			// Create the search box
 			SAssignNew( SearchBox, SSearchBox )
@@ -142,22 +144,34 @@ void SStructureDetailsView::Construct(const FArguments& InArgs)
 
 	if (DetailsViewArgs.bShowOptions)
 	{
+		const TSharedPtr<SLayeredImage> FilterImage = SNew(SLayeredImage)
+			.Image(FAppStyle::Get().GetBrush("DetailsView.ViewOptions"))
+			.ColorAndOpacity(FSlateColor::UseForeground());
+
+		// Badge the filter icon if there are filters active
+		FilterImage->AddLayer(TAttribute<const FSlateBrush*>(
+			this, &SStructureDetailsView::GetViewOptionsBadgeIcon)
+		);
+
 		FilterBoxRow->AddSlot()
+			.Padding(0.0f)
 			.HAlign(HAlign_Right)
+			.VAlign(VAlign_Center)
 			.AutoWidth()
 			[
-				SNew( SComboButton )
+				SNew(SComboButton)
+				.HasDownArrow(false)
 				.ContentPadding(0.0f)
-				.ForegroundColor( FSlateColor::UseForeground() )
-				.ButtonStyle( FAppStyle::Get(), "ToggleButton" )
+				.ForegroundColor(FSlateColor::UseForeground())
+				.ButtonStyle(FAppStyle::Get(), "SimpleButton")
+				.AddMetaData<FTagMetaData>(FTagMetaData("ViewOptions"))
 				.MenuContent()
 				[
 					DetailViewOptions.MakeWidget()
 				]
 				.ButtonContent()
 				[
-					SNew(SImage)
-					.Image( FAppStyle::GetBrush("GenericViewButton") )
+					FilterImage.ToSharedRef()
 				]
 			];
 	}
@@ -182,7 +196,7 @@ void SStructureDetailsView::Construct(const FArguments& InArgs)
 			SNew( SVerticalBox )
 			+ SVerticalBox::Slot()
 			.AutoHeight()
-			.Padding( 0.0f, 0.0f, 0.0f, 2.0f )
+			.Padding( 2.0f )
 			[
 				FilterBoxRow
 			]
@@ -365,4 +379,16 @@ EVisibility SStructureDetailsView::GetPropertyEditingVisibility() const
 	return StructProvider.IsValid() && StructProvider->IsValid() && RootNode && RootNode->HasValidStructData() ? EVisibility::Visible : EVisibility::Collapsed;
 }
 
+const FSlateBrush* SStructureDetailsView::GetViewOptionsBadgeIcon() const
+{
+	// Badge the icon if any view option that narrows down the results is checked
+	const bool bHasBadge =
+		(DetailsViewArgs.bShowModifiedPropertiesOption && IsShowOnlyModifiedChecked()) ||
+		(DetailsViewArgs.bShowDifferingPropertiesOption && IsShowOnlyAllowedChecked()) ||
+		(DetailsViewArgs.bShowKeyablePropertiesOption && IsShowKeyableChecked()) ||
+		(DetailsViewArgs.bShowAnimatedPropertiesOption && IsShowAnimatedChecked());
+
+	return bHasBadge ? FAppStyle::Get().GetBrush("Icons.BadgeModified") : nullptr;
+}
+
 #undef LOCTEXT_NAMESPACE
diff --git a/Engine/Source/Editor/PropertyEditor/Private/SStructureDetailsView.h b/Engine/Source/Editor/PropertyEditor/Private/SStructureDetailsView.h
--- a/Engine/Source/Editor/PropertyEditor/Private/SStructureDetailsView.h
+++ b/Engine/Source/Editor/PropertyEditor/Private/SStructureDetailsView.h
@@ -123,6 +123,9 @@ protected:
 
 	EVisibility GetPropertyEditingVisibility() const;
 
+	/** Get the badge for the view options icon (if there are any options active */
+	const FSlateBrush* GetViewOptionsBadgeIcon() const;
+
 private:
 	TSharedPtr<IStructureDataProvider> StructProvider;
 	FRootPropertyNodeList RootNodes;
