PropertyEditor: Rebuild the DetailTreeNodewhen a struct changes. Items are created or removed. The editcondition runs on invalid node.

diff --git a/Engine/Source/Editor/PropertyEditor/Private/StructurePropertyNode.cpp b/Engine/Source/Editor/PropertyEditor/Private/StructurePropertyNode.cpp
--- a/Engine/Source/Editor/PropertyEditor/Private/StructurePropertyNode.cpp
+++ b/Engine/Source/Editor/PropertyEditor/Private/StructurePropertyNode.cpp
@@ -298,7 +298,8 @@ EPropertyDataValidationResult FStructurePropertyNode::EnsureDataIsValid()
 	if (GetBaseStructure() != CachedBaseStruct)
 	{
 		RebuildChildren();
-		return EPropertyDataValidationResult::ChildrenRebuilt;
+		// Invalidate the object to prevent EditCondition evaluating on old nodes.
+		return EPropertyDataValidationResult::ObjectInvalid;
 	}
 	
 	return FPropertyNode::EnsureDataIsValid();
