PropertyEditor: Updated the details view to make sections compatible with root object customizations.

Previously, sections would not show up if root object customizations were being used.

#rnx

diff --git a/Engine/Source/Editor/PropertyEditor/Private/SDetailsView.cpp b/Engine/Source/Editor/PropertyEditor/Private/SDetailsView.cpp
--- a/Engine/Source/Editor/PropertyEditor/Private/SDetailsView.cpp
+++ b/Engine/Source/Editor/PropertyEditor/Private/SDetailsView.cpp
@@ -1509,12 +1509,21 @@ TMap<FName, FText> SDetailsView::GetAllSections() const
 		{
 			if (TreeNode->GetNodeType() == EDetailNodeType::Category)
 			{
-				FDetailCategoryImpl& Category = (FDetailCategoryImpl&) TreeNode.Get();
-				TArray<TSharedRef<FDetailTreeNode>> Children;
-				Category.GetGeneratedChildren(Children, true, false);
-
 				LayoutCategories.Add(TreeNode->GetNodeName());
 			}
+			else if (TreeNode->GetNodeType() == EDetailNodeType::Object)
+			{
+				TArray<TSharedRef<IDetailTreeNode>> RootObjectChildren;
+				TreeNode->GetChildren(RootObjectChildren);
+
+				for (const TSharedRef<IDetailTreeNode>& RootObjectChild : RootObjectChildren)
+				{
+					if (RootObjectChild->GetNodeType() == EDetailNodeType::Category)
+					{
+						LayoutCategories.Add(RootObjectChild->GetNodeName());
+					}
+				}
+			}
 		}
 	}
 
