Localization:
Fixed a bug with the property matrix and how it handled pasting text into a textbox for text properties. The text correctly gets pasted and imported now.

#test Got the help of UE QA to verify that the problem was no longer occuring in the paste case for text. THe localization caches are now being correctly populated.

diff --git a/Engine/Source/Editor/PropertyEditor/Private/PropertyHandleImpl.cpp b/Engine/Source/Editor/PropertyEditor/Private/PropertyHandleImpl.cpp
--- a/Engine/Source/Editor/PropertyEditor/Private/PropertyHandleImpl.cpp
+++ b/Engine/Source/Editor/PropertyEditor/Private/PropertyHandleImpl.cpp
@@ -38,6 +38,7 @@
 #include "SStandaloneCustomizedValueWidget.h"
 #include "Modules/ModuleManager.h"
 #include "PropertyEditorModule.h"
+#include "Internationalization/TextPackageNamespaceUtil.h"
 
 #define LOCTEXT_NAMESPACE "PropertyHandleImplementation"
 
@@ -5460,7 +5461,7 @@ FPropertyAccess::Result FPropertyHandleText::SetValue(const FText& NewValue, EPr
 {
 	FString StringValue;
 	FTextStringHelper::WriteToBuffer(StringValue, NewValue);
-	return Implementation->ImportText(StringValue, Flags);
+	return ImportFormattedTextString(StringValue, Flags);
 }
 
 FPropertyAccess::Result FPropertyHandleText::SetValue(const FString& NewValue, EPropertyValueSetFlags::Type Flags)
@@ -5473,6 +5474,93 @@ FPropertyAccess::Result FPropertyHandleText::SetValue(const TCHAR* NewValue, EPr
 	return SetValue(FText::FromString(NewValue), Flags);
 }
 
+FPropertyAccess::Result FPropertyHandleText::SetValueFromFormattedString(const FString& InValue, EPropertyValueSetFlags::Type Flags)
+{
+	return ImportFormattedTextString(InValue, Flags);
+}
+
+FPropertyAccess::Result FPropertyHandleText::ImportFormattedTextString(const FString& NewValue, EPropertyValueSetFlags::Type Flags)
+{
+	TSharedPtr<FPropertyNode> PropertyNode = GetPropertyNode();
+
+	TArray<FObjectBaseAddress> ObjectsToModify;
+	Implementation->GetObjectsToModify(ObjectsToModify, PropertyNode.Get());
+
+	TArray<FString> NewTextValuesToImport;
+	{
+		TArray<UPackage*> OuterPackages;
+		GetOuterPackages(OuterPackages);
+
+		FTextProperty* TextProperty = CastField<FTextProperty>(GetProperty());
+		FStringView PropertyPath = GetPropertyPath();
+
+		const FText NewTextValue = FTextStringHelper::CreateFromBuffer(*NewValue);
+		for (int32 ObjectIndex = 0; ObjectIndex < ObjectsToModify.Num(); ++ObjectIndex)
+		{
+			const FObjectBaseAddress& ObjectToModify = ObjectsToModify[ObjectIndex];
+			if (ObjectToModify.BaseAddress)
+			{
+				FText ResolvedTextValue = NewTextValue;
+
+				// If the FText value was initialized from a string, then try and generate a stable and deterministic key for it instead
+				if (TextProperty && NewTextValue.IsInitializedFromString())
+				{
+					UPackage* OuterPackage = OuterPackages.IsValidIndex(ObjectIndex) ? OuterPackages[ObjectIndex] : nullptr;
+
+					auto DeterministicTextKeyGenerator = [PropertyPath, OuterPackage, TextProperty]()
+					{
+						auto GetNameKeyHash = [](const FName Name)
+						{
+							FNameBuilder Builder(Name);
+							return TextKeyUtil::HashString(Builder.ToString(), Builder.Len());
+						};
+
+						auto GetObjectKeyHash = [](const UObject* Obj)
+						{
+							FNameBuilder Builder;
+							if (Obj)
+							{
+								Obj->GetPathName(nullptr, Builder);
+							}
+							return TextKeyUtil::HashString(Builder.ToString(), Builder.Len());
+						};
+
+						// Build a (hopefully) unique and deterministic key from a combination of the outer package and text property info
+						FGuid KeyGuid(GetObjectKeyHash(OuterPackage), TextKeyUtil::HashString(PropertyPath), GetNameKeyHash(TextProperty->GetOwnerStruct()->GetFName()), GetNameKeyHash(TextProperty->GetFName()));
+#if USE_STABLE_LOCALIZATION_KEYS
+						if (OuterPackage)
+						{
+							if (const FString PackageNamespace = TextNamespaceUtil::EnsurePackageNamespace(OuterPackage);
+								!PackageNamespace.IsEmpty())
+							{
+								// Mix the package namespace hash into the outer hash
+								KeyGuid.A = TextKeyUtil::HashString(PackageNamespace, KeyGuid.A);
+							}
+						}
+#endif // USE_STABLE_LOCALIZATION_KEYS
+						return KeyGuid.ToString();
+					};
+
+					const FText* CurrentTextValue = reinterpret_cast<FText*>(ObjectToModify.BaseAddress);
+					ResolvedTextValue = *CurrentTextValue;
+					TextNamespaceUtil::EditTextProperty_Direct(OuterPackage, &ResolvedTextValue, TextProperty, TextNamespaceUtil::ETextEditAction::SourceString, NewTextValue.ToString(), DeterministicTextKeyGenerator);
+				}
+
+				FString& NewTextValueToImport = NewTextValuesToImport.AddDefaulted_GetRef();
+				FTextStringHelper::WriteToBuffer(NewTextValueToImport, ResolvedTextValue);
+			}
+		}
+	}
+
+	FPropertyAccess::Result Result = FPropertyAccess::Fail;
+	if (NewTextValuesToImport.Num() > 0)
+	{
+		Result = Implementation->ImportText(ObjectsToModify, NewTextValuesToImport, PropertyNode.Get(), Flags);
+	}
+
+	return Result;
+}
+
 // Sets
 
 bool FPropertyHandleSet::Supports(TSharedRef<FPropertyNode> PropertyNode)
diff --git a/Engine/Source/Editor/PropertyEditor/Private/PropertyHandleImpl.h b/Engine/Source/Editor/PropertyEditor/Private/PropertyHandleImpl.h
--- a/Engine/Source/Editor/PropertyEditor/Private/PropertyHandleImpl.h
+++ b/Engine/Source/Editor/PropertyEditor/Private/PropertyHandleImpl.h
@@ -755,6 +755,10 @@ public:
 
 	virtual FPropertyAccess::Result SetValue(const FString& InValue, EPropertyValueSetFlags::Type Flags = EPropertyValueSetFlags::DefaultFlags) override;
 	virtual FPropertyAccess::Result SetValue(const TCHAR* InValue, EPropertyValueSetFlags::Type Flags = EPropertyValueSetFlags::DefaultFlags) override;
+	virtual FPropertyAccess::Result SetValueFromFormattedString(const FString& InValue, EPropertyValueSetFlags::Type Flags = EPropertyValueSetFlags::DefaultFlags) override;
+
+private:
+	FPropertyAccess::Result ImportFormattedTextString(const FString& InValue, EPropertyValueSetFlags::Type Flags);
 };
 
 class FPropertyHandleSet : public FPropertyHandleBase, public IPropertyHandleSet
diff --git a/Engine/Source/Runtime/CoreUObject/Private/Internationalization/TextPackageNamespaceUtil.cpp b/Engine/Source/Runtime/CoreUObject/Private/Internationalization/TextPackageNamespaceUtil.cpp
--- a/Engine/Source/Runtime/CoreUObject/Private/Internationalization/TextPackageNamespaceUtil.cpp
+++ b/Engine/Source/Runtime/CoreUObject/Private/Internationalization/TextPackageNamespaceUtil.cpp
@@ -273,14 +273,19 @@ bool TextNamespaceUtil::EditTextProperty(UObject* InTextOwner, const FTextProper
 		return false;
 	}
 
+	return EditTextProperty_Direct(InTextOwner->GetPackage(), InTextProperty->GetPropertyValuePtr_InContainer(InTextOwner), InTextProperty, InEditAction, InEditValue, InTextKeyGenerator, bApplyPackageNamespace);
+}
+
+bool TextNamespaceUtil::EditTextProperty_Direct(UPackage* InPackage, void* InTextValue, const FTextProperty* InTextProperty, const ETextEditAction InEditAction, const FString& InEditValue, TFunctionRef<FString()> InTextKeyGenerator, const bool bApplyPackageNamespace)
+{
 	if (InEditAction == ETextEditAction::SourceString && InEditValue.IsEmpty())
 	{
 		// Empty source strings always produce an empty text
-		InTextProperty->SetPropertyValue_InContainer(InTextOwner, FText());
+		InTextProperty->SetPropertyValue(InTextValue, FText());
 		return true;
 	}
 
-	const FText CurrentTextValue = InTextProperty->GetPropertyValue_InContainer(InTextOwner);
+	const FText CurrentTextValue = InTextProperty->GetPropertyValue(InTextValue);
 	const FTextId CurrentTextId = FTextInspector::GetTextId(CurrentTextValue);
 	const FString* CurrentSourceString = FTextInspector::GetSourceString(CurrentTextValue);
 	const bool bIsCurrentTextLocalized = !CurrentTextValue.IsCultureInvariant() && !CurrentTextValue.IsFromStringTable();
@@ -324,8 +329,8 @@ bool TextNamespaceUtil::EditTextProperty(UObject* InTextOwner, const FTextProper
 
 	FString StableNamespace;
 	FString StableKey;
-	GetTextIdForEdit(InTextOwner->GetPackage(), InEditAction, SourceString, ProposedNamespace, ProposedKey, StableNamespace, StableKey, InTextKeyGenerator, bApplyPackageNamespace);
+	GetTextIdForEdit(InPackage, InEditAction, SourceString, ProposedNamespace, ProposedKey, StableNamespace, StableKey, InTextKeyGenerator, bApplyPackageNamespace);
 
-	InTextProperty->SetPropertyValue_InContainer(InTextOwner, FInternationalization::ForUseOnlyByLocMacroAndGraphNodeTextLiterals_CreateText(*SourceString, *StableNamespace, *StableKey));
+	InTextProperty->SetPropertyValue(InTextValue, FInternationalization::ForUseOnlyByLocMacroAndGraphNodeTextLiterals_CreateText(*SourceString, *StableNamespace, *StableKey));
 	return true;
 }
diff --git a/Engine/Source/Runtime/CoreUObject/Public/Internationalization/TextPackageNamespaceUtil.h b/Engine/Source/Runtime/CoreUObject/Public/Internationalization/TextPackageNamespaceUtil.h
--- a/Engine/Source/Runtime/CoreUObject/Public/Internationalization/TextPackageNamespaceUtil.h
+++ b/Engine/Source/Runtime/CoreUObject/Public/Internationalization/TextPackageNamespaceUtil.h
@@ -163,4 +163,19 @@ COREUOBJECT_API void GetTextIdForEdit(UPackage* InPackage, const ETextEditAction
  */
 COREUOBJECT_API bool EditTextProperty(UObject* InTextOwner, const FTextProperty* InTextProperty, const ETextEditAction InEditAction, const FString& InEditValue, TFunctionRef<FString()> InTextKeyGenerator = &GenerateRandomTextKey, const bool bApplyPackageNamespace = true);
 
+/**
+ * Edit an attribute of the given text property, akin to what happens when editing a text property in a details panel.
+ *
+ * @param InPackage						The package that hosts the text value.
+ * @param InTextValue					The raw value of the TextProperty to be edited.
+ * @param InTextProperty				The text property to edit.
+ * @param InEditAction					How has the given text been edited?
+ * @param InEditValue					The new value of the attribute that was edited.
+ * @param InTextKeyGenerator			Generator for the new text key. Will generate a random key by default.
+ * @param bApplyPackageNamespace		If true, apply the package namespace to the generated text ID (when USE_STABLE_LOCALIZATION_KEYS is true).
+ *
+ * @return True if edit was possible, or false if not.
+ */
+COREUOBJECT_API bool EditTextProperty_Direct(UPackage* InPackage, void* InTextValue, const FTextProperty* InTextProperty, const ETextEditAction InEditAction, const FString& InEditValue, TFunctionRef<FString()> InTextKeyGenerator = &GenerateRandomTextKey, const bool bApplyPackageNamespace = true);
+
 }
