Make EditConditions work for customized properties

* This change systemically allows customized properties to use the EditConditions UPROPERTY meta specifier.
* This was specifically requested for FGuid as it's details panel representation is implemented as a customization at the engine level

diff --git a/Engine/Source/Editor/PropertyEditor/Private/DetailPropertyRow.cpp b/Engine/Source/Editor/PropertyEditor/Private/DetailPropertyRow.cpp
--- a/Engine/Source/Editor/PropertyEditor/Private/DetailPropertyRow.cpp
+++ b/Engine/Source/Editor/PropertyEditor/Private/DetailPropertyRow.cpp
@@ -21,6 +21,8 @@
 
 #define LOCTEXT_NAMESPACE	"DetailPropertyRow"
 
+static void TogglePropertyEditorEditCondition(bool bValue, TWeakPtr<FPropertyEditor> PropertyEditorWeak);
+
 FDetailPropertyRow::FDetailPropertyRow(TSharedPtr<FPropertyNode> InPropertyNode, TSharedRef<FDetailCategoryImpl> InParentCategory, TSharedPtr<FComplexPropertyNode> InExternalRootNode)
 	: PropertyNode( InPropertyNode )
 	, ParentCategory( InParentCategory )
@@ -381,6 +383,18 @@ void FDetailPropertyRow::OnItemNodeInitialized( TSharedRef<FDetailCategoryImpl>
 
 		CustomTypeInterface->CustomizeHeader(PropertyHandle.ToSharedRef(), *CustomPropertyWidget, *this);
 
+		// set edit condition handlers - use customized if provided
+		TAttribute<bool> EditConditionValue = CustomEditConditionValue;
+		if (!EditConditionValue.IsSet() && PropertyEditor.IsValid())
+		{
+			EditConditionValue = TAttribute<bool>(PropertyEditor.ToSharedRef(), &FPropertyEditor::IsEditConditionMet);
+		}
+		
+		TWeakPtr<FPropertyEditor> PropertyEditorWeak = PropertyEditor;
+		FOnBooleanValueChanged OnEditConditionValueChanged = FOnBooleanValueChanged::CreateStatic(&TogglePropertyEditorEditCondition, PropertyEditorWeak);
+
+		CustomPropertyWidget->EditCondition(EditConditionValue, OnEditConditionValueChanged);
+
 		FixEmptyHeaderRowInContainers(PropertyHandle, CustomPropertyWidget);
 
 		// set initial value of enabled attribute to settings from struct customization
