The Details Panel's "Create New Blueprint" button now works correctly with Soft Class Properties

[RN] The Details Panel's Create New Blueprint button now works correctly with Soft Class Properties

diff --git a/Engine/Source/Editor/PropertyEditor/Private/Presentation/PropertyEditor/PropertyEditor.cpp b/Engine/Source/Editor/PropertyEditor/Private/Presentation/PropertyEditor/PropertyEditor.cpp
--- a/Engine/Source/Editor/PropertyEditor/Private/Presentation/PropertyEditor/PropertyEditor.cpp
+++ b/Engine/Source/Editor/PropertyEditor/Private/Presentation/PropertyEditor/PropertyEditor.cpp
@@ -29,6 +29,29 @@
 
 DEFINE_LOG_CATEGORY_STATIC(LogPropertyEditor, Log, All);
 
+namespace UE::Private
+{
+	static UClass* GetMetaClass(const FProperty* ForProperty)
+	{
+		if(!ForProperty)
+		{
+			return nullptr;
+		}
+		else if(const FClassProperty* ClassProp = CastField<FClassProperty>(ForProperty))
+		{
+			return ToRawPtr(ClassProp->MetaClass);
+		}
+		else if(const FSoftClassProperty* SoftClassProp = CastField<FSoftClassProperty>(ForProperty))
+		{
+			return ToRawPtr(SoftClassProp->MetaClass);
+		}
+		else
+		{
+			return ToRawPtr(FEditorClassUtils::GetClassFromString(ForProperty->GetMetaData("MetaClass")));
+		}
+	}
+}
+
 const FString FPropertyEditor::MultipleValuesDisplayName = NSLOCTEXT("PropertyEditor", "MultipleValues", "Multiple Values").ToString();
 
 TSharedRef< FPropertyEditor > FPropertyEditor::Create( const TSharedRef< class FPropertyNode >& InPropertyNode, const TSharedRef<class IPropertyUtilities >& InPropertyUtilities )
@@ -357,8 +380,7 @@ void FPropertyEditor::OnClearItem()
 void FPropertyEditor::MakeNewBlueprint()
 {
 	FProperty* NodeProperty = PropertyNode->GetProperty();
-	FClassProperty* ClassProp = CastField<FClassProperty>(NodeProperty);
-	UClass* Class = (ClassProp ? ToRawPtr(ClassProp->MetaClass) : ToRawPtr(FEditorClassUtils::GetClassFromString(NodeProperty->GetMetaData("MetaClass"))));
+	UClass* Class = UE::Private::GetMetaClass(NodeProperty);
 
 	UClass* RequiredInterface = FEditorClassUtils::GetClassFromString(NodeProperty->GetMetaData("MustImplement"));
 
