Functions that do not take parameters defined in editor function libraries can once again be invoked from the defaults panel. Clean up short lived TFunction

#rnx

diff --git a/Engine/Source/Editor/PropertyEditor/Private/PropertyCustomizationHelpers.cpp b/Engine/Source/Editor/PropertyEditor/Private/PropertyCustomizationHelpers.cpp
--- a/Engine/Source/Editor/PropertyEditor/Private/PropertyCustomizationHelpers.cpp
+++ b/Engine/Source/Editor/PropertyEditor/Private/PropertyCustomizationHelpers.cpp
@@ -329,26 +329,26 @@ namespace PropertyCustomizationHelpers
 			OnExecute = FPropertyFunctionCallDelegates::FOnExecute::CreateLambda(
 				[OnGetExecutionContext = InArgs.OnGetExecutionContext](const TWeakObjectPtr<UFunction>& InWeakFunction)
 				{
-					TArray<TWeakObjectPtr<UObject>> WeakExecutionObjects = OnGetExecutionContext.Execute(InWeakFunction);
-					if (!WeakExecutionObjects.IsEmpty())
+					using namespace UE::Reflection;
+					if (UFunction* Function = InWeakFunction.Get())
 					{
-						using namespace UE::Reflection;
-						if (UFunction* Function = InWeakFunction.Get())
-						{
-							// @todo: Consider naming the transaction scope after the fully qualified function name for better UX
-							FScopedTransaction Transaction(LOCTEXT("ExecuteCallInEditorMethod", "Call In Editor Action"));
-							TStrongObjectPtr<UFunction> CallingFunction(Function);
+						// @todo: Consider naming the transaction scope after the fully qualified function name for better UX
+						FScopedTransaction Transaction(LOCTEXT("ExecuteCallInEditorMethod", "Call In Editor Action"));
+						TStrongObjectPtr<UFunction> CallingFunction(Function);
 
-							if (Function->HasMetaData(NAME_WorldContext)
-								&& DoesStaticFunctionSignatureMatch<void(TObjectPtr<UObject>)>(Function))
-							{
-								FEditorScriptExecutionGuard ScriptGuard;
-								UEditorEngine* EditorEngine = Cast<UEditorEngine>(GEngine);
-								UObject* WorldContextObject = EditorEngine->GetEditorWorldContext().World();
-								TStrongObjectPtr<UObject> CDO(Function->GetOwnerClass()->ClassDefaultObject);
-								CDO->ProcessEvent(Function, &WorldContextObject);
-							}
-							else
+						if (Function->HasMetaData(NAME_WorldContext)
+							&& DoesStaticFunctionSignatureMatch<void(TObjectPtr<UObject>)>(Function))
+						{
+							FEditorScriptExecutionGuard ScriptGuard;
+							UEditorEngine* EditorEngine = Cast<UEditorEngine>(GEngine);
+							UObject* WorldContextObject = EditorEngine->GetEditorWorldContext().World();
+							TStrongObjectPtr<UObject> CDO(Function->GetOwnerClass()->ClassDefaultObject);
+							CDO->ProcessEvent(Function, &WorldContextObject);
+						}
+						else
+						{
+							TArray<TWeakObjectPtr<UObject>> WeakExecutionObjects = OnGetExecutionContext.Execute(InWeakFunction);
+							if (!WeakExecutionObjects.IsEmpty())
 							{
 								FEditorScriptExecutionGuard ScriptGuard;
 								for (const TWeakObjectPtr<UObject>& WeakExecutionObject : WeakExecutionObjects)
@@ -1255,56 +1255,43 @@ namespace PropertyCustomizationHelpers
 {
 	namespace Private
 	{
-		void GetCallInEditorFunctionsForClassInternal(const UClass* InClass, TOptional<TFunctionRef<bool(const UFunction*)>> InFunctionFilter, TArray<UFunction*>& OutCallInEditorFunctions, EFieldIterationFlags InIterationFlags)
-		{
-			// metadata tag for defining sort order of function buttons within a Category
-			static const FName NAME_DisplayPriority("DisplayPriority");
-
-			// FBlueprintMetadata::MD_CallInEditor
-			static const FName NAME_CallInEditor(TEXT("CallInEditor"));
-
-			const bool bDisallowEditorUtilityBlueprintFunctions = GetDefault<UBlueprintEditorProjectSettings>()->bDisallowEditorUtilityBlueprintFunctionsInDetailsView;
+		static const FName NAME_CallInEditor(TEXT("CallInEditor"));
+		static const FName NAME_WorldContext(TEXT("WorldContext"));
 
-			TFunction<bool(const UFunction*)> CanDisplayAndCallFunction =
-				[bDisallowEditorUtilityBlueprintFunctions](const UFunction* TestFunction)
+		static bool CanCallFunctionBasedOnParams(const UFunction* TestFunction)
+		{
+			// If the function only takes a world context object we can use the editor's
+			// world context - but only if the blueprint is editor only and the function
+			// is static:
+			if (UClass* TestFunctionOwnerClass = TestFunction->GetOwnerClass())
+			{
+				if (UBlueprint* Blueprint = Cast<UBlueprint>(TestFunctionOwnerClass->ClassGeneratedBy))
 				{
-					bool bCanCall = TestFunction->GetBoolMetaData(NAME_CallInEditor)
-								&& TestFunction->ParmsSize == 0; // Params not supported
-
-					if (bCanCall)
+					if (FBlueprintEditorUtils::IsEditorUtilityBlueprint(Blueprint) && Blueprint->BlueprintType == BPTYPE_FunctionLibrary)
 					{
-						if (const UClass* TestFunctionOwnerClass = TestFunction->GetOwnerClass())
-						{
-							if (const UBlueprint* Blueprint = Cast<UBlueprint>(TestFunctionOwnerClass->ClassGeneratedBy))
-							{
-								if (FBlueprintEditorUtils::IsEditorUtilityBlueprint(Blueprint))
-								{
-									// Skip Blutilities if disabled via project settings
-									bCanCall = !bDisallowEditorUtilityBlueprintFunctions;
-								}
-							}
-						}
+						using namespace UE::Reflection;
+						return TestFunction->HasMetaData(NAME_WorldContext) &&
+							DoesStaticFunctionSignatureMatch<void(TObjectPtr<UObject>)>(TestFunction);
 					}
+				}
+			}
+
+			const bool bCanCall = TestFunction->GetBoolMetaData(NAME_CallInEditor) && (TestFunction->ParmsSize == 0); // no params required, we can call it!
+			return bCanCall;
+		}
 
-					return bCanCall;
-				};
+		void GetCallInEditorFunctionsForClassInternal(const UClass* InClass, TOptional<TFunctionRef<bool(const UFunction*)>> InFunctionFilter, TArray<UFunction*>& OutCallInEditorFunctions, EFieldIterationFlags InIterationFlags)
+		{
+			// metadata tag for defining sort order of function buttons within a Category
+			static const FName NAME_DisplayPriority("DisplayPriority");
 
-			if (InFunctionFilter.IsSet())
-			{
-				// If the user provided a filter, we call it if the default (OuterFunc) passes - it should always be CallInEditor and have no params
-				CanDisplayAndCallFunction = [OuterFunc = CanDisplayAndCallFunction, InnerFunc = InFunctionFilter.GetValue()](const UFunction* TestFunction)
-				{
-					return OuterFunc(TestFunction)
-						? InnerFunc(TestFunction)
-						: false;
-				};
-			}
+			const bool bDisallowEditorUtilityBlueprintFunctions = GetDefault<UBlueprintEditorProjectSettings>()->bDisallowEditorUtilityBlueprintFunctionsInDetailsView;
 
 			// Get all of the functions we need to display (done ahead of time so we can sort them)
 			for (TFieldIterator<UFunction> FunctionIter(InClass, InIterationFlags); FunctionIter; ++FunctionIter)
 			{
 				const UFunction* TestFunction = *FunctionIter;
-				if (CanDisplayAndCallFunction(TestFunction))
+				if (CanCallFunctionBasedOnParams(TestFunction) && (!InFunctionFilter.IsSet() || InFunctionFilter.GetValue()(TestFunction)))
 				{
 					const FName FunctionName = TestFunction->GetFName();
 
