Fix to ensure that new container items created in the Property Editor get scoped to the correct instance when IDO features are enabled.

diff --git a/Engine/Source/Editor/PropertyEditor/Private/PropertyHandleImpl.cpp b/Engine/Source/Editor/PropertyEditor/Private/PropertyHandleImpl.cpp
--- a/Engine/Source/Editor/PropertyEditor/Private/PropertyHandleImpl.cpp
+++ b/Engine/Source/Editor/PropertyEditor/Private/PropertyHandleImpl.cpp
@@ -1158,9 +1158,17 @@ int32 FPropertyValueImpl::AddChild()
 
 						int32 Index = INDEX_NONE;
 
+						// If the current object is an IDO, map the object back to the associated instance so that new items don't parent to the IDO. Note
+						// however that we're still going to modify the container property's value on the IDO and not its associated instance in this case.
+						UObject* NewItemOuter = Obj;
+						if (const UObject* InstanceFromIDO = UE::FPropertyBagRepository::Get().FindInstanceForDataObject(NewItemOuter))
+						{
+							NewItemOuter = const_cast<UObject*>(InstanceFromIDO);
+						}
+
 						if (Array)
 						{
-							Array->PerformOperationWithSetter(Obj, Addr, [Obj, Array, &Index, &ReturnLogicalIndex](void* DirectAddress)
+							Array->PerformOperationWithSetter(Obj, Addr, [NewItemOuter, Array, &Index, &ReturnLogicalIndex](void* DirectAddress)
 							{								
 								FScriptArrayHelper	ArrayHelper(Array, DirectAddress);
 								Index = ArrayHelper.AddValue();
@@ -1170,14 +1178,14 @@ int32 FPropertyValueImpl::AddChild()
 								FObjectProperty* InnerObjectProperty = CastField<FObjectProperty>(Array->Inner);
 								if (InnerObjectProperty && InnerObjectProperty->HasAnyPropertyFlags(CPF_NonNullable) && !InnerObjectProperty->PropertyClass->HasAnyClassFlags(CLASS_Abstract))
 								{
-									UObject* NewItem = NewObject<UObject>(Obj, InnerObjectProperty->PropertyClass);
+									UObject* NewItem = NewObject<UObject>(NewItemOuter, InnerObjectProperty->PropertyClass);
 									InnerObjectProperty->SetObjectPropertyValue(ArrayHelper.GetRawPtr(Index), NewItem);
 								}
 							});
 						}
 						else if (Set)
 						{
-							Set->PerformOperationWithSetter(Obj, Addr, [Obj, Set, &Index, &ReturnLogicalIndex](void* DirectAddress)
+							Set->PerformOperationWithSetter(Obj, Addr, [NewItemOuter, Set, &Index, &ReturnLogicalIndex](void* DirectAddress)
 							{
 								FScriptSetHelper	SetHelper(Set, DirectAddress);
 								Index = SetHelper.AddDefaultValue_Invalid_NeedsRehash();
@@ -1187,7 +1195,7 @@ int32 FPropertyValueImpl::AddChild()
 								FObjectProperty* ElementObjectProperty = CastField<FObjectProperty>(Set->ElementProp);
 								if (ElementObjectProperty && ElementObjectProperty->HasAnyPropertyFlags(CPF_NonNullable) && !ElementObjectProperty->PropertyClass->HasAnyClassFlags(CLASS_Abstract))
 								{
-									UObject* NewItem = NewObject<UObject>(Obj, ElementObjectProperty->PropertyClass);
+									UObject* NewItem = NewObject<UObject>(NewItemOuter, ElementObjectProperty->PropertyClass);
 									ElementObjectProperty->SetObjectPropertyValue(SetHelper.GetElementPtr(Index), NewItem);
 								}
 
@@ -1196,7 +1204,7 @@ int32 FPropertyValueImpl::AddChild()
 						}
 						else if (Map)
 						{
-							Map->PerformOperationWithSetter(Obj, Addr, [Obj, Map, &Index, &bAddedMapEntry, &ReturnLogicalIndex](void* DirectAddress)
+							Map->PerformOperationWithSetter(Obj, Addr, [NewItemOuter, Map, &Index, &bAddedMapEntry, &ReturnLogicalIndex](void* DirectAddress)
 							{
 								FScriptMapHelper	MapHelper(Map, DirectAddress);
 								Index = MapHelper.AddDefaultValue_Invalid_NeedsRehash();
@@ -1207,7 +1215,7 @@ int32 FPropertyValueImpl::AddChild()
 									FObjectProperty* KeyObjectProperty = CastField<FObjectProperty>(Map->KeyProp);
 									if (KeyObjectProperty && KeyObjectProperty->HasAnyPropertyFlags(CPF_NonNullable) && !KeyObjectProperty->PropertyClass->HasAnyClassFlags(CLASS_Abstract))
 									{
-										UObject* NewItem = NewObject<UObject>(Obj, KeyObjectProperty->PropertyClass);
+										UObject* NewItem = NewObject<UObject>(NewItemOuter, KeyObjectProperty->PropertyClass);
 										KeyObjectProperty->SetObjectPropertyValue(MapHelper.GetKeyPtr(Index), NewItem);
 									}
 								}
@@ -1216,7 +1224,7 @@ int32 FPropertyValueImpl::AddChild()
 									FObjectProperty* ValueObjectProperty = CastField<FObjectProperty>(Map->ValueProp);
 									if (ValueObjectProperty && ValueObjectProperty->HasAnyPropertyFlags(CPF_NonNullable) && !ValueObjectProperty->PropertyClass->HasAnyClassFlags(CLASS_Abstract))
 									{
-										UObject* NewItem = NewObject<UObject>(Obj, ValueObjectProperty->PropertyClass);
+										UObject* NewItem = NewObject<UObject>(NewItemOuter, ValueObjectProperty->PropertyClass);
 										ValueObjectProperty->SetObjectPropertyValue(MapHelper.GetValuePtr(Index), NewItem);
 									}
 								}
