Sequencer: Fix issue caused by my earlier change 36002464 where payload properties on director blueprint conditions would appear twice in the details panel. What was happening was:

1) There is a director blueprint condition property that is part of a condition container struct. The condition container is added as an external structure. Its director blueprint condition property has a customization that adds its payload variable properties.
2) The code in DetailLayoutBuilderImpl copies the properties from the external structure to the parent category. This includes both a) the director blueprint condition property, and b) the payload variables.
3) The newly added director blueprint condition property from #2 that got put in the parent category has its customization run again, and that customization again adds payload variables, meaning we have 2 copies of the payload variables added.

To fix this, in step 2, I check to see if the child nodes have an external root property node. If they do, I do not copy them to the parent category.

diff --git a/Engine/Source/Editor/PropertyEditor/Private/DetailLayoutBuilderImpl.cpp b/Engine/Source/Editor/PropertyEditor/Private/DetailLayoutBuilderImpl.cpp
--- a/Engine/Source/Editor/PropertyEditor/Private/DetailLayoutBuilderImpl.cpp
+++ b/Engine/Source/Editor/PropertyEditor/Private/DetailLayoutBuilderImpl.cpp
@@ -482,6 +482,12 @@ void FDetailLayoutBuilderImpl::GenerateDetailLayout()
 							continue;
 						}
 					}
+					else if (ChildNode->GetExternalRootPropertyNode()) 
+					{
+						// Also skip children that have been added externally, as those are likely generated from one of the other
+						// child nodes, and will be generated again upon moving that node to the outer category
+						continue;
+					}
 
 					// If there is no outer object then the class is the object root and there is only one instance
 					FName InstanceName = ParentStructPropertyName;
