Detail Panels- Change FDetailLayoutMap::ShouldShowGroup to respect ShowOnlyInnerProperties meta on structs.

diff --git a/Engine/Source/Editor/PropertyEditor/Private/DetailCategoryBuilderImpl.cpp b/Engine/Source/Editor/PropertyEditor/Private/DetailCategoryBuilderImpl.cpp
--- a/Engine/Source/Editor/PropertyEditor/Private/DetailCategoryBuilderImpl.cpp
+++ b/Engine/Source/Editor/PropertyEditor/Private/DetailCategoryBuilderImpl.cpp
@@ -1359,3 +1359,39 @@ bool FDetailCategoryImpl::IsParentEnabled() const
 	IDetailsViewPrivate* DetailsView = GetDetailsView();
 	return !DetailsView || DetailsView->IsPropertyEditingEnabled();
 }
+
+bool FDetailLayoutMap::ShouldShowGroup(FName RequiredGroupName) const
+{
+	auto HasInnerPropertiesOnly = [](const TArray<FDetailLayoutCustomization>& Layouts) -> bool
+		{
+			static const FName Name_ShowOnlyInnerProperties("ShowOnlyInnerProperties");
+
+			for (const FDetailLayoutCustomization& Customization : Layouts)
+			{
+				if (const TSharedPtr<FPropertyNode> PropertyNode = Customization.GetPropertyNode())
+				{
+					if (const TSharedPtr<FPropertyNode> ParentNode = PropertyNode->GetParentNodeSharedPtr())
+					{
+						if (const FProperty* ParentNodeProperty = ParentNode->GetProperty())
+						{
+							if (ParentNodeProperty->HasMetaData(Name_ShowOnlyInnerProperties))
+							{
+								return true;
+							}
+						}
+					}
+				}
+			}
+			return false;
+		};
+
+	bool bInnerPropertiesOnly = false;
+
+	if (const FDetailLayout* FoundLayout = Layouts.FindByPredicate([RequiredGroupName](const FDetailLayout& Layout) { return Layout.GetInstanceName() == RequiredGroupName; }))
+	{
+		bInnerPropertiesOnly = HasInnerPropertiesOnly(FoundLayout->GetSimpleLayouts()) || HasInnerPropertiesOnly(FoundLayout->GetAdvancedLayouts());
+	}
+
+	// Should show the group if the group name is not empty and there are more than two entries in the list where one of them is not the default "none" entry (represents the base object)
+	return !bInnerPropertiesOnly && RequiredGroupName != NAME_None && Layouts.Num() > 1 && (Layouts.Num() > 2 || !bContainsBaseInstance);
+}
diff --git a/Engine/Source/Editor/PropertyEditor/Private/DetailCategoryBuilderImpl.h b/Engine/Source/Editor/PropertyEditor/Private/DetailCategoryBuilderImpl.h
--- a/Engine/Source/Editor/PropertyEditor/Private/DetailCategoryBuilderImpl.h
+++ b/Engine/Source/Editor/PropertyEditor/Private/DetailCategoryBuilderImpl.h
@@ -145,11 +145,8 @@ public:
 	/**
 	 * @return Whether or not we need to display a group border around a list of details.
 	 */
-	bool ShouldShowGroup(FName RequiredGroupName) const
-	{
-		// Should show the group if the group name is not empty and there are more than two entries in the list where one of them is not the default "none" entry (represents the base object)
-		return RequiredGroupName != NAME_None && Layouts.Num() > 1 && (Layouts.Num() > 2 || !bContainsBaseInstance);
-	}
+	bool ShouldShowGroup(FName RequiredGroupName) const;
+
 private:
 	TArray<FDetailLayout> Layouts;
 	bool bContainsBaseInstance;
