Sequencer: Fix Mac/Linux build.

diff --git a/Engine/Source/Editor/PropertyEditor/Private/DetailCategoryBuilderImpl.cpp b/Engine/Source/Editor/PropertyEditor/Private/DetailCategoryBuilderImpl.cpp
--- a/Engine/Source/Editor/PropertyEditor/Private/DetailCategoryBuilderImpl.cpp
+++ b/Engine/Source/Editor/PropertyEditor/Private/DetailCategoryBuilderImpl.cpp
@@ -1362,11 +1362,11 @@ bool FDetailCategoryImpl::IsParentEnabled() const
 
 bool FDetailLayoutMap::ShouldShowGroup(FName RequiredGroupName) const
 {
-	auto HasInnerPropertiesOnly = [](const TArray<FDetailLayoutCustomization>& Layouts) -> bool
+	auto HasInnerPropertiesOnly = [](const TArray<FDetailLayoutCustomization>& CustomizationLayouts) -> bool
 		{
 			static const FName Name_ShowOnlyInnerProperties("ShowOnlyInnerProperties");
 
-			for (const FDetailLayoutCustomization& Customization : Layouts)
+			for (const FDetailLayoutCustomization& Customization : CustomizationLayouts)
 			{
 				if (const TSharedPtr<FPropertyNode> PropertyNode = Customization.GetPropertyNode())
 				{
