[BUGFIX] Interchange UI - "All results have been filtered..." message text is not aligned with the window
- SDetailsView correctly handles auto wrapped text.

diff --git a/Engine/Source/Editor/PropertyEditor/Private/SDetailsView.cpp b/Engine/Source/Editor/PropertyEditor/Private/SDetailsView.cpp
--- a/Engine/Source/Editor/PropertyEditor/Private/SDetailsView.cpp
+++ b/Engine/Source/Editor/PropertyEditor/Private/SDetailsView.cpp
@@ -449,12 +449,13 @@ void SDetailsView::Construct(const FArguments& InArgs, const FDetailsViewArgs& I
 			]
 		]
 		+ SOverlay::Slot()
-		.HAlign(HAlign_Center)
+		.HAlign(HAlign_Fill)
 		.Padding(2.0f, 24.0f, 2.0f, 2.0f)
 		[
 			SNew(STextBlock)
-			.Text(LOCTEXT("AllItemsFiltered", "All results have been filtered. Try changing your active filters above."))
+			.Text(LOCTEXT("AllItemsFiltered", "All results have been filtered.\nTry changing your active filters above."))
 			.AutoWrapText(true)
+			.Justification(ETextJustify::Center)
 			.Visibility_Lambda([this]() { return ((this->GetFilterBoxVisibility() == EVisibility::Visible) && !this->CurrentFilter.IsEmptyFilter() && RootTreeNodes.Num() == 0) ? EVisibility::HitTestInvisible : EVisibility::Collapsed; })
 		]
 		+ SOverlay::Slot()
