[PropertyEditor] Fixed batch copy/paste not persisting

#rnx

diff --git a/Engine/Source/Editor/PropertyEditor/Private/SDetailSingleItemRow.cpp b/Engine/Source/Editor/PropertyEditor/Private/SDetailSingleItemRow.cpp
--- a/Engine/Source/Editor/PropertyEditor/Private/SDetailSingleItemRow.cpp
+++ b/Engine/Source/Editor/PropertyEditor/Private/SDetailSingleItemRow.cpp
@@ -1521,11 +1521,10 @@ bool SDetailSingleItemRow::PasteFromText(const FString& InTag, const FString& In
 
 	const bool bIsTagged = !InTag.IsEmpty();
 
-	// If tagged, add the InteractiveChange flag so as not to run PECP
+	// If tagged, skip individual property transactions. Instead, a single undo will revert all changes in the batch paste.
 	// @todo: would be better to indicate that this is a batched paste rather than checking for a tag
 	if (bIsTagged)
 	{
-		PropertyValueSetFlags |= EPropertyValueSetFlags::InteractiveChange;
 		PropertyValueSetFlags |= EPropertyValueSetFlags::NotTransactable;
 	}
 
