Make PropertyNodes only resolve IDOs as an optional behaviour
* The new property node flag is inherited from the setting on the parent object, if set.  This makes it unecessary to change all places where the InitNode function is used.
* Added CVar DetailsPanel.ResolveInstanceDataObjects.Default to change this behaviour

diff --git a/Engine/Plugins/Experimental/InstanceDataObjectFixupTool/Source/Private/InstanceDataObjectFixupPanel.cpp b/Engine/Plugins/Experimental/InstanceDataObjectFixupTool/Source/Private/InstanceDataObjectFixupPanel.cpp
--- a/Engine/Plugins/Experimental/InstanceDataObjectFixupTool/Source/Private/InstanceDataObjectFixupPanel.cpp
+++ b/Engine/Plugins/Experimental/InstanceDataObjectFixupTool/Source/Private/InstanceDataObjectFixupPanel.cpp
@@ -369,6 +369,7 @@ TSharedPtr<IDetailsView>& FInstanceDataObjectFixupPanel::GenerateDetailsView(boo
 	DetailsViewArgs.ExternalScrollbar = SAssignNew(LinkableScrollBar, SLinkableScrollBar);
 	DetailsViewArgs.ScrollbarAlignment = bScrollbarOnLeft ? HAlign_Left : HAlign_Right;
 	DetailsViewArgs.DetailsNameWidgetOverrideCustomization = MakeShared<FInstanceDataObjectNameWidgetOverride>(SharedThis(this));
+	DetailsViewArgs.bResolveInstanceDataObjects.Emplace(true);
 	DetailsViewArgs.bShowLooseProperties = !HasViewFlag(EViewFlags::HideLooseProperties);
 
 	if (HasViewFlag(EViewFlags::IncludeOnlySetBySerialization))
diff --git a/Engine/Source/Editor/PropertyEditor/Private/ObjectPropertyNode.cpp b/Engine/Source/Editor/PropertyEditor/Private/ObjectPropertyNode.cpp
--- a/Engine/Source/Editor/PropertyEditor/Private/ObjectPropertyNode.cpp
+++ b/Engine/Source/Editor/PropertyEditor/Private/ObjectPropertyNode.cpp
@@ -55,10 +55,13 @@ const UPackage* FObjectPropertyNode::GetUPackage(int32 InIndex) const
 // Adds a new object to the list.
 void FObjectPropertyNode::AddObject( UObject* InObject )
 {
-	UE::FPropertyBagRepository& Repository = UE::FPropertyBagRepository::Get();
-	if (UObject* Found = Repository.FindInstanceDataObject(InObject))
+	if (HasNodeFlags(EPropertyNodeFlags::ResolveInstanceDataObjects))
 	{
-		InObject = Found;
+		UE::FPropertyBagRepository& Repository = UE::FPropertyBagRepository::Get();
+		if (UObject* Found = Repository.FindInstanceDataObject(InObject))
+		{
+			InObject = Found;
+		}
 	}
 	Objects.Add( InObject );
 }
@@ -83,10 +86,13 @@ void FObjectPropertyNode::RemoveObject( UObject* InObject )
 	}
 	else
 	{
-		UE::FPropertyBagRepository& Repository = UE::FPropertyBagRepository::Get();
-		if (UObject* Found = Repository.FindInstanceDataObject(InObject))
+		if (HasNodeFlags(EPropertyNodeFlags::ResolveInstanceDataObjects))
 		{
-			RemoveObject(Found);
+			UE::FPropertyBagRepository& Repository = UE::FPropertyBagRepository::Get();
+			if (UObject* Found = Repository.FindInstanceDataObject(InObject))
+			{
+				RemoveObject(Found);
+			}
 		}
 	}
 }
@@ -464,8 +470,31 @@ uint8* FObjectPropertyNode::GetValueBaseAddress(uint8* StartAddress, bool bIsSpa
 
 void FObjectPropertyNode::InitBeforeNodeFlags()
 {
+	// Note: InitBeforeNodeFlags is called _after_ initializing many node flags in FPropertyNode::InitNode
 	StoredProperty = Property;
 	Property = NULL;
+	
+	if (HasNodeFlags(EPropertyNodeFlags::ResolveInstanceDataObjects))
+	{
+		// Replace any objects that have been set with their IDO counterparts if they exist
+		for (int32 Index = 0, End = Objects.Num(); Index < End; ++Index)
+		{
+			TWeakObjectPtr<UObject> Object = Objects[Index];
+			TStrongObjectPtr<UObject> Pinned = Object.Pin();
+			if (Pinned)
+			{
+				if (HasNodeFlags(EPropertyNodeFlags::ResolveInstanceDataObjects))
+				{
+					UE::FPropertyBagRepository& Repository = UE::FPropertyBagRepository::Get();
+				
+					if (UObject* Found = Repository.FindInstanceDataObject(Pinned.Get()))
+					{
+						Objects[Index] = Found;
+					}					
+				}
+			}
+		}
+	}
 
 	Finalize();
 }
@@ -592,6 +621,7 @@ void FObjectPropertyNode::InternalInitChildNodes( FName SinglePropertyName )
 
 	const bool bShouldShowHiddenProperties = !!HasNodeFlags(EPropertyNodeFlags::ShouldShowHiddenProperties);
 	const bool bShouldShowDisableEditOnInstance = !!HasNodeFlags(EPropertyNodeFlags::ShouldShowDisableEditOnInstance);
+	const bool bResolveInstanceDataObjects = !!HasNodeFlags(EPropertyNodeFlags::ResolveInstanceDataObjects);
 
 	if (UClass* ResolvedBaseClass = BaseClass.Get())
 	{
diff --git a/Engine/Source/Editor/PropertyEditor/Private/PropertyNode.cpp b/Engine/Source/Editor/PropertyEditor/Private/PropertyNode.cpp
--- a/Engine/Source/Editor/PropertyEditor/Private/PropertyNode.cpp
+++ b/Engine/Source/Editor/PropertyEditor/Private/PropertyNode.cpp
@@ -280,10 +280,13 @@ void FPropertyNode::InitNode(const FPropertyNodeInitParams& InitParams)
 
 		// We are advanced if our parent is advanced or our property is marked as advanced
 		SetNodeFlags(EPropertyNodeFlags::IsAdvanced, ParentNode->HasNodeFlags(EPropertyNodeFlags::IsAdvanced) || bAdvanced );
+
+		SetNodeFlags(EPropertyNodeFlags::ResolveInstanceDataObjects, ParentNode->HasNodeFlags(EPropertyNodeFlags::ResolveInstanceDataObjects));
 	}
 	else
 	{
 		SetNodeFlags(EPropertyNodeFlags::ShowCategories, InitParams.bCreateCategoryNodes );
+		SetNodeFlags(EPropertyNodeFlags::ShowCategories, InitParams.bResolveInstanceDataObjects );
 	}
 
 	SetNodeFlags(EPropertyNodeFlags::ShouldShowHiddenProperties, InitParams.bForceHiddenPropertyVisibility);
@@ -833,10 +836,13 @@ EPropertyDataValidationResult FPropertyNode::EnsureDataIsValid()
 				if (ObjectProperty && !bIgnoreAllMismatch)
 				{
 					UObject* Obj = ObjectProperty->GetObjectPropertyValue(Addr);
-					UE::FPropertyBagRepository& Repository = UE::FPropertyBagRepository::Get();
-					if (UObject* Found = Repository.FindInstanceDataObject(Obj))
+					if (HasNodeFlags(EPropertyNodeFlags::ResolveInstanceDataObjects))
 					{
-						Obj = Found;
+						UE::FPropertyBagRepository& Repository = UE::FPropertyBagRepository::Get();
+						if (UObject* Found = Repository.FindInstanceDataObject(Obj))
+						{
+							Obj = Found;
+						}
 					}
 					if (IsValid(Obj))
 					{
diff --git a/Engine/Source/Editor/PropertyEditor/Private/PropertyNode.h b/Engine/Source/Editor/PropertyEditor/Private/PropertyNode.h
--- a/Engine/Source/Editor/PropertyEditor/Private/PropertyNode.h
+++ b/Engine/Source/Editor/PropertyEditor/Private/PropertyNode.h
@@ -77,6 +77,8 @@ namespace EPropertyNodeFlags
 
 	inline const Type	ShouldShowInViewport			= 1 << 25;	/** true if the property should be shown in the viewport context menu */
 
+	inline const Type   ResolveInstanceDataObjects      = 1 << 26;  /** Object properties will resolve to the instance data objects if they exist */
+
 	inline const Type 	NoFlags							= 0;
 
 };
@@ -239,6 +241,9 @@ struct FPropertyNodeInitParams
 	bool bCreateCategoryNodes;
 	/** Whether or not to create nodes for properties marked CPF_DisableEditOnInstance */
 	bool bCreateDisableEditOnInstanceNodes;
+	/** Resolve Objects to their IDOs */
+	bool bResolveInstanceDataObjects;
+	
 	/** Whether or not this property is sparse data */
 	EIsSparseDataProperty IsSparseProperty;
 
@@ -251,6 +256,7 @@ struct FPropertyNodeInitParams
 		, bForceHiddenPropertyVisibility( false )
 		, bCreateCategoryNodes( true )
 		, bCreateDisableEditOnInstanceNodes( true )
+		, bResolveInstanceDataObjects(false)
 		, IsSparseProperty( EIsSparseDataProperty::Inherit )
 	{}
 };
diff --git a/Engine/Source/Editor/PropertyEditor/Private/SDetailsView.cpp b/Engine/Source/Editor/PropertyEditor/Private/SDetailsView.cpp
--- a/Engine/Source/Editor/PropertyEditor/Private/SDetailsView.cpp
+++ b/Engine/Source/Editor/PropertyEditor/Private/SDetailsView.cpp
@@ -38,6 +38,16 @@ static TAutoConsoleVariable<bool> CVarDetailsPanelEnableCardLayout(
 	TEXT("Specifies whether the card layout is in effect for the Details View."),
 	ECVF_Default);
 
+namespace UE::PropertyEditor
+{
+	bool GPanelResolveInstanceDataObjectsDefault = false;
+	static FAutoConsoleVariableRef CVarDetailsPanelResolveInstanceDataObjectsDefault(
+	TEXT("DetailsPanel.ResolveInstanceDataObjects.Default"),
+		GPanelResolveInstanceDataObjectsDefault,
+		TEXT("Default value for the ResolveInstanceDataObjects DetailsView init parameter if not specified.  Will only apply to newly created details panels"),
+		ECVF_Default);
+}
+
 SDetailsView::~SDetailsView()
 {
 	const FRootPropertyNodeList& RootNodes = GetRootNodes();
@@ -1031,6 +1041,7 @@ void SDetailsView::PostSetObject(const TArray<FDetailsViewObjectRoot>& Roots)
 		FPropertySettings::Get().ShowHiddenProperties() || 
 		bShowHiddenPropertiesWhilePlaying ||
 		DetailsViewArgs.bForceHiddenPropertyVisibility;
+	InitParams.bResolveInstanceDataObjects = DetailsViewArgs.bResolveInstanceDataObjects.Get(UE::PropertyEditor::GPanelResolveInstanceDataObjectsDefault);
 
 	switch ( DetailsViewArgs.DefaultsOnlyVisibility )
 	{
diff --git a/Engine/Source/Editor/PropertyEditor/Public/DetailsViewArgs.h b/Engine/Source/Editor/PropertyEditor/Public/DetailsViewArgs.h
--- a/Engine/Source/Editor/PropertyEditor/Public/DetailsViewArgs.h
+++ b/Engine/Source/Editor/PropertyEditor/Public/DetailsViewArgs.h
@@ -122,6 +122,9 @@ struct FDetailsViewArgs
 	DECLARE_DELEGATE_RetVal_OneParam(bool, FShouldForceHideProperty, const TSharedRef<FPropertyNode>&);
 	FShouldForceHideProperty ShouldForceHideProperty;
 
+	/** Resolve any UObject* to their IDO counterparts if they exist */
+	TOptional<bool> bResolveInstanceDataObjects;
+
 public:
 
 	FDetailsViewArgs()
