Flip DetailsPanel.ResolveInstanceDataObjects.Default to true
* Fix bug in setting the ShowCategories flag with the wrong init value

#rnx

diff --git a/Engine/Source/Editor/PropertyEditor/Private/PropertyNode.cpp b/Engine/Source/Editor/PropertyEditor/Private/PropertyNode.cpp
--- a/Engine/Source/Editor/PropertyEditor/Private/PropertyNode.cpp
+++ b/Engine/Source/Editor/PropertyEditor/Private/PropertyNode.cpp
@@ -286,7 +286,7 @@ void FPropertyNode::InitNode(const FPropertyNodeInitParams& InitParams)
 	else
 	{
 		SetNodeFlags(EPropertyNodeFlags::ShowCategories, InitParams.bCreateCategoryNodes );
-		SetNodeFlags(EPropertyNodeFlags::ShowCategories, InitParams.bResolveInstanceDataObjects );
+		SetNodeFlags(EPropertyNodeFlags::ResolveInstanceDataObjects, InitParams.bResolveInstanceDataObjects );
 	}
 
 	SetNodeFlags(EPropertyNodeFlags::ShouldShowHiddenProperties, InitParams.bForceHiddenPropertyVisibility);
diff --git a/Engine/Source/Editor/PropertyEditor/Private/SDetailsView.cpp b/Engine/Source/Editor/PropertyEditor/Private/SDetailsView.cpp
--- a/Engine/Source/Editor/PropertyEditor/Private/SDetailsView.cpp
+++ b/Engine/Source/Editor/PropertyEditor/Private/SDetailsView.cpp
@@ -40,7 +40,7 @@ static TAutoConsoleVariable<bool> CVarDetailsPanelEnableCardLayout(
 
 namespace UE::PropertyEditor
 {
-	bool GPanelResolveInstanceDataObjectsDefault = false;
+	bool GPanelResolveInstanceDataObjectsDefault = true;
 	static FAutoConsoleVariableRef CVarDetailsPanelResolveInstanceDataObjectsDefault(
 	TEXT("DetailsPanel.ResolveInstanceDataObjects.Default"),
 		GPanelResolveInstanceDataObjectsDefault,
