Flip DetailsPanel.ResolveInstanceDataObjects.Default to false
* The bug was fixed in the same CL that this was flipped earlier.  SetNodeFlags was being called to set the ShowCategories flag inappropriately

diff --git a/Engine/Source/Editor/PropertyEditor/Private/SDetailsView.cpp b/Engine/Source/Editor/PropertyEditor/Private/SDetailsView.cpp
--- a/Engine/Source/Editor/PropertyEditor/Private/SDetailsView.cpp
+++ b/Engine/Source/Editor/PropertyEditor/Private/SDetailsView.cpp
@@ -40,7 +40,7 @@ static TAutoConsoleVariable<bool> CVarDetailsPanelEnableCardLayout(
 
 namespace UE::PropertyEditor
 {
-	bool GPanelResolveInstanceDataObjectsDefault = true;
+	bool GPanelResolveInstanceDataObjectsDefault = false;
 	static FAutoConsoleVariableRef CVarDetailsPanelResolveInstanceDataObjectsDefault(
 	TEXT("DetailsPanel.ResolveInstanceDataObjects.Default"),
 		GPanelResolveInstanceDataObjectsDefault,
