Fix advanced category not expanding in objects' details added through external objects

CL 26991755 added lazy refresh to FDetailCategoryImpl which introduced a regression that causes categories added through AddExtenalObjects to not get refreshed when RefreshTree is called. This is because RefreshTree adds the category node to the parent layout builder's list of nodes to tick, but DetailLayoutBuilderImpl's tick method isn't called for external objects.

You can see this behavior when adding a LiveLink component to an actor and selecting a subject, then trying to expand this object's advanced category.

diff --git a/Engine/Source/Editor/PropertyEditor/Private/DetailCategoryBuilderImpl.cpp b/Engine/Source/Editor/PropertyEditor/Private/DetailCategoryBuilderImpl.cpp
--- a/Engine/Source/Editor/PropertyEditor/Private/DetailCategoryBuilderImpl.cpp
+++ b/Engine/Source/Editor/PropertyEditor/Private/DetailCategoryBuilderImpl.cpp
@@ -752,9 +752,19 @@ void FDetailCategoryImpl::Tick(float DeltaTime)
 
 void FDetailCategoryImpl::RefreshTree(bool bRefilterCategory)
 {
-	bPendingRefresh = true;
-	bPendingRefreshNeedsRefilter = bRefilterCategory;
-	AddTickableNode(*this);
+	TSharedPtr<FDetailLayoutBuilderImpl> ParentLayout = GetParentLayoutImpl();
+
+	// If this is an external property node, refresh it directly since external objects' detail layout builder don't get ticked.
+	if (ParentLayout.IsValid() && ParentLayout->IsLayoutForExternalRoot())
+	{
+		RefreshTreeInternal(bRefilterCategory);
+	}
+	else
+	{
+		bPendingRefresh = true;
+		bPendingRefreshNeedsRefilter = bRefilterCategory;
+		AddTickableNode(*this);
+	}
 }
 
 void FDetailCategoryImpl::RefreshTreeInternal(bool bRefilterCategory)
