[Bugfix] FPropertyNode::NotifyPostChange was overwriting existing array indices. Now it only attempts to generate them if they already don't exist

diff --git a/Engine/Source/Editor/PropertyEditor/Private/PropertyHandleImpl.cpp b/Engine/Source/Editor/PropertyEditor/Private/PropertyHandleImpl.cpp
--- a/Engine/Source/Editor/PropertyEditor/Private/PropertyHandleImpl.cpp
+++ b/Engine/Source/Editor/PropertyEditor/Private/PropertyHandleImpl.cpp
@@ -3326,6 +3326,32 @@ void FPropertyHandleBase::NotifyPreChange()
 	}
 }
 
+static TArray<TMap<FString, int32>> BuildArrayIndices(const TSharedPtr<FPropertyNode>& PropertyNode)
+{
+	TArray<TMap<FString, int32>> ArrayIndices;
+	FReadAddressList ReadAddresses;
+	PropertyNode->GetReadAddress(PropertyNode->HasNodeFlags(EPropertyNodeFlags::SingleSelectOnly), ReadAddresses, false, true);
+
+	const int32 ObjectNum = ReadAddresses.Num();
+	if (ArrayIndices.IsEmpty())
+	{
+		ArrayIndices.SetNum(ObjectNum);
+	}
+	
+	for (const FPropertyNode* ItemNode = PropertyNode.Get(); ItemNode; ItemNode = ItemNode->GetParentNode())
+	{
+		if (int32 Index = ItemNode->GetArrayIndex(); Index != INDEX_NONE)
+		{
+			for (int ObjectIndex = 0; ObjectIndex < ObjectNum; ++ObjectIndex)
+			{
+				ArrayIndices[ObjectIndex].Add(ItemNode->GetProperty()->GetName(), Index);
+			}
+		}
+	} 
+
+	return ArrayIndices;
+}
+
 void FPropertyHandleBase::NotifyPostChange( EPropertyChangeType::Type ChangeType )
 {
 	TSharedPtr<FPropertyNode> PropertyNode = Implementation->GetPropertyNode();
@@ -3343,6 +3369,8 @@ void FPropertyHandleBase::NotifyPostChange( EPropertyChangeType::Type ChangeType
 		}
 
 		FPropertyChangedEvent PropertyChangedEvent( PropertyNode->GetProperty(), ChangeType, MakeArrayView(ObjectsBeingChanged) );
+		TArray<TMap<FString, int32>> ArrayIndices = BuildArrayIndices(PropertyNode);
+		PropertyChangedEvent.SetArrayIndexPerObject(ArrayIndices);
 		PropertyNode->NotifyPostChange( PropertyChangedEvent, Implementation->GetNotifyHook());
 	}
 }
diff --git a/Engine/Source/Editor/PropertyEditor/Private/PropertyNode.cpp b/Engine/Source/Editor/PropertyEditor/Private/PropertyNode.cpp
--- a/Engine/Source/Editor/PropertyEditor/Private/PropertyNode.cpp
+++ b/Engine/Source/Editor/PropertyEditor/Private/PropertyNode.cpp
@@ -3152,7 +3152,6 @@ void FPropertyNode::NotifyPreChangeInternal( TSharedRef<FEditPropertyChain> Prop
 void FPropertyNode::NotifyPostChange( FPropertyChangedEvent& InPropertyChangedEvent, class FNotifyHook* InNotifyHook )
 {
 	TSharedRef<FEditPropertyChain> PropertyChain = BuildPropertyChain( InPropertyChangedEvent.Property );
-	TArray<TMap<FString, int32>> ArrayIndices = BuildArrayIndices();
 	
 	// remember the property that was the chain's original active property; this will correspond to the outermost property of struct/array that was modified
 	FProperty* const OriginalActiveProperty = PropertyChain->GetActiveMemberNode() ? PropertyChain->GetActiveMemberNode()->GetValue() : nullptr;
@@ -3212,7 +3211,7 @@ void FPropertyNode::NotifyPostChange( FPropertyChangedEvent& InPropertyChangedEv
 
 				// Use a scope to ensure that only local variable are use in the loop.
 				// Since this object can be destroyed in this loop.
-				auto ScopePostEditChange = [&PropertyChain, &InPropertyChangedEvent, &CurProperty, CurrentObjectIndex, &ArrayIndices](UObject* Object)
+				auto ScopePostEditChange = [&PropertyChain, &InPropertyChangedEvent, &CurProperty, CurrentObjectIndex](UObject* Object)
 				{
 					// copy the property changed event
 					FPropertyChangedEvent ChangedEvent = InPropertyChangedEvent;
@@ -3232,10 +3231,6 @@ void FPropertyNode::NotifyPostChange( FPropertyChangedEvent& InPropertyChangedEv
 					else
 					{
 						FPropertyChangedChainEvent ChainEvent(*PropertyChain, ChangedEvent);
-						if (!ArrayIndices.IsEmpty())
-						{
-							ChainEvent.SetArrayIndexPerObject(ArrayIndices);
-						}
 						ChainEvent.ObjectIteratorIndex = CurrentObjectIndex;
 
 						Object->PostEditChangeChainProperty(ChainEvent);
@@ -3570,31 +3565,6 @@ TSharedRef<FEditPropertyChain> FPropertyNode::BuildPropertyChain( FProperty* InP
 	return PropertyChain;
 }
 
-TArray<TMap<FString, int32>> FPropertyNode::BuildArrayIndices() const
-{
-	TArray<TMap<FString, int32>> ArrayIndices;
-	FReadAddressList ReadAddresses;
-    GetReadAddress(HasNodeFlags(EPropertyNodeFlags::SingleSelectOnly), ReadAddresses, false, true);
-	const int32 ObjectNum = ReadAddresses.Num();
-
-	for (const FPropertyNode* ItemNode = this; ItemNode; ItemNode = ItemNode->GetParentNode())
-	{
-		if (int32 Index = ItemNode->GetArrayIndex(); Index != INDEX_NONE)
-		{
-			if (ArrayIndices.IsEmpty())
-            {
-                ArrayIndices.SetNum(ObjectNum);
-            }
-			for (int ObjectIndex = 0; ObjectIndex < ObjectNum; ++ObjectIndex)
-			{
-                ArrayIndices[ObjectIndex].Add(ItemNode->Property->GetName(), Index);
-			}
-		}
-	} 
-
-	return ArrayIndices;
-}
-
 TSharedRef<FEditPropertyChain> FPropertyNode::BuildPropertyChain( FProperty* InProperty, const TSet<UObject*>& InAffectedArchetypeInstances ) const
 {
 	TSharedRef<FEditPropertyChain> PropertyChain = BuildPropertyChain(InProperty);
diff --git a/Engine/Source/Editor/PropertyEditor/Private/PropertyNode.h b/Engine/Source/Editor/PropertyEditor/Private/PropertyNode.h
--- a/Engine/Source/Editor/PropertyEditor/Private/PropertyNode.h
+++ b/Engine/Source/Editor/PropertyEditor/Private/PropertyNode.h
@@ -1124,7 +1124,6 @@ public:
 	}
 
 	TSharedRef<FEditPropertyChain> BuildPropertyChain( FProperty* PropertyAboutToChange ) const;
-	TArray<TMap<FString, int32>> BuildArrayIndices() const;
 	
 protected:
 	TSharedRef<FEditPropertyChain> BuildPropertyChain( FProperty* PropertyAboutToChange, const TSet<UObject*>& InAffectedArchetypeInstances ) const;
