PropertyEditor: Allowed GetAllowedClasses metadata to use fully-qualified paths to function libraries, enabling use on USTRUCTs

diff --git a/Engine/Source/Editor/PropertyEditor/Private/PropertyEditorUtils.cpp b/Engine/Source/Editor/PropertyEditor/Private/PropertyEditorUtils.cpp
--- a/Engine/Source/Editor/PropertyEditor/Private/PropertyEditorUtils.cpp
+++ b/Engine/Source/Editor/PropertyEditor/Private/PropertyEditorUtils.cpp
@@ -159,54 +159,77 @@ namespace PropertyEditorUtils
 			const FString GetAllowedClassesFunctionName = MetadataProperty.GetOwnerProperty()->GetMetaData("GetAllowedClasses");
 			if (!GetAllowedClassesFunctionName.IsEmpty())
 			{
-				for (UObject* Object : ObjectList)
+				auto GetAllowedClasses = [&bMergeAllowedClasses, &AllowedClasses, &bExactClass, &DisallowedClasses, ObjectClass](UObject* InObject, const UFunction* InGetAllowedClassesFunction)-> bool
 				{
-					const UFunction* GetAllowedClassesFunction = Object ? Object->FindFunction(*GetAllowedClassesFunctionName) : nullptr;
-					if (GetAllowedClassesFunction)
+					DECLARE_DELEGATE_RetVal(TArray<UClass*>, FGetAllowedClasses);
+					if (!bMergeAllowedClasses)
 					{
-						DECLARE_DELEGATE_RetVal(TArray<UClass*>, FGetAllowedClasses);
-						if (!bMergeAllowedClasses)
+						AllowedClasses.Append(FGetAllowedClasses::CreateUFunction(InObject, InGetAllowedClassesFunction->GetFName()).Execute());
+						if (AllowedClasses.IsEmpty())
 						{
-							AllowedClasses.Append(FGetAllowedClasses::CreateUFunction(Object, GetAllowedClassesFunction->GetFName()).Execute());
-							if (AllowedClasses.IsEmpty())
-							{
-								// No allowed class means all classes are valid
-								continue;
-							}
-							bMergeAllowedClasses = true;
+							// No allowed class means all classes are valid
+							return true;
 						}
-						else
+						bMergeAllowedClasses = true;
+					}
+					else
+					{
+						TArray<UClass*> MergedClasses = FGetAllowedClasses::CreateUFunction(InObject, InGetAllowedClassesFunction->GetFName()).Execute();
+						if (MergedClasses.IsEmpty())
 						{
-							TArray<UClass*> MergedClasses = FGetAllowedClasses::CreateUFunction(Object, GetAllowedClassesFunction->GetFName()).Execute();
-							if (MergedClasses.IsEmpty())
-							{
-								// No allowed class means all classes are valid
-								continue;
-							}
-							
-							TArray<const UClass*> CurrentAllowedClassFilters = MoveTemp(AllowedClasses);
-							ensure(AllowedClasses.IsEmpty());
-							for (const UClass* MergedClass : MergedClasses)
+							// No allowed class means all classes are valid
+							return true;
+						}
+
+						TArray<const UClass*> CurrentAllowedClassFilters = MoveTemp(AllowedClasses);
+						ensure(AllowedClasses.IsEmpty());
+						for (const UClass* MergedClass : MergedClasses)
+						{
+							// Keep classes that match both allow list
+							for (const UClass* CurrentClass : CurrentAllowedClassFilters)
 							{
-								// Keep classes that match both allow list
-								for (const UClass* CurrentClass : CurrentAllowedClassFilters)
+								if (CurrentClass == MergedClass || (!bExactClass && CurrentClass->IsChildOf(MergedClass)))
 								{
-									if (CurrentClass == MergedClass || (!bExactClass && CurrentClass->IsChildOf(MergedClass)))
-									{
-										AllowedClasses.Add(CurrentClass);
-										break;
-									}
-									if (!bExactClass && MergedClass->IsChildOf(CurrentClass))
-									{
-										AllowedClasses.Add(MergedClass);
-										break;
-									}
+									AllowedClasses.Add(CurrentClass);
+									break;
+								}
+								if (!bExactClass && MergedClass->IsChildOf(CurrentClass))
+								{
+									AllowedClasses.Add(MergedClass);
+									break;
 								}
 							}
-							if (AllowedClasses.IsEmpty())
+						}
+						if (AllowedClasses.IsEmpty())
+						{
+							// An empty AllowedClasses array means that everything is allowed: in this case, forbid UObject
+							DisallowedClasses.Add(ObjectClass);
+							return false;
+						}
+					}
+					return true;
+				};
+
+				// First look for a library function assuming fully qualified path, e.g. /Script/ModuleName.ClassName:FunctionName
+				if(GetAllowedClassesFunctionName.Contains(TEXT(".")))
+				{
+					const UFunction* GetAllowedClassesFunction = FindObject<UFunction>(nullptr, *GetAllowedClassesFunctionName, true);
+					if (ensureMsgf(GetAllowedClassesFunction && GetAllowedClassesFunction->HasAnyFunctionFlags(EFunctionFlags::FUNC_Static), TEXT("Invalid GetAllowedClasses: %s"), *GetAllowedClassesFunctionName))
+					{
+						UObject* GetAllowedClassesCDO = GetAllowedClassesFunction->GetOwnerClass()->GetDefaultObject();
+						GetAllowedClasses(GetAllowedClassesCDO, GetAllowedClassesFunction);
+					}
+				}
+				else
+				{
+					// Otherwise interrogate the object list
+					for (UObject* Object : ObjectList)
+					{
+						const UFunction* GetAllowedClassesFunction = Object ? Object->FindFunction(*GetAllowedClassesFunctionName) : nullptr;
+						if (GetAllowedClassesFunction)
+						{
+							if(!GetAllowedClasses(Object, GetAllowedClassesFunction))
 							{
-								// An empty AllowedClasses array means that everything is allowed: in this case, forbid UObject
-								DisallowedClasses.Add(ObjectClass);
 								return;
 							}
 						}
diff --git a/Engine/Source/Editor/PropertyEditor/Private/PropertyHandleImpl.cpp b/Engine/Source/Editor/PropertyEditor/Private/PropertyHandleImpl.cpp
--- a/Engine/Source/Editor/PropertyEditor/Private/PropertyHandleImpl.cpp
+++ b/Engine/Source/Editor/PropertyEditor/Private/PropertyHandleImpl.cpp
@@ -4295,30 +4295,53 @@ FPropertyAccess::Result FPropertyHandleObject::SetValueFromFormattedString(const
 				const FString& GetAllowedClassesFunctionName = NodeProperty->GetMetaData("GetAllowedClasses");
 				if (!GetAllowedClassesFunctionName.IsEmpty() && NodeProperty->GetOwnerUObject())
 				{
-					TArray<UObject*> OuterObjects;
-					GetOuterObjects(OuterObjects);
-					for (UObject* Object : OuterObjects)
+					auto GetAllowedClasses = [&QualifiedObject, &QualifiedClass, &bSupportedObject](UObject* InObject, const UFunction* InGetAllowedClassesFunction) -> bool
 					{
-						const UFunction* GetAllowedClassesFunction = Object->FindFunction(*GetAllowedClassesFunctionName);
-						if (GetAllowedClassesFunction)
+						DECLARE_DELEGATE_RetVal(TArray<UClass*>, FGetAllowedClasses);
+						TArray<UClass*> AllowedClasses = FGetAllowedClasses::CreateUFunction(InObject, InGetAllowedClassesFunction->GetFName()).Execute();
+						if (AllowedClasses.Num() > 0)
 						{
-							DECLARE_DELEGATE_RetVal(TArray<UClass*>, FGetAllowedClasses);
-							TArray<UClass*> AllowedClasses = FGetAllowedClasses::CreateUFunction(Object, GetAllowedClassesFunction->GetFName()).Execute();
-							if (AllowedClasses.Num() > 0)
+							bSupportedObject = false;
+							for (const UClass* AllowedClass : AllowedClasses)
 							{
-								bSupportedObject = false;
-								for (const UClass* AllowedClass : AllowedClasses)
+								const bool bIsInterface = AllowedClass && AllowedClass->HasAnyClassFlags(CLASS_Interface);
+			
+								// Check if the object is an allowed class type this property supports
+								if ((AllowedClass && QualifiedClass->IsChildOf(AllowedClass)) || (bIsInterface && QualifiedObject->GetClass()->ImplementsInterface(AllowedClass)))
 								{
-									const bool bIsInterface = AllowedClass && AllowedClass->HasAnyClassFlags(CLASS_Interface);
-				
-									// Check if the object is an allowed class type this property supports
-									if ((AllowedClass && QualifiedClass->IsChildOf(AllowedClass)) || (bIsInterface && QualifiedObject->GetClass()->ImplementsInterface(AllowedClass)))
-									{
-										bSupportedObject = true;
-										break;
-									}
+									bSupportedObject = true;
+									break;
 								}
-								if (!bSupportedObject)
+							}
+							if (!bSupportedObject)
+							{
+								return false;
+							}
+						}
+
+						return true;
+					};
+
+					// First look for a library function assuming fully qualified path, e.g. /Script/ModuleName.ClassName:FunctionName
+					const UFunction* LibraryFunction = FindObject<UFunction>(nullptr, *GetAllowedClassesFunctionName);
+					if(LibraryFunction)
+					{
+						if(UObject* CDO = LibraryFunction->GetOwnerClass()->GetDefaultObject())
+						{
+							GetAllowedClasses(CDO, LibraryFunction);
+						}
+					}
+					else
+					{
+						// Otherwise interrogate the object list
+						TArray<UObject*> OuterObjects;
+						GetOuterObjects(OuterObjects);
+						for (UObject* Object : OuterObjects)
+						{
+							const UFunction* GetAllowedClassesFunction = Object->FindFunction(*GetAllowedClassesFunctionName);
+							if (GetAllowedClassesFunction)
+							{
+								if(!GetAllowedClasses(Object, GetAllowedClassesFunction))
 								{
 									break;
 								}
