[Details Panel] Fix a crash when a TOptional uproperty is used with an abstract base class

Note - doesn't actually fix the underlying issue (to be fixed later), just makes it less crashy in the short term.

#rnx

diff --git a/Engine/Source/Editor/PropertyEditor/Private/PropertyHandleImpl.cpp b/Engine/Source/Editor/PropertyEditor/Private/PropertyHandleImpl.cpp
--- a/Engine/Source/Editor/PropertyEditor/Private/PropertyHandleImpl.cpp
+++ b/Engine/Source/Editor/PropertyEditor/Private/PropertyHandleImpl.cpp
@@ -5378,7 +5378,7 @@ FPropertyAccess::Result FPropertyHandleOptional::SetOptionalValue(FProperty* New
 			void* ClassPropertyValuePtr = ClassProperty->ContainerPtrToValuePtr<void>(Addr);
 			ClassProperty->SetObjectPropertyValue(ClassPropertyValuePtr, ClassProperty->MetaClass);
 		}
-		else if (ObjectProperty)
+		else if (ObjectProperty && !ObjectProperty->PropertyClass->HasAnyClassFlags(EClassFlags::CLASS_Abstract))
 		{
 			UObject* Outer = Obj;
 
