Disable None option of a class picker if property is NonNullable

diff --git a/Engine/Source/Editor/PropertyEditor/Private/UserInterface/PropertyEditor/SPropertyEditorEditInline.cpp b/Engine/Source/Editor/PropertyEditor/Private/UserInterface/PropertyEditor/SPropertyEditorEditInline.cpp
--- a/Engine/Source/Editor/PropertyEditor/Private/UserInterface/PropertyEditor/SPropertyEditorEditInline.cpp
+++ b/Engine/Source/Editor/PropertyEditor/Private/UserInterface/PropertyEditor/SPropertyEditorEditInline.cpp
@@ -337,7 +337,7 @@ TSharedRef<SWidget> SPropertyEditorEditInline::GenerateClassPicker(TSharedRef<IP
 			bContainerHasNoClear = ParentNode->GetProperty()->HasAllPropertyFlags(CPF_NoClear);
 		}
 	}
-	Options.bShowNoneOption = !Property->HasAllPropertyFlags(CPF_NoClear) && !bContainerHasNoClear;
+	Options.bShowNoneOption = !Property->HasAllPropertyFlags(CPF_NoClear) && !bContainerHasNoClear && !Property->HasAllPropertyFlags(CPF_NonNullable);
 
 	FObjectPropertyNode* ObjectPropertyNode = PropertyNode->FindObjectItemParent();
 	if (ObjectPropertyNode)
