Added IDetailsView::GetRightColumnMinWidth so the width of the right column can be increased and decreased in a similar way that sequencer does in SDetailsViewBase::SetKeyframeHandler.

diff --git a/Engine/Source/Editor/PropertyEditor/Private/SDetailsViewBase.cpp b/Engine/Source/Editor/PropertyEditor/Private/SDetailsViewBase.cpp
--- a/Engine/Source/Editor/PropertyEditor/Private/SDetailsViewBase.cpp
+++ b/Engine/Source/Editor/PropertyEditor/Private/SDetailsViewBase.cpp
@@ -468,6 +468,11 @@ void SDetailsViewBase::GetHeadNodes(TArray<TWeakPtr<FDetailTreeNode>>& OutNodes)
 	}
 }
 
+float SDetailsViewBase::GetRightColumnMinWidth() const
+{
+	return ColumnSizeData.GetRightColumnMinWidth().Get(0);
+}
+
 void SDetailsViewBase::SetRightColumnMinWidth(float InMinWidth)
 {
 	ColumnSizeData.SetRightColumnMinWidth(InMinWidth);
diff --git a/Engine/Source/Editor/PropertyEditor/Private/SDetailsViewBase.h b/Engine/Source/Editor/PropertyEditor/Private/SDetailsViewBase.h
--- a/Engine/Source/Editor/PropertyEditor/Private/SDetailsViewBase.h
+++ b/Engine/Source/Editor/PropertyEditor/Private/SDetailsViewBase.h
@@ -129,6 +129,7 @@ public:
 	{
 		CustomValidatePropertyNodesFunction = MoveTemp(InCustomValidatePropertyNodesFunction);
 	}
+	virtual float GetRightColumnMinWidth() const override;
 	virtual void SetRightColumnMinWidth(float InMinWidth) override;
 	
 	/** IDetailsViewPrivate interface */
diff --git a/Engine/Source/Editor/PropertyEditor/Public/IDetailsView.h b/Engine/Source/Editor/PropertyEditor/Public/IDetailsView.h
--- a/Engine/Source/Editor/PropertyEditor/Public/IDetailsView.h
+++ b/Engine/Source/Editor/PropertyEditor/Public/IDetailsView.h
@@ -299,6 +299,8 @@ public:
 	 * This is useful if your implementation doesn't need to validate nodes every tick or needs to perform some other form of validation. */
 	virtual void SetCustomValidatePropertyNodesFunction(FOnValidateDetailsViewPropertyNodes InCustomValidatePropertyNodesFunction) = 0;
 
+	/** Gets the minimum width of the detail panel's right column. */
+	virtual float GetRightColumnMinWidth() const = 0;
 	/** Sets the minimum width of the detail panel's right column. */
 	virtual void SetRightColumnMinWidth(float InMinWidth) = 0;
 
