UE-227696 - "TOptional Details customization should handle abstract base classes and interfaces"

- sets Verse `script device` object customization to display value editor
- adds UI for FObjectProperty/FClassProperty optional's which dispalys a class-picker to set a value from (thus, allowing abstract/interface types to work with optionals)
- updates verse property generation code to set metadata for interface FClassProperty's and FObjectProperty's properly

(job 6719313ee327078c3e32c52c, errors were pre-existing in base)
#tests SceneGraph
(job 6719335d3802ce37a35285ad, compared against base job 67194ce74660451b58993c85 to confirm errors were pre-existing)

diff --git a/Engine/Source/Editor/PropertyEditor/Private/Presentation/PropertyEditor/PropertyEditor.cpp b/Engine/Source/Editor/PropertyEditor/Private/Presentation/PropertyEditor/PropertyEditor.cpp
--- a/Engine/Source/Editor/PropertyEditor/Private/Presentation/PropertyEditor/PropertyEditor.cpp
+++ b/Engine/Source/Editor/PropertyEditor/Private/Presentation/PropertyEditor/PropertyEditor.cpp
@@ -330,10 +330,10 @@ void FPropertyEditor::OnAddGivenItem(const FString InGivenItem)
 	}
 }
 
-void FPropertyEditor::SetOptionalItem(FProperty* NewValue)
+void FPropertyEditor::SetOptionalItem(FProperty* NewProperty, const UClass* NewObjectClass)
 {
 	// This action must be deferred until next tick so that we avoid accessing invalid data before we have a chance to tick
-	PropertyUtilities->EnqueueDeferredAction(FSimpleDelegate::CreateSP(this, &FPropertyEditor::OnSetOptionalValue, NewValue));
+	PropertyUtilities->EnqueueDeferredAction(FSimpleDelegate::CreateSP(this, &FPropertyEditor::OnSetOptionalValue, NewProperty, NewObjectClass));
 }
 
 void FPropertyEditor::ClearOptionalItem()
@@ -342,24 +342,30 @@ void FPropertyEditor::ClearOptionalItem()
 	PropertyUtilities->EnqueueDeferredAction(FSimpleDelegate::CreateSP(this, &FPropertyEditor::OnClearOptionalValue));
 }
 
-void FPropertyEditor::OnSetOptionalValue(FProperty* NewValue)
+void FPropertyEditor::OnSetOptionalValue(FProperty* NewProperty, const UClass* NewObjectClass)
 {
 	TSharedPtr<IPropertyHandleOptional> OptionalHandle = PropertyHandle->AsOptional();
+	if (!OptionalHandle.IsValid())
+	{
+		// This could be the value calling as we share UI between options and their value
+		OptionalHandle = PropertyHandle->GetParentHandle()->AsOptional();
+	}
+
 	if (OptionalHandle.IsValid())
 	{
-		OptionalHandle->SetOptionalValue(NewValue);
+		OptionalHandle->SetOptionalValue(NewProperty, NewObjectClass);
 	}
 }
 
 void FPropertyEditor::OnClearOptionalValue()
 {
-	// This func is called by an options value, not the option itself.
-	if (!PropertyNode->IsOptionalValueNode())
+	TSharedPtr<IPropertyHandleOptional> OptionalHandle = PropertyHandle->AsOptional();
+	if (!OptionalHandle.IsValid())
 	{
-		return;
+		// This could be the value calling as we share UI between options and their value
+		OptionalHandle = PropertyHandle->GetParentHandle()->AsOptional();
 	}
 
-	TSharedPtr<IPropertyHandleOptional> OptionalHandle = PropertyHandle->GetParentHandle()->AsOptional();
 	if (OptionalHandle.IsValid())
 	{
 		OptionalHandle->ClearOptionalValue();
diff --git a/Engine/Source/Editor/PropertyEditor/Private/Presentation/PropertyEditor/PropertyEditor.h b/Engine/Source/Editor/PropertyEditor/Private/Presentation/PropertyEditor/PropertyEditor.h
--- a/Engine/Source/Editor/PropertyEditor/Private/Presentation/PropertyEditor/PropertyEditor.h
+++ b/Engine/Source/Editor/PropertyEditor/Private/Presentation/PropertyEditor/PropertyEditor.h
@@ -96,7 +96,7 @@ public:
 	* If this is an editor of an optional, set that optional to the passed in value.
 	* If nullptr is passed in, the optional is set and default-initialized.
 	*/
-	void SetOptionalItem(FProperty* NewValue);
+	void SetOptionalItem(FProperty* NewProperty, const UClass* NewObjectClass);
 
 	/**	If this is an editor of an optional, unset that optional. */
 	void ClearOptionalItem();
@@ -121,7 +121,7 @@ private:
 	void OnDuplicateItem();
 	void OnBrowseTo();
 	void OnEmptyArray();
-	void OnSetOptionalValue(FProperty* NewValue);
+	void OnSetOptionalValue(FProperty* NewProperty, const UClass* NewObjectClass);
 	void OnClearOptionalValue();
 
 private:
diff --git a/Engine/Source/Editor/PropertyEditor/Private/PropertyCustomizationHelpers.cpp b/Engine/Source/Editor/PropertyEditor/Private/PropertyCustomizationHelpers.cpp
--- a/Engine/Source/Editor/PropertyEditor/Private/PropertyCustomizationHelpers.cpp
+++ b/Engine/Source/Editor/PropertyEditor/Private/PropertyCustomizationHelpers.cpp
@@ -38,6 +38,7 @@
 #include "Kismet2/BlueprintEditorUtils.h"
 #include "ObjectTools.h"
 #include "PropertyPermissionList.h"
+#include "PropertyNode.h"
 #include "Reflection/FunctionUtils.h"
 #include "ScopedTransaction.h"
 #include "Settings/BlueprintEditorProjectSettings.h"
@@ -197,18 +198,19 @@ namespace PropertyCustomizationHelpers
 			.IsFocusable( false );
 	}
 
-	TSharedRef<SWidget> MakeSetOptionalButton(FSimpleDelegate OnSetOptionalClicked, TAttribute<FText> OptionalToolTipText, TAttribute<bool> IsEnabled)
+	TSharedRef<SWidget> MakeSetOptionalButton(FOnSetOptional OnSetOptional, TAttribute<FText> OptionalToolTipText, TAttribute<bool> IsEnabled)
 	{
 		// Custom widget for this button as it has no image and should fill a larger space
 		return SNew(SBox)
 			.HAlign(HAlign_Left)
 			.VAlign(VAlign_Center)
+			.IsEnabled(IsEnabled)
 			.ToolTipText(OptionalToolTipText.Get().IsEmpty() ? LOCTEXT("SetOptionalButtonToolTipText", "Set Optional to default value.") : OptionalToolTipText)
 			[
 				SNew(SButton)
 				.ButtonStyle(FAppStyle::Get(), "Button")
-				.OnClicked_Lambda([OnSetOptionalClicked](){
-					OnSetOptionalClicked.ExecuteIfBound();
+				.OnClicked_Lambda([OnSetOptional](){
+					OnSetOptional.ExecuteIfBound(nullptr, nullptr);
 					return FReply::Handled();
 				})
 				.Text(LOCTEXT("SetButtonText", "Set to Value"))
@@ -217,13 +219,75 @@ namespace PropertyCustomizationHelpers
 			];
 	}
 
-	TSharedRef<SWidget> MakeClearOptionalButton(FSimpleDelegate OnClearOptionalClicked, TAttribute<FText> OptionalToolTipText, TAttribute<bool> IsEnabled)
+	TSharedRef<SWidget> MakePickOptionalButton(FOnSetOptional OnSetOptional, FSimpleDelegate OnClearOptional, const TSharedRef<FPropertyNode>& PropertyNode, TAttribute<FText> OptionalToolTipText, TAttribute<bool> IsEnabled)
+	{
+		const FOptionalProperty* OptionalProperty = CastField<FOptionalProperty>(PropertyNode->GetProperty());
+		const FProperty* ValueProperty = OptionalProperty->GetValueProperty();
+
+		const UClass* MetaClass = UObject::StaticClass();
+		const UClass* RequiredInterface = nullptr;
+		const UClass* SelectedClass = nullptr;
+
+		if (const FClassProperty* ClassProperty = CastField<FClassProperty>(ValueProperty))
+		{
+			MetaClass = ClassProperty->MetaClass;
+			RequiredInterface = ClassProperty->GetOwnerProperty()->GetClassMetaData(TEXT("MustImplement"));
+		}
+		else if(const FInterfaceProperty* InterfaceProperty = CastField<FInterfaceProperty>(ValueProperty))
+		{
+			RequiredInterface = InterfaceProperty->InterfaceClass;
+		}
+		else if (const FObjectPropertyBase* ObjectProperty = CastField<FObjectProperty>(ValueProperty))
+		{
+			MetaClass = ObjectProperty->PropertyClass;
+
+			// Verse generated interfaces are not fully supported yet, for now they are stored in FObjectProperty metadata
+			RequiredInterface = ObjectProperty->GetOwnerProperty()->GetClassMetaData(TEXT("MustImplement"));
+
+			if (!RequiredInterface && ObjectProperty->PropertyClass->HasAnyClassFlags(EClassFlags::CLASS_Interface))
+			{
+				RequiredInterface = ObjectProperty->PropertyClass;
+			}
+		}
+
+		TWeakPtr<FPropertyNode> WeakPropertyNode = PropertyNode;
+
+		uint8* ValueAddress = nullptr;
+		const bool bMultipleValues = PropertyNode->GetSingleReadAddress(ValueAddress) == FPropertyAccess::MultipleValues;
+
+		return SNew(SPropertyEditorClass)
+			.MetaClass(MetaClass)
+			.RequiredInterface(RequiredInterface)
+			.AllowAbstract(false)
+			.AllowNone(bMultipleValues) // selecting `None` is how we clear multiple-values
+			.ShowTree(false)
+			.SelectedClass(SelectedClass)
+			.IsEnabled(IsEnabled)
+			.ToolTipText(OptionalToolTipText.Get().IsEmpty() ? LOCTEXT("PickOptionalButtonToolTipText", "Pick a value to set this Optional to.") : OptionalToolTipText)
+			.OnSetClass_Lambda([OnSetOptional, OnClearOptional, WeakPropertyNode] (const UClass* SelectedClass){
+					if (WeakPropertyNode.IsValid())
+					{
+						if (SelectedClass)
+						{
+							OnSetOptional.ExecuteIfBound(nullptr, SelectedClass);
+						}
+						else
+						{
+							OnClearOptional.ExecuteIfBound();
+						}
+					}
+				})
+			.InvalidObjectDisplayText("Set to Value")
+			.bIsMultiOptionalSetter(true);
+	}
+
+	TSharedRef<SWidget> MakeClearOptionalButton(FSimpleDelegate OnClearOptional, TAttribute<FText> OptionalToolTipText, TAttribute<bool> IsEnabled)
 	{
 		return
 			SNew(SPropertyEditorButton)
 			.Text(OptionalToolTipText.Get().IsEmpty() ? LOCTEXT("ClearOptionalButtonToolTipText", "Clear Optional") : OptionalToolTipText)
 			.Image(FAppStyle::Get().GetBrush("Icons.X"))
-			.OnClickAction(OnClearOptionalClicked)
+			.OnClickAction(OnClearOptional)
 			.IsEnabled(IsEnabled)
 			.IsFocusable(false);
 	}
diff --git a/Engine/Source/Editor/PropertyEditor/Private/PropertyEditorHelpers.cpp b/Engine/Source/Editor/PropertyEditor/Private/PropertyEditorHelpers.cpp
--- a/Engine/Source/Editor/PropertyEditor/Private/PropertyEditorHelpers.cpp
+++ b/Engine/Source/Editor/PropertyEditor/Private/PropertyEditorHelpers.cpp
@@ -847,6 +847,7 @@ namespace PropertyEditorHelpers
 		// Handle an optional value node.
 		if (PropertyNode->IsOptionalValueNode())
 		{
+			// Add optional 'X' button
 			OutRequiredButtons.Add(EPropertyButton::OptionalClear);
 		}
 	}
@@ -1006,26 +1007,26 @@ namespace PropertyEditorHelpers
 	{
 		TSharedPtr<SWidget> NewButton;
 
-		TWeakPtr<FPropertyNode> WeakPropertyEditor = PropertyEditor->GetPropertyNode();
+		TWeakPtr<FPropertyNode> WeakPropertyNode = PropertyEditor->GetPropertyNode();
 
-		TAttribute<bool>::FGetter IsPropertyButtonEnabledDelegate = TAttribute<bool>::FGetter::CreateStatic(&IsPropertyButtonEnabled, WeakPropertyEditor);
+		TAttribute<bool>::FGetter IsPropertyButtonEnabledDelegate = TAttribute<bool>::FGetter::CreateStatic(&IsPropertyButtonEnabled, WeakPropertyNode);
 		TAttribute<bool> IsEnabledAttribute = TAttribute<bool>::Create( IsPropertyButtonEnabledDelegate );
 
 		TAttribute<bool> IsAddEnabledAttribute = TAttribute<bool>::Create(
-			TAttribute<bool>::FGetter::CreateLambda([WeakPropertyEditor, PropertyEditor]() {
-			if (WeakPropertyEditor.IsValid())
+			TAttribute<bool>::FGetter::CreateLambda([WeakPropertyNode, PropertyEditor]() {
+			if (WeakPropertyNode.IsValid())
 			{
-				FProperty* Property = WeakPropertyEditor.Pin()->GetProperty();
+				FProperty* Property = WeakPropertyNode.Pin()->GetProperty();
 				// Check for multiple array selections with mismatched values
 				FArrayProperty* ArrayProperty = CastField<FArrayProperty>(Property);
 				if (ArrayProperty)
 				{
 					FString ArrayString;
 					FPropertyAccess::Result GetValResult = PropertyEditor->GetPropertyHandle()->GetValueAsDisplayString(ArrayString);
-					return IsPropertyButtonEnabled(WeakPropertyEditor) && GetValResult == FPropertyAccess::Success;
+					return IsPropertyButtonEnabled(WeakPropertyNode) && GetValResult == FPropertyAccess::Success;
 				}
 			}
-			return IsPropertyButtonEnabled(WeakPropertyEditor);
+			return IsPropertyButtonEnabled(WeakPropertyNode);
 		}));
 
 		switch( ButtonType )
@@ -1072,8 +1073,8 @@ namespace PropertyEditorHelpers
 		case EPropertyButton::Use:
 			{
 				FSimpleDelegate OnClickDelegate = FSimpleDelegate::CreateSP(PropertyEditor, &FPropertyEditor::UseSelected);
-				TAttribute<bool>::FGetter EnabledDelegate = TAttribute<bool>::FGetter::CreateStatic(&IsUseSelectedUnrestricted, WeakPropertyEditor);
-				TAttribute<FText>::FGetter TooltipDelegate = TAttribute<FText>::FGetter::CreateStatic(&GetUseSelectedTooltip, WeakPropertyEditor);
+				TAttribute<bool>::FGetter EnabledDelegate = TAttribute<bool>::FGetter::CreateStatic(&IsUseSelectedUnrestricted, WeakPropertyNode);
+				TAttribute<FText>::FGetter TooltipDelegate = TAttribute<FText>::FGetter::CreateStatic(&GetUseSelectedTooltip, WeakPropertyNode);
 
 				NewButton = PropertyCustomizationHelpers::MakeUseSelectedButton(OnClickDelegate, TAttribute<FText>::Create(TooltipDelegate), TAttribute<bool>::Create(EnabledDelegate));
 				break;
@@ -1104,7 +1105,14 @@ namespace PropertyEditorHelpers
 			break;
 
 		case EPropertyButton::OptionalSet:
-			NewButton = PropertyCustomizationHelpers::MakeSetOptionalButton(FSimpleDelegate::CreateSP(PropertyEditor, &FPropertyEditor::SetOptionalItem, (FProperty*)nullptr), FText(), IsEnabledAttribute);
+			NewButton = PropertyCustomizationHelpers::MakeSetOptionalButton(FOnSetOptional::CreateSP(PropertyEditor, &FPropertyEditor::SetOptionalItem), FText(), IsEnabledAttribute);
+			break;
+
+		case EPropertyButton::OptionalPick:
+			NewButton = PropertyCustomizationHelpers::MakePickOptionalButton(
+				FOnSetOptional::CreateSP(PropertyEditor, &FPropertyEditor::SetOptionalItem), 
+				FSimpleDelegate::CreateSP(PropertyEditor, &FPropertyEditor::ClearOptionalItem),
+				PropertyEditor->GetPropertyNode(), FText(), IsEnabledAttribute);
 			break;
 
 		case EPropertyButton::OptionalClear:
diff --git a/Engine/Source/Editor/PropertyEditor/Private/PropertyEditorHelpers.h b/Engine/Source/Editor/PropertyEditor/Private/PropertyEditorHelpers.h
--- a/Engine/Source/Editor/PropertyEditor/Private/PropertyEditorHelpers.h
+++ b/Engine/Source/Editor/PropertyEditor/Private/PropertyEditorHelpers.h
@@ -38,6 +38,7 @@ namespace EPropertyButton
 		EditConfigHierarchy,
 		Documentation,
 		OptionalSet,
+		OptionalPick,
 		OptionalClear,
 	};
 }
diff --git a/Engine/Source/Editor/PropertyEditor/Private/PropertyHandleImpl.cpp b/Engine/Source/Editor/PropertyEditor/Private/PropertyHandleImpl.cpp
--- a/Engine/Source/Editor/PropertyEditor/Private/PropertyHandleImpl.cpp
+++ b/Engine/Source/Editor/PropertyEditor/Private/PropertyHandleImpl.cpp
@@ -5307,7 +5307,7 @@ FPropertyAccess::Result FPropertyHandleOptional::GetOptionalValue(FProperty*& Ou
 	return Result;
 }
 
-FPropertyAccess::Result FPropertyHandleOptional::SetOptionalValue(FProperty* NewValue)
+FPropertyAccess::Result FPropertyHandleOptional::SetOptionalValue(FProperty* NewProperty, const UClass* NewObjectClass)
 {
 	TSharedPtr<FPropertyNode> PropertyNode = Implementation->GetPropertyNode();
 	FOptionalProperty* OptionalProperty = CastField<FOptionalProperty>(PropertyNode->GetProperty());
@@ -5360,40 +5360,51 @@ FPropertyAccess::Result FPropertyHandleOptional::SetOptionalValue(FProperty* New
 	{
 		void* Addr = ReadAddresses.GetAddress(i);
 		OptionalProperty->MarkSetAndGetInitializedValuePointerToReplace(Addr);
-		
-		if (NewValue)
-		{
-			OptionalProperty->SetValueProperty(NewValue);
-		}
 
-		UObject* Obj = ObjectNode ? ObjectNode->GetUObject(i) : nullptr;
-		TopLevelObjects.Add(Obj);
-
-		// If our OptionalValue is a ptr to an object and we are not setting to a passed in value
-		// we need to initialize a default of that object and set the ptr to it.
-		FObjectProperty* ObjectProperty = CastField<FObjectProperty>(OptionalProperty->GetValueProperty());
-		FClassProperty* ClassProperty = CastField<FClassProperty>(OptionalProperty->GetValueProperty());
-		if (ClassProperty)
+		if (NewProperty)
 		{
-			void* ClassPropertyValuePtr = ClassProperty->ContainerPtrToValuePtr<void>(Addr);
-			ClassProperty->SetObjectPropertyValue(ClassPropertyValuePtr, ClassProperty->MetaClass);
+			OptionalProperty->SetValueProperty(NewProperty);
 		}
-		else if (ObjectProperty && !ObjectProperty->PropertyClass->HasAnyClassFlags(EClassFlags::CLASS_Abstract))
-		{
-			UObject* Outer = Obj;
 
-			if (Outer)
+		FProperty* ValueProperty = OptionalProperty->GetValueProperty();
+		if (ValueProperty && NewObjectClass)
+		{
+			if (FClassProperty* ClassProperty = CastField<FClassProperty>(ValueProperty))
 			{
-				UObject* NewDefaultObjectValue = NewObject<UObject>(Outer, ObjectProperty->PropertyClass);
-				if (NewDefaultObjectValue)
-				{
-					void* ObjectPropertyValuePtr = ObjectProperty->ContainerPtrToValuePtr<void>(Addr);
-					ObjectProperty->SetObjectPropertyValue(ObjectPropertyValuePtr, NewDefaultObjectValue);
-				}
+				void* ClassPropertyValuePtr = ClassProperty->ContainerPtrToValuePtr<void>(Addr);
+				ClassProperty->SetObjectPropertyValue(ClassPropertyValuePtr, const_cast<UClass*>(NewObjectClass));
+			}
+			else if (NewObjectClass->HasAnyClassFlags(EClassFlags::CLASS_Abstract))
+			{
+				Implementation->ShowInvalidOperationError(LOCTEXT("SetOptionalElement", "Error: Cannot construct an object from an 'abstract' class."));
+			}
+			else if (NewObjectClass->HasAnyClassFlags(EClassFlags::CLASS_Interface))
+			{
+				Implementation->ShowInvalidOperationError(LOCTEXT("SetOptionalElement", "Error: Cannot construct an object from an 'interface' class."));
 			}
 			else
 			{
-				Implementation->ShowInvalidOperationError(LOCTEXT("SetOptionalElement", "Could not create a default value for optional object as could not determine outer object."));
+				UObject* Outer = ObjectNode ? ObjectNode->GetUObject(i) : nullptr;
+				if(!Outer)
+				{
+					Implementation->ShowInvalidOperationError(LOCTEXT("SetOptionalElement", "Could not create value for optional object as could not determine outer object."));
+				}
+				TopLevelObjects.Add(Outer);
+
+				// TODO: Construct object using `NewObjectClass`
+				UObject* Obj = NewObject<UObject>(Outer, NewObjectClass);
+				if (Obj)
+				{
+					if (FObjectProperty* ObjectProperty = CastField<FObjectProperty>(ValueProperty))
+					{
+						void* ObjectPropertyValuePtr = ObjectProperty->ContainerPtrToValuePtr<void>(Addr);
+						ObjectProperty->SetObjectPropertyValue(ObjectPropertyValuePtr, Obj);
+					}
+				}
+				else
+				{
+					Implementation->ShowInvalidOperationError(LOCTEXT("SetOptionalElement", "Could not create value for optional object (support for this type may not yet be implemented)"));
+				}
 			}
 		}
 
diff --git a/Engine/Source/Editor/PropertyEditor/Private/PropertyHandleImpl.h b/Engine/Source/Editor/PropertyEditor/Private/PropertyHandleImpl.h
--- a/Engine/Source/Editor/PropertyEditor/Private/PropertyHandleImpl.h
+++ b/Engine/Source/Editor/PropertyEditor/Private/PropertyHandleImpl.h
@@ -739,7 +739,7 @@ public:
 	static bool Supports(TSharedRef<FPropertyNode> PropertyNode);
 	/** IPropertyHandleOptional interface */
 	virtual FPropertyAccess::Result GetOptionalValue(FProperty*& OutValue) override;
-	virtual FPropertyAccess::Result SetOptionalValue(FProperty* NewValue) override;
+	virtual FPropertyAccess::Result SetOptionalValue(FProperty* NewProperty, const UClass* NewObjectClass) override;
 	virtual FPropertyAccess::Result ClearOptionalValue() override;
 	virtual TSharedPtr<IPropertyHandleOptional> AsOptional() override;
 	virtual bool IsEditable() const override;
diff --git a/Engine/Source/Editor/PropertyEditor/Private/UserInterface/PropertyEditor/SPropertyEditorClass.cpp b/Engine/Source/Editor/PropertyEditor/Private/UserInterface/PropertyEditor/SPropertyEditorClass.cpp
--- a/Engine/Source/Editor/PropertyEditor/Private/UserInterface/PropertyEditor/SPropertyEditorClass.cpp
+++ b/Engine/Source/Editor/PropertyEditor/Private/UserInterface/PropertyEditor/SPropertyEditorClass.cpp
@@ -219,6 +219,8 @@ void SPropertyEditorClass::Construct(const FArguments& InArgs, const TSharedPtr<
 		OnSetClass = InArgs._OnSetClass;
 	}
 
+	InvalidObjectDisplayText = InArgs._InvalidObjectDisplayText;
+
 	CreateClassFilter(ClassViewerFilters);
 
 	SAssignNew(ComboButton, SComboButton)
@@ -240,7 +242,7 @@ void SPropertyEditorClass::Construct(const FArguments& InArgs, const TSharedPtr<
 }
 
 /** Util to give better names for BP generated classes */
-static FText GetClassDisplayName(const UObject* Object, bool bShowDisplayNames)
+static FText GetClassDisplayName(const UObject* Object, bool bShowDisplayNames, const FString& InvalidObjectDisplayText)
 {
 	const UClass* Class = Cast<UClass>(Object);
 	if (Class != nullptr)
@@ -256,7 +258,7 @@ static FText GetClassDisplayName(const UObject* Object, bool bShowDisplayNames)
 			return FText::FromString(BP->GetName());
 		}
 	}
-	return (Object) ? FText::FromString(Object->GetName()) : LOCTEXT("InvalidObject", "None");
+	return (Object) ? FText::FromString(Object->GetName()) : FText::FromString(InvalidObjectDisplayText);
 }
 
 FText SPropertyEditorClass::GetDisplayValueAsString() const
@@ -266,6 +268,11 @@ FText SPropertyEditorClass::GetDisplayValueAsString() const
 	// Guard against re-entrancy which can happen if the delegate executed below (SelectedClass.Get()) forces a slow task dialog to open, thus causing this to lose context and regain focus later starting the loop over again
 	if( !bIsReentrant )
 	{
+		if (bIsMultiOptionalSetter)
+		{
+			return FText::FromString("Multiple States");
+		}
+
 		TGuardValue<bool> Guard( bIsReentrant, true );
 		if(PropertyEditor.IsValid())
 		{
@@ -274,19 +281,18 @@ FText SPropertyEditorClass::GetDisplayValueAsString() const
 
 			if(Result == FPropertyAccess::Success && ObjectValue != nullptr)
 			{
-				return GetClassDisplayName(ObjectValue, bShowDisplayNames);
+				return GetClassDisplayName(ObjectValue, bShowDisplayNames, InvalidObjectDisplayText);
 			}
 
 			return FText::FromString(FPaths::GetBaseFilename(PropertyEditor->GetValueAsString()));
 		}
 
-		return GetClassDisplayName(SelectedClass.Get(), bShowDisplayNames);
+		return GetClassDisplayName(SelectedClass.Get(), bShowDisplayNames, InvalidObjectDisplayText);
 	}
 	else
 	{
 		return FText::GetEmpty();
 	}
-	
 }
 
 void SPropertyEditorClass::CreateClassFilter(const TArray<TSharedRef<IClassViewerFilter>>& InClassFilters)
diff --git a/Engine/Source/Editor/PropertyEditor/Private/UserInterface/PropertyEditor/SPropertyEditorClass.h b/Engine/Source/Editor/PropertyEditor/Private/UserInterface/PropertyEditor/SPropertyEditorClass.h
--- a/Engine/Source/Editor/PropertyEditor/Private/UserInterface/PropertyEditor/SPropertyEditorClass.h
+++ b/Engine/Source/Editor/PropertyEditor/Private/UserInterface/PropertyEditor/SPropertyEditorClass.h
@@ -34,6 +34,8 @@ public:
 		, _IsBlueprintBaseOnly(false)
 		, _AllowNone(true)
 		, _ShowDisplayNames(false)
+		, _InvalidObjectDisplayText("None")
+		, _bIsMultiOptionalSetter(false)
 		{}
 		SLATE_ARGUMENT(FSlateFontInfo, Font)
 
@@ -64,6 +66,10 @@ public:
 		SLATE_EVENT(FOnSetClass, OnSetClass)
 		/** Custom class filter(s) to be applied to the class picker widget (may be empty) */
 		SLATE_ARGUMENT(TArray<TSharedRef<class IClassViewerFilter>>, ClassViewerFilters)
+		/** Text to show when no class is selected (Default of "None") */
+		SLATE_ARGUMENT(FString, InvalidObjectDisplayText)
+		/** True if this class editor is for setting a multiple-value optional */
+		SLATE_ARGUMENT(bool, bIsMultiOptionalSetter)
 	SLATE_END_ARGS()
 
 	static bool Supports(const TSharedRef< class FPropertyEditor >& InPropertyEditor);
@@ -140,7 +146,10 @@ private:
 	TArray<const UClass*> AllowedClassFilters;
 	/** Classes that can NOT be picked with this property */
 	TArray<const UClass*> DisallowedClassFilters;
-
+	/** Text which is displayed instead of the normal class-text (used to show 'Set to Value' on optionals) */
+	FString InvalidObjectDisplayText;
+	/** True if this class editor is for setting a multiple-value optional */
+	bool bIsMultiOptionalSetter;
 
 	/** Attribute used to get the currently selected class (required if PropertyEditor == null) */
 	TAttribute<const UClass*> SelectedClass;
diff --git a/Engine/Source/Editor/PropertyEditor/Private/UserInterface/PropertyEditor/SPropertyEditorEditInline.cpp b/Engine/Source/Editor/PropertyEditor/Private/UserInterface/PropertyEditor/SPropertyEditorEditInline.cpp
--- a/Engine/Source/Editor/PropertyEditor/Private/UserInterface/PropertyEditor/SPropertyEditorEditInline.cpp
+++ b/Engine/Source/Editor/PropertyEditor/Private/UserInterface/PropertyEditor/SPropertyEditorEditInline.cpp
@@ -272,6 +272,12 @@ TSharedRef<SWidget> SPropertyEditorEditInline::GenerateClassPicker(TSharedRef<IP
 	static const FName NAME_MustImplement(ANSITEXTVIEW("MustImplement"));
 	ClassFilter->InterfaceThatMustBeImplemented = MetadataProperty->GetClassMetaData(NAME_MustImplement);
 
+	if (ClassFilter->ObjProperty)
+	{
+		// Verse generated interfaces are not fully supported by yet, for now they are stored in FObjectProperty metadata
+		ClassFilter->InterfaceThatMustBeImplemented = ClassFilter->ObjProperty->GetOwnerProperty()->GetClassMetaData(TEXT("MustImplement"));
+	}
+
 	if (ClassFilter->InterfaceThatMustBeImplemented != nullptr)
 	{
 		if (!ClassFilter->InterfaceThatMustBeImplemented->HasAnyClassFlags(CLASS_Interface))
diff --git a/Engine/Source/Editor/PropertyEditor/Private/UserInterface/PropertyEditor/SPropertyEditorOptional.cpp b/Engine/Source/Editor/PropertyEditor/Private/UserInterface/PropertyEditor/SPropertyEditorOptional.cpp
--- a/Engine/Source/Editor/PropertyEditor/Private/UserInterface/PropertyEditor/SPropertyEditorOptional.cpp
+++ b/Engine/Source/Editor/PropertyEditor/Private/UserInterface/PropertyEditor/SPropertyEditorOptional.cpp
@@ -22,16 +22,29 @@ void SPropertyEditorOptional::Construct(const FArguments& InArgs, const TSharedR
 	MaxDesiredWidth = 0.0f;
 	PropertyEditor = InPropertyEditor;
 
-	const TSharedRef<FPropertyNode> PropertyNode = InPropertyEditor->GetPropertyNode();
-
 	uint8* ValueAddress = nullptr;
-	if (PropertyNode->GetSingleReadAddress(ValueAddress) == FPropertyAccess::MultipleValues)
+	TSharedRef<FPropertyNode> PropertyNode = InPropertyEditor->GetPropertyNode();
+	const FOptionalProperty* OptionalProperty = CastField<FOptionalProperty>(PropertyNode->GetProperty());
+	const FProperty* ValueProperty = OptionalProperty->GetValueProperty();
+
+	if (ValueProperty->IsA(FObjectPropertyBase::StaticClass()) || ValueProperty->IsA(FInterfaceProperty::StaticClass()))
+	{
+		ValueEditorWidget = SNew(SBox)
+			.HAlign(HAlign_Left)
+			.VAlign(VAlign_Center)
+			.Padding(2.0f, 0.0f)
+			[
+				PropertyEditorHelpers::MakePropertyButton(EPropertyButton::OptionalPick, InPropertyEditor)
+			];
+	}
+	else if (PropertyNode->GetSingleReadAddress(ValueAddress) == FPropertyAccess::MultipleValues)
 	{	
-		MultiselectOptions.Add(MakeShareable( new FText(LOCTEXT("OptionalMultiselect_Multiple", "MultipleStates"))));
+		// Note: multi-select for `OptionalPick` types is handled in `MakePickOptionalButton`
+		MultiselectOptions.Add(MakeShareable( new FText(LOCTEXT("OptionalMultiselect_Multiple", "Multiple States"))));
 		MultiselectOptions.Add(MakeShareable( new FText(LOCTEXT("OptionalMultiselect_Value", "Set all to Value"))));
 		MultiselectOptions.Add(MakeShareable( new FText(LOCTEXT("OptionalMultiselect_None", "Set all to None"))));
 
-		ValueEditorWidget = SNew(SComboBox<TSharedPtr<FText>>)
+		SAssignNew(ValueEditorWidget, SComboBox<TSharedPtr<FText>>)
 			.OptionsSource(&MultiselectOptions)
 			.OnSelectionChanged(this, &SPropertyEditorOptional::OnOptionChanged)
 			.OnGenerateWidget(this, &SPropertyEditorOptional::MakeWidgetForOption)
@@ -112,7 +125,7 @@ void SPropertyEditorOptional::OnOptionChanged(TSharedPtr<FText> NewOption, ESele
 	if (NewOption->EqualTo(*MultiselectOptions[1]) /*"Set all to Value"*/)
 	{
 		OptionalHandle->ClearOptionalValue(); // Clear all values
-		OptionalHandle->SetOptionalValue(nullptr); // Set values to default-initialized
+		OptionalHandle->SetOptionalValue(nullptr, nullptr); // Set values to default-initialized
 	}
 	else if (NewOption->EqualTo(*MultiselectOptions[2]) /*"Set all to None"*/)
 	{
diff --git a/Engine/Source/Editor/PropertyEditor/Public/PropertyCustomizationHelpers.h b/Engine/Source/Editor/PropertyEditor/Public/PropertyCustomizationHelpers.h
--- a/Engine/Source/Editor/PropertyEditor/Public/PropertyCustomizationHelpers.h
+++ b/Engine/Source/Editor/PropertyEditor/Public/PropertyCustomizationHelpers.h
@@ -38,6 +38,7 @@ class SPropertyEditorStruct;
 class SToolTip;
 class UActorComponent;
 class UFactory;
+class FPropertyNode;
 struct FAssetData;
 struct FSceneOutlinerFilters;
 
@@ -52,6 +53,7 @@ DECLARE_DELEGATE_RetVal(FString, FOnGetPropertyComboBoxValue);
 DECLARE_DELEGATE_OneParam(FOnPropertyComboBoxValueSelected, const FString&);
 DECLARE_DELEGATE_ThreeParams(FOnInstancedPropertyIteration, IDetailCategoryBuilder&, IDetailGroup*, TSharedRef<IPropertyHandle>&);
 DECLARE_DELEGATE_RetVal(bool, FOnIsEnabled);
+DECLARE_DELEGATE_TwoParams(FOnSetOptional, FProperty*, const UClass*);
 
 /** Collects advanced arguments for MakePropertyComboBox */
 struct FPropertyComboBoxArgs
@@ -177,8 +179,9 @@ namespace PropertyCustomizationHelpers
 	PROPERTYEDITOR_API TSharedRef<SWidget> MakeEditConfigHierarchyButton(FSimpleDelegate OnEditConfigClicked, TAttribute<FText> OptionalToolTipText = FText(), TAttribute<bool> IsEnabled = true);
 	PROPERTYEDITOR_API TSharedRef<SWidget> MakeDocumentationButton(const TSharedRef<FPropertyEditor>& InPropertyEditor);
 	PROPERTYEDITOR_API TSharedRef<SWidget> MakeSaveButton(FSimpleDelegate OnSaveClicked, TAttribute<FText> OptionalToolTipText = FText(), TAttribute<bool> IsEnabled = true);
-	PROPERTYEDITOR_API TSharedRef<SWidget> MakeSetOptionalButton(FSimpleDelegate OnSetOptionalClicked, TAttribute<FText> OptionalToolTipText = FText(), TAttribute<bool> IsEnabled = true);
-	PROPERTYEDITOR_API TSharedRef<SWidget> MakeClearOptionalButton(FSimpleDelegate OnClearOptionalClicked, TAttribute<FText> OptionalToolTipText = FText(), TAttribute<bool> IsEnabled = true);
+	PROPERTYEDITOR_API TSharedRef<SWidget> MakeSetOptionalButton(FOnSetOptional OnSetOptional, TAttribute<FText> OptionalToolTipText = FText(), TAttribute<bool> IsEnabled = true);
+	PROPERTYEDITOR_API TSharedRef<SWidget> MakePickOptionalButton(FOnSetOptional OnSetOptional, FSimpleDelegate OnClearOptional, const TSharedRef<FPropertyNode>& PropertyNode, TAttribute<FText> OptionalToolTipText = FText(), TAttribute<bool> IsEnabled = true);
+	PROPERTYEDITOR_API TSharedRef<SWidget> MakeClearOptionalButton(FSimpleDelegate OnClearOptional, TAttribute<FText> OptionalToolTipText = FText(), TAttribute<bool> IsEnabled = true);
 	PROPERTYEDITOR_API TSharedRef<SWidget> MakeFunctionCallButton(const FPropertyFunctionCallArgs& InArgs);
 	PROPERTYEDITOR_API void AddFunctionCallWidgets(IDetailGroup& RootGroup, const TArrayView<UFunction*>& InCallInEditorFunctions, const FPropertyFunctionCallDelegates& InArgs);
 	PROPERTYEDITOR_API void AddFunctionCallWidgets(IDetailLayoutBuilder& DetailBuilder, const TArrayView<UFunction*>& InCallInEditorFunctions, const FPropertyFunctionCallDelegates& InArgs);
diff --git a/Engine/Source/Editor/PropertyEditor/Public/PropertyHandle.h b/Engine/Source/Editor/PropertyEditor/Public/PropertyHandle.h
--- a/Engine/Source/Editor/PropertyEditor/Public/PropertyHandle.h
+++ b/Engine/Source/Editor/PropertyEditor/Public/PropertyHandle.h
@@ -986,7 +986,7 @@ public:
 	 * @param NewValue	The value to set the optional to. Pass in nullptr to default-initialize the value.
 	 * @return			Whether or not this was successful
 	 */
-	virtual FPropertyAccess::Result SetOptionalValue(FProperty* NewValue) = 0;
+	virtual FPropertyAccess::Result SetOptionalValue(FProperty* NewProperty, const UClass* NewObjectClass) = 0;
 
 	/**
 	 * Clear the item
