Add property status widget for subobjects of verse::components
* DetailsDisplayManager initialized using an array of objects instead of an object name. This avoids requring implementations to cache their own object array and match against the "CategoryObjectName"
* PropertyPath within DisplayManager will be trimmed to be relative to the object's class instead of relative to the root object of the details panel.  Property initialization track and override serialization both require this, it is also easier to reason about as the path can be used to address objects in isolation from their parents.
* Override menu changed so that override/clear override buttons still work for non-component properties.  RemoveComponent and data fixup tool will not show up for non-component properties but will do so for components
* Removed unused PasteComponentValues action
Note: Array aspect of the objects given to DetailsDisplayManager not currently used as multi-select is not yet possible.

diff --git a/Engine/Source/Editor/PropertyEditor/Private/DetailCategoryBuilderImpl.cpp b/Engine/Source/Editor/PropertyEditor/Private/DetailCategoryBuilderImpl.cpp
--- a/Engine/Source/Editor/PropertyEditor/Private/DetailCategoryBuilderImpl.cpp
+++ b/Engine/Source/Editor/PropertyEditor/Private/DetailCategoryBuilderImpl.cpp
@@ -902,13 +902,29 @@ TSharedRef<ITableRow> FDetailCategoryImpl::GenerateWidgetForTableView(const TSha
 		HeaderContent = Row.ValueWidget.Widget;
 	}
 
-	InitializeObjectName();
 	TSharedPtr<FDetailLayoutBuilderImpl> ParentLayout = GetParentLayoutImpl();
+	
+	TArray<TWeakObjectPtr<UObject>> Objects;
+	TSharedPtr<FComplexPropertyNode> Node = DetailLayoutBuilder.Pin()->GetRootNode();
+	if (Node.IsValid())
+	{
+		if (FObjectPropertyNode* Object = Node->AsObjectNode())
+		{
+			Objects.Reserve(Object->GetNumObjects());
+			for(int32 Index = 0, End = Object->GetNumObjects(); Index < End; ++Index)
+			{
+				if (UObject* CategoryObject = Object->GetUObject(Index))
+				{
+					Objects.Add(CategoryObject);
+				}
+			}			
+		}
+	}
 
 	TSharedRef<SDetailCategoryTableRow> RowWidget =
 		SNew(SDetailCategoryTableRow, AsShared(), OwnerTable)
 		.PasteFromText(OnPasteFromText())
-		.ObjectName( ObjectName )
+		.Objects(&Objects)
 		.IsEmpty( bIsEmpty )
 		.InnerCategory(ParentLayout.IsValid() ? ParentLayout->IsLayoutForExternalRoot() : false)
 		.DisplayName(GetDisplayName())
@@ -923,31 +939,6 @@ TSharedRef<ITableRow> FDetailCategoryImpl::GenerateWidgetForTableView(const TSha
 	return RowWidget;
 }
 
-void FDetailCategoryImpl::InitializeObjectName()
-{
-	if (!DetailLayoutBuilder.IsValid())
-	{
-		return;
-	}
-	
-	const TSharedPtr<FComplexPropertyNode> Node = DetailLayoutBuilder.Pin()->GetRootNode();
-	
-	if (Node.IsValid())
-	{
-		if (const FObjectPropertyNode* Object = Node->AsObjectNode())
-		{
-			if (Object->GetNumObjects() > 0)
-			{
-				if (const UObject* CategoryObject = Object->GetUObject(0))
-				{
-					const FName Name{ CategoryObject->GetName() };
-					ObjectName = Name;
-				}
-			}
-		}
-	}	
-}
-
 bool FDetailCategoryImpl::GenerateStandaloneWidget(FDetailWidgetRow& OutRow) const
 {
 	TSharedPtr<SWidget> HeaderContent = HeaderContentWidget;
@@ -1213,11 +1204,6 @@ void FDetailCategoryImpl::GetChildren(FDetailNodeList& OutChildren, const bool&
 	GetGeneratedChildren(OutChildren, bInIgnoreVisibility, false);
 }
 
-FName FDetailCategoryImpl::GetObjectName() const
-{
-	return ObjectName;
-}
-
 void FDetailCategoryImpl::GetGeneratedChildren(FDetailNodeList& OutChildren, bool bIgnoreVisibility, bool bIgnoreAdvanced)
 {
 	for (TSharedRef<FDetailTreeNode>& Child : SimpleChildNodes)
diff --git a/Engine/Source/Editor/PropertyEditor/Private/DetailCategoryBuilderImpl.h b/Engine/Source/Editor/PropertyEditor/Private/DetailCategoryBuilderImpl.h
--- a/Engine/Source/Editor/PropertyEditor/Private/DetailCategoryBuilderImpl.h
+++ b/Engine/Source/Editor/PropertyEditor/Private/DetailCategoryBuilderImpl.h
@@ -215,11 +215,6 @@ public:
 	virtual bool ShouldShowOnlyChildren() const override { return bShowOnlyChildren; }
 	virtual FName GetNodeName() const override { return GetCategoryName(); }
 
-	/**
-	 * Returns the name of the object which populates this Category
-	 */
-	FName GetObjectName() const;
-
 	/**
 	 * Gets all generated children with options for ignoring current child visibility or advanced dropdowns
 	 */
@@ -419,9 +414,6 @@ private:
 	 */
 	void RefreshTreeInternal(bool bRefilterCategory);
 
-	/** Initialize the name of the Category  */
-	void InitializeObjectName();
-
 private:
 	/** Layouts that appear in this category category */
 	FDetailLayoutMap LayoutMap;
@@ -449,9 +441,7 @@ private:
 	TWeakPtr<FDetailLayoutBuilderImpl> DetailLayoutBuilder;
 	/** Delegate handling pasting an optionally tagged text snippet */
 	TSharedPtr<FOnPasteFromText> PasteFromTextDelegate;
-
-	/** The name of the object that this category is being defined by */
-	FName ObjectName;
+	
 	/** The category identifier */
 	FName CategoryName;
 	/** The sort order of this category (amongst all categories) */
diff --git a/Engine/Source/Editor/PropertyEditor/Private/DetailsDisplayManager.cpp b/Engine/Source/Editor/PropertyEditor/Private/DetailsDisplayManager.cpp
--- a/Engine/Source/Editor/PropertyEditor/Private/DetailsDisplayManager.cpp
+++ b/Engine/Source/Editor/PropertyEditor/Private/DetailsDisplayManager.cpp
@@ -38,11 +38,6 @@ bool FDetailsDisplayManager::ShouldShowCategoryMenu()
 	return false;
 }
 
-void FDetailsDisplayManager::SetCategoryObjectName(FName InCategoryObjectName)
-{
-	CategoryObjectName = InCategoryObjectName;
-}
-
 TSharedPtr<SWidget> FDetailsDisplayManager::GetCategoryMenu(FName InCategoryObjectName)
 {
 	return nullptr;
diff --git a/Engine/Source/Editor/PropertyEditor/Private/SDetailCategoryTableRow.cpp b/Engine/Source/Editor/PropertyEditor/Private/SDetailCategoryTableRow.cpp
--- a/Engine/Source/Editor/PropertyEditor/Private/SDetailCategoryTableRow.cpp
+++ b/Engine/Source/Editor/PropertyEditor/Private/SDetailCategoryTableRow.cpp
@@ -32,7 +32,6 @@ void SDetailCategoryTableRow::Construct(const FArguments& InArgs, TSharedRef<FDe
 
 	IDetailsViewPrivate* DetailsView = InOwnerTreeNode->GetDetailsView();
 	FDetailColumnSizeData& ColumnSizeData = DetailsView->GetColumnSizeData();
-	ObjectName = InArgs._ObjectName;
 	DisplayManager = DetailsView->GetDisplayManager();
 
 	if (DisplayManager.IsValid() && DisplayManager->GetDetailsViewStyle())
@@ -40,7 +39,10 @@ void SDetailCategoryTableRow::Construct(const FArguments& InArgs, TSharedRef<FDe
 		DetailsViewStyle = DisplayManager->GetDetailsViewStyle();
 	}
 
-	InitializeDisplayManager();
+	if (DisplayManager.IsValid())
+	{
+		DisplayManager->SetIsOuterCategory( !bIsInnerCategory );
+	}
 
 	PulseAnimation.AddCurve(0.0f, UE::PropertyEditor::Private::PulseAnimationLength, ECurveEaseFunction::CubicInOut);
 	CopyAction.ExecuteAction = FExecuteAction::CreateSP(this, &SDetailCategoryTableRow::OnCopyCategory);
@@ -108,12 +110,11 @@ void SDetailCategoryTableRow::Construct(const FArguments& InArgs, TSharedRef<FDe
 	}
 
 	OwnerTableViewWeak = InOwnerTableView;
-
-
-
+	
 	PRAGMA_DISABLE_DEPRECATION_WARNINGS
 	if (DisplayManager->CanConstructPropertyUpdatedWidgetBuilder())
 	{
+		TSharedPtr<FPropertyNode> Node = InOwnerTreeNode->GetPropertyNode();
 		FConstructPropertyUpdatedWidgetBuilderArgs Args;
 		Args.ResetToDefaultAction = FExecuteAction::CreateLambda( [this]
 		{
@@ -123,7 +124,7 @@ void SDetailCategoryTableRow::Construct(const FArguments& InArgs, TSharedRef<FDe
 			}
 		});
 		Args.PropertyPath = nullptr;
-		Args.CategoryObjectName = ObjectName;
+		Args.Objects = InArgs._Objects;
 		
 		PropertyUpdatedWidgetBuilder = DisplayManager->ConstructPropertyUpdatedWidgetBuilder(Args);
 		if (PropertyUpdatedWidgetBuilder.IsValid())
@@ -176,26 +177,14 @@ void SDetailCategoryTableRow::Construct(const FArguments& InArgs, TSharedRef<FDe
 			.VAlign(VAlign_Center)
 			.AutoWidth()
 			[
-			PropertyUpdatedWidgetBuilder.IsValid() ?
-								PropertyUpdatedWidgetBuilder->Bind_IsVisible(TAttribute<EVisibility>::CreateLambda([this]
-									{
-										return EVisibility::Visible;
-									})).GenerateWidget().ToSharedRef() :
-									SNullWidget::NullWidget
-			]
-			+SHorizontalBox::Slot()
-			.HAlign(HAlign_Right)
-			.VAlign(VAlign_Center)
-			.AutoWidth()
-			[
-			DisplayManager.IsValid() ?
-				*FCategoryMenuComboButtonBuilder( DisplayManager.ToSharedRef() )
-				.Set_OnGetContent(FOnGetContent::CreateLambda([this]
-				{
-					const TSharedPtr<SWidget> Menu = DisplayManager->GetCategoryMenu(ObjectName);
-					return Menu.IsValid() ? Menu.ToSharedRef() : SNullWidget::NullWidget;
-				})) :
-				SNullWidget::NullWidget
+				PropertyUpdatedWidgetBuilder.IsValid() ?
+					PropertyUpdatedWidgetBuilder->
+						Bind_IsVisible(TAttribute<EVisibility>::CreateLambda([this]
+						{
+							return EVisibility::Visible;
+						}))
+						.GenerateWidget().ToSharedRef() :
+						SNullWidget::NullWidget
 			]
 			]
 			] 
@@ -587,12 +576,3 @@ FReply SDetailCategoryTableRow::OnMouseButtonDoubleClick(const FGeometry& InMyGe
 {
 	return OnMouseButtonDown(InMyGeometry, InMouseEvent);
 }
-
-void SDetailCategoryTableRow::InitializeDisplayManager()
-{
-	if (DisplayManager.IsValid())
-	{
-		DisplayManager->SetCategoryObjectName( ObjectName );
-		DisplayManager->SetIsOuterCategory( !bIsInnerCategory );
-	}
-}
diff --git a/Engine/Source/Editor/PropertyEditor/Private/SDetailCategoryTableRow.h b/Engine/Source/Editor/PropertyEditor/Private/SDetailCategoryTableRow.h
--- a/Engine/Source/Editor/PropertyEditor/Private/SDetailCategoryTableRow.h
+++ b/Engine/Source/Editor/PropertyEditor/Private/SDetailCategoryTableRow.h
@@ -22,7 +22,7 @@ public:
 		, _WholeRowHeaderContent(false)
 		, _ShowBorder(true)
 		, _IsEmpty(false)
-		, _ObjectName(NAME_Name)
+		, _Objects(nullptr)
 	{}
 		SLATE_ARGUMENT(FText, DisplayName)
 		SLATE_ARGUMENT(bool, InnerCategory)
@@ -34,7 +34,7 @@ public:
 		* empty stub with no expansion arrow */
 		SLATE_ARGUMENT(bool, IsEmpty)
 		SLATE_ARGUMENT(TSharedPtr<FOnPasteFromText>, PasteFromText)
-		SLATE_ARGUMENT(FName, ObjectName)
+		SLATE_ARGUMENT(TArray<TWeakObjectPtr<UObject>>*, Objects)
 	SLATE_END_ARGS()
 
 	void Construct(const FArguments& InArgs, TSharedRef<FDetailTreeNode> InOwnerTreeNode, const TSharedRef<STableViewBase>& InOwnerTableView);
@@ -71,16 +71,10 @@ private:
 	virtual FReply OnMouseButtonDown(const FGeometry& MyGeometry, const FPointerEvent& MouseEvent) override;
 	virtual FReply OnMouseButtonDoubleClick(const FGeometry& InMyGeometry, const FPointerEvent& InMouseEvent) override;
 
-	/** Initializes the @code TSharedPtr<FDetailsDisplayManager> DisplayManager @endcode for this category */
-	void InitializeDisplayManager();
-
 private:
 	/** Cached category name. */
 	FText DisplayName;
 
-	/** The name of the object that is specified by this category */
-	FName ObjectName;
-
 	/**
 	* The style of the details view. This holds information and methods to get things like the row border image,
 	* table padding, and so forth.
diff --git a/Engine/Source/Editor/PropertyEditor/Private/SDetailSingleItemRow.cpp b/Engine/Source/Editor/PropertyEditor/Private/SDetailSingleItemRow.cpp
--- a/Engine/Source/Editor/PropertyEditor/Private/SDetailSingleItemRow.cpp
+++ b/Engine/Source/Editor/PropertyEditor/Private/SDetailSingleItemRow.cpp
@@ -592,23 +592,68 @@ void SDetailSingleItemRow::Construct( const FArguments& InArgs, FDetailLayoutCus
 					DetailsView->GetDisplayManager().IsValid())
 			{
 				TSharedPtr<FDetailsDisplayManager> DisplayManagerLocal = DetailsView->GetDisplayManager();
-
+				
 				PRAGMA_DISABLE_DEPRECATION_WARNINGS
 				if (DisplayManagerLocal->CanConstructPropertyUpdatedWidgetBuilder())
 				{
 					DisplayManager = MoveTemp(DisplayManagerLocal);
 
-					FConstructPropertyUpdatedWidgetBuilderArgs Args;
-					Args.ResetToDefaultAction = FExecuteAction::CreateSP(this, &SDetailSingleItemRow::OnResetToDefaultClicked);
-					Args.PropertyPath = MakeShared<FPropertyPath>(InOwnerTreeNode->GetPropertyPath());
-					Args.CategoryObjectName = Category->GetObjectName();
-						
-					PropertyUpdatedWidgetBuilder = DisplayManager->ConstructPropertyUpdatedWidgetBuilder(Args);
+					auto FindParentUObject = [](FPropertyNode* Cursor) -> FObjectPropertyNode*
+					{
+						while(Cursor)
+						{
+							FObjectPropertyNode* ParentObjectNode = Cursor->AsObjectNode();
+							if (ParentObjectNode)
+							{
+								return ParentObjectNode;
+							}
+							else
+							{
+								Cursor = Cursor->GetParentNode();
+							}
+						}
+						return nullptr;
+					};
+
+					FObjectPropertyNode* OwningObjectNode = FindParentUObject(Customization->GetPropertyNode().Get());
+					TArray<TWeakObjectPtr<UObject>> Objects;
+					if (OwningObjectNode)
+					{
+						Objects.Reserve(OwningObjectNode->GetNumObjects());
+						for(int32 Index = 0, End = OwningObjectNode->GetNumObjects(); Index < End; ++Index)
+						{
+							if (UObject* CategoryObject = OwningObjectNode->GetUObject(Index))
+							{
+								Objects.Add(CategoryObject);
+							}
+						}	
+					}
 
-					if (PropertyUpdatedWidgetBuilder.IsValid())
+					TSharedPtr<FPropertyNode> PropertyNode = InOwnerTreeNode->GetPropertyNode();
+					TStrongObjectPtr<UObject> PinnedObject;
+					if (!Objects.IsEmpty())
+					{
+						PinnedObject = Objects[0].Pin();
+					}
+					if (PropertyNode && PinnedObject.IsValid())
 					{
-						TAttribute<bool> IsHovered = TAttribute<bool>::CreateSP( this, &SDetailSingleItemRow::IsHovered);
-						PropertyUpdatedWidgetBuilder->Bind_IsRowHovered(IsHovered);
+						TSharedRef<FPropertyPath> PropertyPath = FPropertyNode::CreatePropertyPath( PropertyNode.ToSharedRef() );
+						
+						if (PropertyPath->IsValid())
+						{
+							FConstructPropertyUpdatedWidgetBuilderArgs Args;
+							Args.ResetToDefaultAction = FExecuteAction::CreateSP(this, &SDetailSingleItemRow::OnResetToDefaultClicked);
+							Args.Objects = &Objects;
+							Args.PropertyPath = PropertyPath;
+						
+							PropertyUpdatedWidgetBuilder = DisplayManager->ConstructPropertyUpdatedWidgetBuilder(Args);
+
+							if (PropertyUpdatedWidgetBuilder.IsValid())
+							{
+								TAttribute<bool> IsHovered = TAttribute<bool>::CreateSP( this, &SDetailSingleItemRow::IsHovered);
+								PropertyUpdatedWidgetBuilder->Bind_IsRowHovered(IsHovered);
+							}
+						}
 					}
 				}
 				PRAGMA_ENABLE_DEPRECATION_WARNINGS
diff --git a/Engine/Source/Editor/PropertyEditor/Public/DetailsDisplayManager.h b/Engine/Source/Editor/PropertyEditor/Public/DetailsDisplayManager.h
--- a/Engine/Source/Editor/PropertyEditor/Public/DetailsDisplayManager.h
+++ b/Engine/Source/Editor/PropertyEditor/Public/DetailsDisplayManager.h
@@ -46,13 +46,6 @@ public:
 	 */
 	virtual bool ShouldShowCategoryMenu();
 
-	/**
-	 * Sets the name of the object defined by the currently active category
-	 *
-	 * @param InCategoryObjectName the name of the category object
-	 */
-	void SetCategoryObjectName(FName InCategoryObjectName);            
-
 	/**
 	 * Gets the category menu SWidget and returns a shared pointer to it
 	 *
@@ -164,11 +157,6 @@ protected:
 	 * The primary style key for the details view. 
 	 */
 	FDetailsViewStyleKey PrimaryStyleKey;
-	
-	/**
-	 * The name of the object defined by the currently active category
-	 */
-	FName CategoryObjectName;
 
 	/**
 	 * A bool indicating whether or not the currently active category is an outer category
@@ -191,5 +179,5 @@ struct FConstructPropertyUpdatedWidgetBuilderArgs
 	
 	FExecuteAction ResetToDefaultAction;
 	TSharedPtr<FPropertyPath> PropertyPath;
-	FName CategoryObjectName;
+	TArray<TWeakObjectPtr<UObject>>* Objects = nullptr;
 };
\ No newline at end of file
diff --git a/Engine/Source/Editor/PropertyEditor/Public/PropertyPath.h b/Engine/Source/Editor/PropertyEditor/Public/PropertyPath.h
--- a/Engine/Source/Editor/PropertyEditor/Public/PropertyPath.h
+++ b/Engine/Source/Editor/PropertyEditor/Public/PropertyPath.h
@@ -124,7 +124,7 @@ public:
 		return NewPath;
 	}
 
-	TSharedRef< FPropertyPath > TrimRoot( const int32 AmountToTrim )
+	TSharedRef< FPropertyPath > TrimRoot( const int32 AmountToTrim ) const
 	{
 		TSharedRef< FPropertyPath > NewPath = MakeShareable( new FPropertyPath() );
 
