Update FTEXT_LOC inside of PropertyHandle::SetOptional to not be duplicates

diff --git a/Engine/Source/Editor/PropertyEditor/Private/PropertyHandleImpl.cpp b/Engine/Source/Editor/PropertyEditor/Private/PropertyHandleImpl.cpp
--- a/Engine/Source/Editor/PropertyEditor/Private/PropertyHandleImpl.cpp
+++ b/Engine/Source/Editor/PropertyEditor/Private/PropertyHandleImpl.cpp
@@ -5376,18 +5376,18 @@ FPropertyAccess::Result FPropertyHandleOptional::SetOptionalValue(FProperty* New
 			}
 			else if (NewObjectClass->HasAnyClassFlags(EClassFlags::CLASS_Abstract))
 			{
-				Implementation->ShowInvalidOperationError(LOCTEXT("SetOptionalElement", "Error: Cannot construct an object from an 'abstract' class."));
+				Implementation->ShowInvalidOperationError(LOCTEXT("SetOptionalAbstractError", "Error: Cannot construct an object from an 'abstract' class."));
 			}
 			else if (NewObjectClass->HasAnyClassFlags(EClassFlags::CLASS_Interface))
 			{
-				Implementation->ShowInvalidOperationError(LOCTEXT("SetOptionalElement", "Error: Cannot construct an object from an 'interface' class."));
+				Implementation->ShowInvalidOperationError(LOCTEXT("SetOptionalInterfaceError", "Error: Cannot construct an object from an 'interface' class."));
 			}
 			else
 			{
 				UObject* Outer = ObjectNode ? ObjectNode->GetUObject(i) : nullptr;
 				if(!Outer)
 				{
-					Implementation->ShowInvalidOperationError(LOCTEXT("SetOptionalElement", "Could not create value for optional object as could not determine outer object."));
+					Implementation->ShowInvalidOperationError(LOCTEXT("SetOptionalOuterError", "Could not create value for optional object as could not determine outer object."));
 				}
 				TopLevelObjects.Add(Outer);
 
@@ -5403,7 +5403,7 @@ FPropertyAccess::Result FPropertyHandleOptional::SetOptionalValue(FProperty* New
 				}
 				else
 				{
-					Implementation->ShowInvalidOperationError(LOCTEXT("SetOptionalElement", "Could not create value for optional object (support for this type may not yet be implemented)"));
+					Implementation->ShowInvalidOperationError(LOCTEXT("SetOptionalUnknownType", "Could not create value for optional object (support for this type may not yet be implemented)"));
 				}
 			}
 		}
