Fixed stack corruption caused by use after free when the FDetailLayoutBuilderImpl deconstructor calls FDetailLayoutBuilderImpl::ClearExternalRootPropertyNodes which calls SDetailsViewBase::SaveExpandedItems, but the SDetailsViewBase deconstructor has already run and moved on to the base SCompoundWidget deconstructor which is cleaning up SCompoundWidget::ChildSlot (causing the FDetailLayoutBuilderImpl to get destroyed).

diff --git a/Engine/Plugins/2D/Paper2D/Source/Paper2DEditor/Private/TileMapEditing/PaperTileMapDetailsCustomization.cpp b/Engine/Plugins/2D/Paper2D/Source/Paper2DEditor/Private/TileMapEditing/PaperTileMapDetailsCustomization.cpp
--- a/Engine/Plugins/2D/Paper2D/Source/Paper2DEditor/Private/TileMapEditing/PaperTileMapDetailsCustomization.cpp
+++ b/Engine/Plugins/2D/Paper2D/Source/Paper2DEditor/Private/TileMapEditing/PaperTileMapDetailsCustomization.cpp
@@ -194,7 +194,11 @@ void FPaperTileMapDetailsCustomization::CustomizeDetails(IDetailLayoutBuilder& D
 	}
 
 	// Try to get the hosting command list from the details view
-	TSharedPtr<FUICommandList> CommandList = DetailLayout.GetDetailsView()->GetHostCommandList();
+	TSharedPtr<FUICommandList> CommandList;
+	if (TSharedPtr<IDetailsView> DetailsView = DetailLayout.GetDetailsViewSharedPtr())
+	{
+		CommandList = DetailsView->GetHostCommandList();
+	}
 	if (!CommandList.IsValid())
 	{
 		CommandList = MakeShareable(new FUICommandList);
diff --git a/Engine/Plugins/Animation/ControlRig/Source/ControlRigEditor/Private/ControlRigModuleDetails.cpp b/Engine/Plugins/Animation/ControlRig/Source/ControlRigEditor/Private/ControlRigModuleDetails.cpp
--- a/Engine/Plugins/Animation/ControlRig/Source/ControlRigEditor/Private/ControlRigModuleDetails.cpp
+++ b/Engine/Plugins/Animation/ControlRig/Source/ControlRigEditor/Private/ControlRigModuleDetails.cpp
@@ -723,7 +723,7 @@ void FRigModuleInstanceDetails::CustomizeDetails(IDetailLayoutBuilder& DetailBui
 					})
 				);
 
-				TSharedPtr<SWidget> ValueWidget = DefaultProperty->CreatePropertyValueWidgetWithCustomization(DetailBuilder.GetDetailsView());
+				TSharedPtr<SWidget> ValueWidget = DefaultProperty->CreatePropertyValueWidgetWithCustomization(DetailBuilder.GetDetailsViewSharedPtr().Get());
 
 				const bool bShowChildren = true;
 				Category.AddProperty(DefaultProperty).CustomWidget(bShowChildren)
diff --git a/Engine/Plugins/Animation/DeformerGraph/Source/OptimusEditor/Private/OptimusDetailsCustomization.cpp b/Engine/Plugins/Animation/DeformerGraph/Source/OptimusEditor/Private/OptimusDetailsCustomization.cpp
--- a/Engine/Plugins/Animation/DeformerGraph/Source/OptimusEditor/Private/OptimusDetailsCustomization.cpp
+++ b/Engine/Plugins/Animation/DeformerGraph/Source/OptimusEditor/Private/OptimusDetailsCustomization.cpp
@@ -1665,7 +1665,10 @@ void FOptimusSourceDetailsCustomization::CustomizeDetails(IDetailLayoutBuilder&
 	}
 
 	TSharedRef<IPropertyHandle> SourcePropertyHandle = DetailBuilder.GetProperty(TEXT("SourceText"));
-	DetailBuilder.GetDetailsView()->HideFilterArea(true);
+	if (TSharedPtr<IDetailsView> DetailsView = DetailBuilder.GetDetailsViewSharedPtr())
+	{
+		DetailsView->HideFilterArea(true);
+	}
 	DetailBuilder.EditDefaultProperty(SourcePropertyHandle)->CustomWidget()
 	[
 		SNew(SBox)
diff --git a/Engine/Plugins/Experimental/ActorModifierCore/Source/ActorModifierCoreEditor/Private/Modifiers/Customizations/ActorModifierCoreEditorDetailCustomization.cpp b/Engine/Plugins/Experimental/ActorModifierCore/Source/ActorModifierCoreEditor/Private/Modifiers/Customizations/ActorModifierCoreEditorDetailCustomization.cpp
--- a/Engine/Plugins/Experimental/ActorModifierCore/Source/ActorModifierCoreEditor/Private/Modifiers/Customizations/ActorModifierCoreEditorDetailCustomization.cpp
+++ b/Engine/Plugins/Experimental/ActorModifierCore/Source/ActorModifierCoreEditor/Private/Modifiers/Customizations/ActorModifierCoreEditorDetailCustomization.cpp
@@ -30,9 +30,9 @@ void FActorModifierCoreEditorDetailCustomization::CustomizeDetails(IDetailLayout
 	);
 
 	TSharedPtr<IDetailKeyframeHandler> KeyframeHandler = nullptr;
-	if (InDetailBuilder.GetDetailsView())
+	if (TSharedPtr<IDetailsView> DetailsView = InDetailBuilder.GetDetailsViewSharedPtr())
 	{
-		KeyframeHandler = InDetailBuilder.GetDetailsView()->GetKeyframeHandler();
+		KeyframeHandler = DetailsView->GetKeyframeHandler();
 	}
 	
 	if (!ModifiersPropertyHandle.IsValid()
diff --git a/Engine/Plugins/Experimental/CustomDetailsView/Source/CustomDetailsView/Private/CustomDetailsViewRowExtensions.cpp b/Engine/Plugins/Experimental/CustomDetailsView/Source/CustomDetailsView/Private/CustomDetailsViewRowExtensions.cpp
--- a/Engine/Plugins/Experimental/CustomDetailsView/Source/CustomDetailsView/Private/CustomDetailsViewRowExtensions.cpp
+++ b/Engine/Plugins/Experimental/CustomDetailsView/Source/CustomDetailsView/Private/CustomDetailsViewRowExtensions.cpp
@@ -115,9 +115,12 @@ void FCustomDetailsViewRowExtensions::FillPropertyRightClickMenu(UToolMenu* InTo
 	const bool bCanCopy = CopyAction.ExecuteAction.IsBound();
 	bool bCanPaste = true;
 
-	if (RowMenuContext->DetailsView && !RowMenuContext->DetailsView->IsPropertyEditingEnabled())
+	if (TSharedPtr<IDetailsView> DetailsView = RowMenuContext->DetailsView.Pin())
 	{
-		bCanPaste = false;
+		if (!DetailsView->IsPropertyEditingEnabled())
+		{
+			bCanPaste = false;
+		}
 	}
 	else if (PropertyHandle->IsEditConst() || !PropertyHandle->IsEditable())
 	{
diff --git a/Engine/Plugins/Experimental/CustomDetailsView/Source/CustomDetailsView/Private/Items/CustomDetailsViewDetailTreeNodeItem.cpp b/Engine/Plugins/Experimental/CustomDetailsView/Source/CustomDetailsView/Private/Items/CustomDetailsViewDetailTreeNodeItem.cpp
--- a/Engine/Plugins/Experimental/CustomDetailsView/Source/CustomDetailsView/Private/Items/CustomDetailsViewDetailTreeNodeItem.cpp
+++ b/Engine/Plugins/Experimental/CustomDetailsView/Source/CustomDetailsView/Private/Items/CustomDetailsViewDetailTreeNodeItem.cpp
@@ -91,18 +91,18 @@ void FCustomDetailsViewDetailTreeNodeItem::InitWidget(const TSharedRef<IDetailTr
 	DetailWidgetRow.ExtensionWidget.Widget->SetEnabled(IsEnabledAttribute);
 }
 
-IDetailsView* FCustomDetailsViewDetailTreeNodeItem::GetDetailsView() const
+TSharedPtr<IDetailsView> FCustomDetailsViewDetailTreeNodeItem::GetDetailsView() const
 {
 	if (const TSharedPtr<IDetailTreeNode> DetailTreeNode = GetRowTreeNode())
 	{
-		return DetailTreeNode->GetNodeDetailsView();
+		return DetailTreeNode->GetNodeDetailsViewSharedPtr();
 	}
 
 	TSharedPtr<ICustomDetailsViewItem> ParentItem = GetParent();
 
 	while (ParentItem.IsValid())
 	{
-		if (IDetailsView* DetailsView = ParentItem->GetDetailsView())
+		if (TSharedPtr<IDetailsView> DetailsView = ParentItem->GetDetailsView())
 		{
 			return DetailsView;
 		}
diff --git a/Engine/Plugins/Experimental/CustomDetailsView/Source/CustomDetailsView/Private/Items/CustomDetailsViewDetailTreeNodeItem.h b/Engine/Plugins/Experimental/CustomDetailsView/Source/CustomDetailsView/Private/Items/CustomDetailsViewDetailTreeNodeItem.h
--- a/Engine/Plugins/Experimental/CustomDetailsView/Source/CustomDetailsView/Private/Items/CustomDetailsViewDetailTreeNodeItem.h
+++ b/Engine/Plugins/Experimental/CustomDetailsView/Source/CustomDetailsView/Private/Items/CustomDetailsViewDetailTreeNodeItem.h
@@ -30,7 +30,7 @@ public:
 	void InitWidget(const TSharedRef<IDetailTreeNode>& InDetailTreeNode);
 
 	//~ Begin ICustomDetailsViewItem
-	virtual IDetailsView* GetDetailsView() const override;
+	virtual TSharedPtr<IDetailsView> GetDetailsView() const override;
 	virtual void RefreshItemId() override;
 	virtual const TArray<TSharedPtr<ICustomDetailsViewItem>>& GetChildren() const override final { return Children; }
 	virtual void SetResetToDefaultOverride(const FResetToDefaultOverride& InOverride) override;
diff --git a/Engine/Plugins/Experimental/CustomDetailsView/Source/CustomDetailsView/Private/Items/CustomDetailsViewItemBase.h b/Engine/Plugins/Experimental/CustomDetailsView/Source/CustomDetailsView/Private/Items/CustomDetailsViewItemBase.h
--- a/Engine/Plugins/Experimental/CustomDetailsView/Source/CustomDetailsView/Private/Items/CustomDetailsViewItemBase.h
+++ b/Engine/Plugins/Experimental/CustomDetailsView/Source/CustomDetailsView/Private/Items/CustomDetailsViewItemBase.h
@@ -36,7 +36,7 @@ public:
 
 	//~ Begin ICustomDetailsViewItem
 	virtual TSharedPtr<ICustomDetailsView> GetCustomDetailsView() const override;
-	virtual IDetailsView* GetDetailsView() const override { return nullptr; }
+	virtual TSharedPtr<IDetailsView> GetDetailsView() const override { return nullptr; }
 	virtual const FCustomDetailsViewItemId& GetItemId() const override final;
 	virtual TSharedPtr<ICustomDetailsViewItem> GetRoot() const override;
 	virtual TSharedPtr<ICustomDetailsViewItem> GetParent() const override;
diff --git a/Engine/Plugins/Experimental/CustomDetailsView/Source/CustomDetailsView/Public/Items/ICustomDetailsViewItem.h b/Engine/Plugins/Experimental/CustomDetailsView/Source/CustomDetailsView/Public/Items/ICustomDetailsViewItem.h
--- a/Engine/Plugins/Experimental/CustomDetailsView/Source/CustomDetailsView/Public/Items/ICustomDetailsViewItem.h
+++ b/Engine/Plugins/Experimental/CustomDetailsView/Source/CustomDetailsView/Public/Items/ICustomDetailsViewItem.h
@@ -20,7 +20,7 @@ public:
 	virtual TSharedPtr<ICustomDetailsView> GetCustomDetailsView() const = 0;
 
 	/** Details view containing this item */
-	virtual IDetailsView* GetDetailsView() const = 0;
+	virtual TSharedPtr<IDetailsView> GetDetailsView() const = 0;
 
 	/** Called to update the Item's Id */
 	virtual void RefreshItemId() = 0;
diff --git a/Engine/Plugins/Experimental/DynamicMaterial/Source/DynamicMaterialEditor/Private/DetailsPanel/DMValueDetailsRowExtensions.cpp b/Engine/Plugins/Experimental/DynamicMaterial/Source/DynamicMaterialEditor/Private/DetailsPanel/DMValueDetailsRowExtensions.cpp
--- a/Engine/Plugins/Experimental/DynamicMaterial/Source/DynamicMaterialEditor/Private/DetailsPanel/DMValueDetailsRowExtensions.cpp
+++ b/Engine/Plugins/Experimental/DynamicMaterial/Source/DynamicMaterialEditor/Private/DetailsPanel/DMValueDetailsRowExtensions.cpp
@@ -120,9 +120,12 @@ void FDMValueDetailsRowExtensions::FillPropertyRightClickMenu(UToolMenu* InToolM
 		return;
 	}
 
-	if (RowMenuContext->DetailsView && !RowMenuContext->DetailsView->IsPropertyEditingEnabled())
+	if (TSharedPtr<IDetailsView> DetailsView = RowMenuContext->DetailsView.Pin())
 	{
-		return;
+		if (!DetailsView->IsPropertyEditingEnabled())
+		{
+			return;
+		}
 	}
 
 	if (PropertyHandle->IsEditConst() || !PropertyHandle->IsEditable())
diff --git a/Engine/Plugins/Experimental/MutablePopulation/Source/CustomizableObjectPopulationEditor/Private/MuCOPE/CustomizableObjectPopulationClassDetails.cpp b/Engine/Plugins/Experimental/MutablePopulation/Source/CustomizableObjectPopulationEditor/Private/MuCOPE/CustomizableObjectPopulationClassDetails.cpp
--- a/Engine/Plugins/Experimental/MutablePopulation/Source/CustomizableObjectPopulationEditor/Private/MuCOPE/CustomizableObjectPopulationClassDetails.cpp
+++ b/Engine/Plugins/Experimental/MutablePopulation/Source/CustomizableObjectPopulationEditor/Private/MuCOPE/CustomizableObjectPopulationClassDetails.cpp
@@ -44,8 +44,8 @@ void FCustomizableObjectPopulationClassDetails::CustomizeDetails(IDetailLayoutBu
 	PopulationClass = nullptr;
 	DetailBuilderPtr = &DetailBuilder;
 
-	const IDetailsView* DetailsView = DetailBuilder.GetDetailsView();
-	if (DetailsView->GetSelectedObjects().Num())
+	TSharedPtr<const IDetailsView> DetailsView = DetailBuilder.GetDetailsViewSharedPtr();
+	if (DetailsView && DetailsView->GetSelectedObjects().Num())
 	{
 		PopulationClass = Cast<UCustomizableObjectPopulationClass>(DetailsView->GetSelectedObjects()[0].Get());
 	}
diff --git a/Engine/Plugins/Media/MediaPlate/Source/MediaPlateEditor/Private/MediaPlateCustomization.cpp b/Engine/Plugins/Media/MediaPlate/Source/MediaPlateEditor/Private/MediaPlateCustomization.cpp
--- a/Engine/Plugins/Media/MediaPlate/Source/MediaPlateEditor/Private/MediaPlateCustomization.cpp
+++ b/Engine/Plugins/Media/MediaPlate/Source/MediaPlateEditor/Private/MediaPlateCustomization.cpp
@@ -50,7 +50,7 @@ void FMediaPlateCustomization::CustomizeDetails(IDetailLayoutBuilder& DetailBuil
 	// Is this the media plate editor window?
 	bool bIsMediaPlateWindow = false;
 
-	if (const IDetailsView* DetailsView = DetailBuilder.GetDetailsView())
+	if (TSharedPtr<const IDetailsView> DetailsView = DetailBuilder.GetDetailsViewSharedPtr())
 	{
 		TSharedPtr<FTabManager> HostTabManager = DetailsView->GetHostTabManager();
 		bIsMediaPlateWindow = (HostTabManager.IsValid() == false);
diff --git a/Engine/Plugins/MovieScene/ActorSequence/Source/ActorSequenceEditor/Private/ActorSequenceComponentCustomization.cpp b/Engine/Plugins/MovieScene/ActorSequence/Source/ActorSequenceEditor/Private/ActorSequenceComponentCustomization.cpp
--- a/Engine/Plugins/MovieScene/ActorSequence/Source/ActorSequenceEditor/Private/ActorSequenceComponentCustomization.cpp
+++ b/Engine/Plugins/MovieScene/ActorSequence/Source/ActorSequenceEditor/Private/ActorSequenceComponentCustomization.cpp
@@ -85,8 +85,11 @@ void FActorSequenceComponentCustomization::CustomizeDetails(IDetailLayoutBuilder
 		return;
 	}
 
-	const IDetailsView* DetailsView = DetailBuilder.GetDetailsView();
-	TSharedPtr<FTabManager> HostTabManager = DetailsView->GetHostTabManager();
+	TSharedPtr<FTabManager> HostTabManager;
+	if (TSharedPtr<const IDetailsView> DetailsView = DetailBuilder.GetDetailsViewSharedPtr())
+	{
+		HostTabManager = DetailsView->GetHostTabManager();
+	}
 
 	DetailBuilder.HideProperty("Sequence");
 
diff --git a/Engine/Plugins/MovieScene/LevelSequenceEditor/Source/LevelSequenceEditor/Private/LevelSequenceEditorSubsystem.cpp b/Engine/Plugins/MovieScene/LevelSequenceEditor/Source/LevelSequenceEditor/Private/LevelSequenceEditorSubsystem.cpp
--- a/Engine/Plugins/MovieScene/LevelSequenceEditor/Source/LevelSequenceEditor/Private/LevelSequenceEditorSubsystem.cpp
+++ b/Engine/Plugins/MovieScene/LevelSequenceEditor/Source/LevelSequenceEditor/Private/LevelSequenceEditorSubsystem.cpp
@@ -230,10 +230,10 @@ TSharedRef<SWidget> FMovieSceneBindingPropertyInfoListCustomization::GenerateBin
 		}
 		LevelSequenceEditorSubsystem->AddChangeBindingTypeMenu(MenuBuilder, Sequencer, Bindings, false, [this, Sequencer, &DetailBuilder]()
 			{
-				if (IDetailsView* DetailsView = DetailBuilder.GetDetailsView())
+				if (TSharedPtr<IDetailsView> DetailsView  = DetailBuilder.GetDetailsViewSharedPtr())
 				{
-					LevelSequenceEditorSubsystem->RefreshBindingDetails(DetailsView, BindingGuid);
-					LevelSequenceEditorSubsystem->OnFinishedChangingLocators(FPropertyChangedEvent(nullptr), StaticCastSharedRef<IDetailsView>(DetailsView->AsShared()), BindingGuid);
+					LevelSequenceEditorSubsystem->RefreshBindingDetails(DetailsView.Get(), BindingGuid);
+					LevelSequenceEditorSubsystem->OnFinishedChangingLocators(FPropertyChangedEvent(nullptr), DetailsView.ToSharedRef(), BindingGuid);
 				}
 			});
 	}
@@ -309,9 +309,9 @@ TSharedRef<SWidget> FMovieSceneBindingPropertyInfoListCustomization::OnGetConver
 
 		LevelSequenceEditorSubsystem->AddChangeBindingTypeMenu(MenuBuilder, Sequencer, Bindings, true, [this, &DetailBuilder]()
 			{
-				if (IDetailsView* DetailsView = DetailBuilder.GetDetailsView())
+				if (TSharedPtr<IDetailsView> DetailsView = DetailBuilder.GetDetailsViewSharedPtr())
 				{
-					LevelSequenceEditorSubsystem->RefreshBindingDetails(DetailsView, BindingGuid);
+					LevelSequenceEditorSubsystem->RefreshBindingDetails(DetailsView.Get(), BindingGuid);
 				}
 			});
 	}
@@ -455,9 +455,9 @@ TSharedRef<SWidget> FMovieSceneBindingPropertyInfoDetailCustomization::OnGetChan
 
 		FSequencerUtilities::AddChangeClassMenu(MenuBuilder, Sequencer.ToSharedRef(), Bindings, [this, &StructBuilder, &CustomizationUtils]()
 			{
-				if (IDetailsView* DetailsView = StructBuilder.GetParentCategory().GetParentLayout().GetDetailsView())
+				if (TSharedPtr<IDetailsView> DetailsView = StructBuilder.GetParentCategory().GetParentLayout().GetDetailsViewSharedPtr())
 				{
-					LevelSequenceEditorSubsystem->RefreshBindingDetails(DetailsView, BindingGuid);
+					LevelSequenceEditorSubsystem->RefreshBindingDetails(DetailsView.Get(), BindingGuid);
 				}
 			});
 	}
diff --git a/Engine/Plugins/Mutable/Source/CustomizableObjectEditor/Private/MuCOE/CustomizableInstanceDetails.cpp b/Engine/Plugins/Mutable/Source/CustomizableObjectEditor/Private/MuCOE/CustomizableInstanceDetails.cpp
--- a/Engine/Plugins/Mutable/Source/CustomizableObjectEditor/Private/MuCOE/CustomizableInstanceDetails.cpp
+++ b/Engine/Plugins/Mutable/Source/CustomizableObjectEditor/Private/MuCOE/CustomizableInstanceDetails.cpp
@@ -62,7 +62,8 @@ void FCustomizableInstanceDetails::CustomizeDetails(const TSharedPtr<IDetailLayo
 {
 	LayoutBuilder = DetailBuilder;
 
-	const IDetailsView* DetailsView = DetailBuilder->GetDetailsView();
+	TSharedPtr<const IDetailsView> DetailsView = DetailBuilder->GetDetailsViewSharedPtr();
+	check(DetailsView.IsValid());
 	check(DetailsView->GetSelectedObjects().Num());
 
 	CustomInstance = Cast<UCustomizableObjectInstance>(DetailsView->GetSelectedObjects()[0].Get());
diff --git a/Engine/Plugins/Mutable/Source/CustomizableObjectEditor/Private/MuCOE/CustomizableObjectCustomSettingsDetails.cpp b/Engine/Plugins/Mutable/Source/CustomizableObjectEditor/Private/MuCOE/CustomizableObjectCustomSettingsDetails.cpp
--- a/Engine/Plugins/Mutable/Source/CustomizableObjectEditor/Private/MuCOE/CustomizableObjectCustomSettingsDetails.cpp
+++ b/Engine/Plugins/Mutable/Source/CustomizableObjectEditor/Private/MuCOE/CustomizableObjectCustomSettingsDetails.cpp
@@ -24,8 +24,8 @@ TSharedRef<IDetailCustomization> FCustomizableObjectCustomSettingsDetails::MakeI
 
 void FCustomizableObjectCustomSettingsDetails::CustomizeDetails(IDetailLayoutBuilder& DetailBuilder)
 {
-	const IDetailsView* DetailsView = DetailBuilder.GetDetailsView();
-	if (DetailsView->GetSelectedObjects().Num() > 0 && DetailsView->GetSelectedObjects()[0].IsValid())
+	TSharedPtr<const IDetailsView> DetailsView = DetailBuilder.GetDetailsViewSharedPtr();
+	if (DetailsView && DetailsView->GetSelectedObjects().Num() > 0 && DetailsView->GetSelectedObjects()[0].IsValid())
 	{
 		UCustomSettings* CustomSettings = Cast<UCustomSettings>(DetailsView->GetSelectedObjects()[0]);
 
diff --git a/Engine/Plugins/Mutable/Source/CustomizableObjectEditor/Private/MuCOE/CustomizableObjectDetails.cpp b/Engine/Plugins/Mutable/Source/CustomizableObjectEditor/Private/MuCOE/CustomizableObjectDetails.cpp
--- a/Engine/Plugins/Mutable/Source/CustomizableObjectEditor/Private/MuCOE/CustomizableObjectDetails.cpp
+++ b/Engine/Plugins/Mutable/Source/CustomizableObjectEditor/Private/MuCOE/CustomizableObjectDetails.cpp
@@ -83,8 +83,8 @@ public:
 
 void FCustomizableObjectDetails::CustomizeDetails( IDetailLayoutBuilder& DetailBuilder )
 {
-	const IDetailsView* DetailsView = DetailBuilder.GetDetailsView();
-	if (DetailsView->GetSelectedObjects().Num())
+	TSharedPtr<const IDetailsView> DetailsView = DetailBuilder.GetDetailsViewSharedPtr();
+	if (DetailsView && DetailsView->GetSelectedObjects().Num())
 	{
 		CustomizableObject = Cast<UCustomizableObject>(DetailsView->GetSelectedObjects()[0].Get());
 	}
diff --git a/Engine/Plugins/Mutable/Source/CustomizableObjectEditor/Private/MuCOE/CustomizableObjectNodeObjectGroupDetails.cpp b/Engine/Plugins/Mutable/Source/CustomizableObjectEditor/Private/MuCOE/CustomizableObjectNodeObjectGroupDetails.cpp
--- a/Engine/Plugins/Mutable/Source/CustomizableObjectEditor/Private/MuCOE/CustomizableObjectNodeObjectGroupDetails.cpp
+++ b/Engine/Plugins/Mutable/Source/CustomizableObjectEditor/Private/MuCOE/CustomizableObjectNodeObjectGroupDetails.cpp
@@ -27,9 +27,9 @@ void FCustomizableObjectNodeObjectGroupDetails::CustomizeDetails(IDetailLayoutBu
 {
 	FCustomizableObjectNodeDetails::CustomizeDetails(DetailBuilder);
 
-	const IDetailsView* DetailsView = DetailBuilder.GetDetailsView();
+	TSharedPtr<const IDetailsView> DetailsView = DetailBuilder.GetDetailsViewSharedPtr();
 
-	if (DetailsView->GetSelectedObjects().Num())
+	if (DetailsView && DetailsView->GetSelectedObjects().Num())
 	{
 		if (DetailsView->GetSelectedObjects()[0].Get()->IsA(UCustomizableObjectNodeObjectGroup::StaticClass()))
 		{
diff --git a/Engine/Plugins/Mutable/Source/CustomizableObjectEditor/Private/MuCOE/Nodes/CustomizableObjectNodeDetails.cpp b/Engine/Plugins/Mutable/Source/CustomizableObjectEditor/Private/MuCOE/Nodes/CustomizableObjectNodeDetails.cpp
--- a/Engine/Plugins/Mutable/Source/CustomizableObjectEditor/Private/MuCOE/Nodes/CustomizableObjectNodeDetails.cpp
+++ b/Engine/Plugins/Mutable/Source/CustomizableObjectEditor/Private/MuCOE/Nodes/CustomizableObjectNodeDetails.cpp
@@ -21,8 +21,8 @@ TSharedRef<IDetailCustomization> FCustomizableObjectNodeDetails::MakeInstance()
 
 void FCustomizableObjectNodeDetails::CustomizeDetails(IDetailLayoutBuilder& DetailBuilder)
 {
-	const IDetailsView* DetailsView = DetailBuilder.GetDetailsView();
-	if (DetailsView->GetSelectedObjects().Num())
+	TSharedPtr<const IDetailsView> DetailsView = DetailBuilder.GetDetailsViewSharedPtr();
+	if (DetailsView && DetailsView->GetSelectedObjects().Num())
 	{
 		if (UCustomizableObjectNode* Node = Cast<UCustomizableObjectNode>(DetailsView->GetSelectedObjects()[0]))
 		{
diff --git a/Engine/Plugins/Mutable/Source/CustomizableObjectEditor/Private/MuCOE/Nodes/CustomizableObjectNodeExternalPinDetails.cpp b/Engine/Plugins/Mutable/Source/CustomizableObjectEditor/Private/MuCOE/Nodes/CustomizableObjectNodeExternalPinDetails.cpp
--- a/Engine/Plugins/Mutable/Source/CustomizableObjectEditor/Private/MuCOE/Nodes/CustomizableObjectNodeExternalPinDetails.cpp
+++ b/Engine/Plugins/Mutable/Source/CustomizableObjectEditor/Private/MuCOE/Nodes/CustomizableObjectNodeExternalPinDetails.cpp
@@ -29,8 +29,8 @@ void FCustomizableObjectNodeExternalPinDetails::CustomizeDetails( IDetailLayoutB
 	Node = nullptr;
 	DetailBuilderPtr = &DetailBuilder;
 	
-	const IDetailsView* DetailsView = DetailBuilder.GetDetailsView();
-	if (DetailsView->GetSelectedObjects().Num())
+	TSharedPtr<const IDetailsView> DetailsView = DetailBuilder.GetDetailsViewSharedPtr();
+	if (DetailsView && DetailsView->GetSelectedObjects().Num())
 	{
 		Node = Cast<UCustomizableObjectNodeExternalPin>(DetailsView->GetSelectedObjects()[0].Get());
 	}
diff --git a/Engine/Plugins/Mutable/Source/CustomizableObjectEditor/Private/MuCOE/Nodes/CustomizableObjectNodeMeshMorphDetails.cpp b/Engine/Plugins/Mutable/Source/CustomizableObjectEditor/Private/MuCOE/Nodes/CustomizableObjectNodeMeshMorphDetails.cpp
--- a/Engine/Plugins/Mutable/Source/CustomizableObjectEditor/Private/MuCOE/Nodes/CustomizableObjectNodeMeshMorphDetails.cpp
+++ b/Engine/Plugins/Mutable/Source/CustomizableObjectEditor/Private/MuCOE/Nodes/CustomizableObjectNodeMeshMorphDetails.cpp
@@ -26,8 +26,8 @@ void FCustomizableObjectNodeMeshMorphDetails::CustomizeDetails( IDetailLayoutBui
 {
 	FCustomizableObjectNodeDetails::CustomizeDetails(DetailBuilder);
 
-	const IDetailsView* DetailsView = DetailBuilder.GetDetailsView();
-	if ( DetailsView->GetSelectedObjects().Num() )
+	TSharedPtr<const IDetailsView> DetailsView = DetailBuilder.GetDetailsViewSharedPtr();
+	if (DetailsView && DetailsView->GetSelectedObjects().Num())
 	{
 		Node = Cast<UCustomizableObjectNodeMeshMorph>( DetailsView->GetSelectedObjects()[0].Get() );
 	}
diff --git a/Engine/Plugins/Mutable/Source/CustomizableObjectEditor/Private/MuCOE/Nodes/CustomizableObjectNodeMeshSectionDetails.cpp b/Engine/Plugins/Mutable/Source/CustomizableObjectEditor/Private/MuCOE/Nodes/CustomizableObjectNodeMeshSectionDetails.cpp
--- a/Engine/Plugins/Mutable/Source/CustomizableObjectEditor/Private/MuCOE/Nodes/CustomizableObjectNodeMeshSectionDetails.cpp
+++ b/Engine/Plugins/Mutable/Source/CustomizableObjectEditor/Private/MuCOE/Nodes/CustomizableObjectNodeMeshSectionDetails.cpp
@@ -31,8 +31,8 @@ void FCustomizableObjectNodeMeshSectionDetails::CustomizeDetails( IDetailLayoutB
 {
 	FCustomizableObjectNodeDetails::CustomizeDetails(DetailBuilder);
 
-	const IDetailsView* DetailsView = DetailBuilder.GetDetailsView();
-	if ( DetailsView->GetSelectedObjects().Num() )
+	TSharedPtr<const IDetailsView> DetailsView = DetailBuilder.GetDetailsViewSharedPtr();
+	if (DetailsView && DetailsView->GetSelectedObjects().Num())
 	{
 		Node = Cast<UCustomizableObjectNodeMaterial>( DetailsView->GetSelectedObjects()[0].Get() );
 	}
diff --git a/Engine/Plugins/Mutable/Source/CustomizableObjectEditor/Private/MuCOE/Nodes/CustomizableObjectNodeModifierBaseDetails.cpp b/Engine/Plugins/Mutable/Source/CustomizableObjectEditor/Private/MuCOE/Nodes/CustomizableObjectNodeModifierBaseDetails.cpp
--- a/Engine/Plugins/Mutable/Source/CustomizableObjectEditor/Private/MuCOE/Nodes/CustomizableObjectNodeModifierBaseDetails.cpp
+++ b/Engine/Plugins/Mutable/Source/CustomizableObjectEditor/Private/MuCOE/Nodes/CustomizableObjectNodeModifierBaseDetails.cpp
@@ -32,8 +32,8 @@ void FCustomizableObjectNodeModifierBaseDetails::CustomizeDetails( IDetailLayout
 {
 	FCustomizableObjectNodeDetails::CustomizeDetails(DetailBuilder);
 
-	const IDetailsView* DetailsView = DetailBuilder.GetDetailsView();
-	if ( DetailsView->GetSelectedObjects().Num() )
+	TSharedPtr<const IDetailsView> DetailsView = DetailBuilder.GetDetailsViewSharedPtr();
+	if (DetailsView && DetailsView->GetSelectedObjects().Num())
 	{
 		Node = Cast<UCustomizableObjectNodeModifierBase>( DetailsView->GetSelectedObjects()[0].Get() );
 	}
diff --git a/Engine/Plugins/Mutable/Source/CustomizableObjectEditor/Private/MuCOE/Nodes/CustomizableObjectNodeModifierClipMorphDetails.cpp b/Engine/Plugins/Mutable/Source/CustomizableObjectEditor/Private/MuCOE/Nodes/CustomizableObjectNodeModifierClipMorphDetails.cpp
--- a/Engine/Plugins/Mutable/Source/CustomizableObjectEditor/Private/MuCOE/Nodes/CustomizableObjectNodeModifierClipMorphDetails.cpp
+++ b/Engine/Plugins/Mutable/Source/CustomizableObjectEditor/Private/MuCOE/Nodes/CustomizableObjectNodeModifierClipMorphDetails.cpp
@@ -34,8 +34,8 @@ void FCustomizableObjectNodeModifierClipMorphDetails::CustomizeDetails(IDetailLa
 	Node = nullptr;
 	DetailBuilderPtr = &DetailBuilder;
 
-	const IDetailsView* DetailsView = DetailBuilder.GetDetailsView();
-	if (DetailsView->GetSelectedObjects().Num())
+	TSharedPtr<const IDetailsView> DetailsView = DetailBuilder.GetDetailsViewSharedPtr();
+	if (DetailsView && DetailsView->GetSelectedObjects().Num())
 	{
 		Node = Cast<UCustomizableObjectNodeModifierClipMorph>(DetailsView->GetSelectedObjects()[0].Get());
 	}
diff --git a/Engine/Plugins/Mutable/Source/CustomizableObjectEditor/Private/MuCOE/Nodes/CustomizableObjectNodeModifierClipWithMeshDetails.cpp b/Engine/Plugins/Mutable/Source/CustomizableObjectEditor/Private/MuCOE/Nodes/CustomizableObjectNodeModifierClipWithMeshDetails.cpp
--- a/Engine/Plugins/Mutable/Source/CustomizableObjectEditor/Private/MuCOE/Nodes/CustomizableObjectNodeModifierClipWithMeshDetails.cpp
+++ b/Engine/Plugins/Mutable/Source/CustomizableObjectEditor/Private/MuCOE/Nodes/CustomizableObjectNodeModifierClipWithMeshDetails.cpp
@@ -32,8 +32,8 @@ void FCustomizableObjectNodeModifierClipWithMeshDetails::CustomizeDetails(IDetai
 	Node = nullptr;
 	DetailBuilderPtr = &DetailBuilder;
 
-	const IDetailsView* DetailsView = DetailBuilder.GetDetailsView();
-	if (DetailsView->GetSelectedObjects().Num())
+	TSharedPtr<const IDetailsView> DetailsView = DetailBuilder.GetDetailsViewSharedPtr();
+	if (DetailsView && DetailsView->GetSelectedObjects().Num())
 	{
 		Node = Cast<UCustomizableObjectNodeModifierClipWithMesh>(DetailsView->GetSelectedObjects()[0].Get());
 	}
diff --git a/Engine/Plugins/Mutable/Source/CustomizableObjectEditor/Private/MuCOE/Nodes/CustomizableObjectNodeModifierEditMeshSectionDetails.cpp b/Engine/Plugins/Mutable/Source/CustomizableObjectEditor/Private/MuCOE/Nodes/CustomizableObjectNodeModifierEditMeshSectionDetails.cpp
--- a/Engine/Plugins/Mutable/Source/CustomizableObjectEditor/Private/MuCOE/Nodes/CustomizableObjectNodeModifierEditMeshSectionDetails.cpp
+++ b/Engine/Plugins/Mutable/Source/CustomizableObjectEditor/Private/MuCOE/Nodes/CustomizableObjectNodeModifierEditMeshSectionDetails.cpp
@@ -33,8 +33,8 @@ void FCustomizableObjectNodeModifierEditMeshSectionDetails::CustomizeDetails( ID
 	FCustomizableObjectNodeModifierBaseDetails::CustomizeDetails(DetailBuilder);
 
 	Node = nullptr;
-	const IDetailsView* DetailsView = DetailBuilder.GetDetailsView();
-	if (DetailsView->GetSelectedObjects().Num())
+	TSharedPtr<const IDetailsView> DetailsView = DetailBuilder.GetDetailsViewSharedPtr();
+	if (DetailsView && DetailsView->GetSelectedObjects().Num())
 	{
 		Node = Cast<UCustomizableObjectNodeModifierEditMeshSection>(DetailsView->GetSelectedObjects()[0].Get());
 	}
diff --git a/Engine/Plugins/Mutable/Source/CustomizableObjectEditor/Private/MuCOE/Nodes/CustomizableObjectNodeModifierExtendMeshSectionDetails.cpp b/Engine/Plugins/Mutable/Source/CustomizableObjectEditor/Private/MuCOE/Nodes/CustomizableObjectNodeModifierExtendMeshSectionDetails.cpp
--- a/Engine/Plugins/Mutable/Source/CustomizableObjectEditor/Private/MuCOE/Nodes/CustomizableObjectNodeModifierExtendMeshSectionDetails.cpp
+++ b/Engine/Plugins/Mutable/Source/CustomizableObjectEditor/Private/MuCOE/Nodes/CustomizableObjectNodeModifierExtendMeshSectionDetails.cpp
@@ -31,8 +31,8 @@ void FCustomizableObjectNodeModifierExtendMeshSectionDetails::CustomizeDetails(
 {
 	FCustomizableObjectNodeModifierBaseDetails::CustomizeDetails(DetailBuilder);
 
-	const IDetailsView* DetailsView = DetailBuilder.GetDetailsView();
-	if ( DetailsView->GetSelectedObjects().Num() )
+	TSharedPtr<const IDetailsView> DetailsView = DetailBuilder.GetDetailsViewSharedPtr();
+	if (DetailsView && DetailsView->GetSelectedObjects().Num())
 	{
 		Node = Cast<UCustomizableObjectNodeModifierExtendMeshSection>( DetailsView->GetSelectedObjects()[0].Get() );
 	}
diff --git a/Engine/Plugins/Mutable/Source/CustomizableObjectEditor/Private/MuCOE/Nodes/CustomizableObjectNodeModifierMorphMeshSectionDetails.cpp b/Engine/Plugins/Mutable/Source/CustomizableObjectEditor/Private/MuCOE/Nodes/CustomizableObjectNodeModifierMorphMeshSectionDetails.cpp
--- a/Engine/Plugins/Mutable/Source/CustomizableObjectEditor/Private/MuCOE/Nodes/CustomizableObjectNodeModifierMorphMeshSectionDetails.cpp
+++ b/Engine/Plugins/Mutable/Source/CustomizableObjectEditor/Private/MuCOE/Nodes/CustomizableObjectNodeModifierMorphMeshSectionDetails.cpp
@@ -35,8 +35,8 @@ void FCustomizableObjectNodeModifierMorphMeshSectionDetails::CustomizeDetails( I
 {
 	FCustomizableObjectNodeModifierBaseDetails::CustomizeDetails(DetailBuilder);
 
-	const IDetailsView* DetailsView = DetailBuilder.GetDetailsView();
-	if ( DetailsView->GetSelectedObjects().Num() )
+	TSharedPtr<const IDetailsView> DetailsView = DetailBuilder.GetDetailsViewSharedPtr();
+	if (DetailsView && DetailsView->GetSelectedObjects().Num())
 	{
 		Node = Cast<UCustomizableObjectNodeModifierMorphMeshSection>( DetailsView->GetSelectedObjects()[0].Get() );
 	}
diff --git a/Engine/Plugins/Mutable/Source/CustomizableObjectEditor/Private/MuCOE/Nodes/CustomizableObjectNodeModifierRemoveMeshBlocksDetails.cpp b/Engine/Plugins/Mutable/Source/CustomizableObjectEditor/Private/MuCOE/Nodes/CustomizableObjectNodeModifierRemoveMeshBlocksDetails.cpp
--- a/Engine/Plugins/Mutable/Source/CustomizableObjectEditor/Private/MuCOE/Nodes/CustomizableObjectNodeModifierRemoveMeshBlocksDetails.cpp
+++ b/Engine/Plugins/Mutable/Source/CustomizableObjectEditor/Private/MuCOE/Nodes/CustomizableObjectNodeModifierRemoveMeshBlocksDetails.cpp
@@ -31,8 +31,8 @@ void FCustomizableObjectNodeModifierRemoveMeshBlocksDetails::CustomizeDetails(ID
 	FCustomizableObjectNodeModifierBaseDetails::CustomizeDetails(DetailBuilder);
 
 	Node = nullptr;
-	const IDetailsView* DetailsView = DetailBuilder.GetDetailsView();
-	if (DetailsView->GetSelectedObjects().Num())
+	TSharedPtr<const IDetailsView> DetailsView = DetailBuilder.GetDetailsViewSharedPtr();
+	if (DetailsView && DetailsView->GetSelectedObjects().Num())
 	{
 		Node = Cast<UCustomizableObjectNodeModifierRemoveMeshBlocks>(DetailsView->GetSelectedObjects()[0].Get());
 	}
diff --git a/Engine/Plugins/Mutable/Source/CustomizableObjectEditor/Private/MuCOE/Nodes/CustomizableObjectNodeModifierTransformInMeshDetails.cpp b/Engine/Plugins/Mutable/Source/CustomizableObjectEditor/Private/MuCOE/Nodes/CustomizableObjectNodeModifierTransformInMeshDetails.cpp
--- a/Engine/Plugins/Mutable/Source/CustomizableObjectEditor/Private/MuCOE/Nodes/CustomizableObjectNodeModifierTransformInMeshDetails.cpp
+++ b/Engine/Plugins/Mutable/Source/CustomizableObjectEditor/Private/MuCOE/Nodes/CustomizableObjectNodeModifierTransformInMeshDetails.cpp
@@ -24,8 +24,8 @@ void FCustomizableObjectNodeModifierTransformInMeshDetails::CustomizeDetails(IDe
 	Node = nullptr;
 	DetailBuilderPtr = &DetailBuilder;
 
-	const IDetailsView* DetailsView = DetailBuilder.GetDetailsView();
-	if (DetailsView->GetSelectedObjects().Num())
+	TSharedPtr<const IDetailsView> DetailsView = DetailBuilder.GetDetailsViewSharedPtr();
+	if (DetailsView && DetailsView->GetSelectedObjects().Num())
 	{
 		Node = Cast<UCustomizableObjectNodeModifierTransformInMesh>(DetailsView->GetSelectedObjects()[0].Get());
 	}
diff --git a/Engine/Plugins/Mutable/Source/CustomizableObjectEditor/Private/MuCOE/Nodes/CustomizableObjectNodeObjectDetails.cpp b/Engine/Plugins/Mutable/Source/CustomizableObjectEditor/Private/MuCOE/Nodes/CustomizableObjectNodeObjectDetails.cpp
--- a/Engine/Plugins/Mutable/Source/CustomizableObjectEditor/Private/MuCOE/Nodes/CustomizableObjectNodeObjectDetails.cpp
+++ b/Engine/Plugins/Mutable/Source/CustomizableObjectEditor/Private/MuCOE/Nodes/CustomizableObjectNodeObjectDetails.cpp
@@ -161,8 +161,8 @@ void FCustomizableObjectNodeObjectDetails::CustomizeDetails( IDetailLayoutBuilde
 	BaseObjectNode = nullptr;
 	DetailBuilderPtr = &DetailBuilder;
 
-	const IDetailsView* DetailsView = DetailBuilder.GetDetailsView();
-	if ( DetailsView->GetSelectedObjects().Num() )
+	TSharedPtr<const IDetailsView> DetailsView = DetailBuilder.GetDetailsViewSharedPtr();
+	if (DetailsView && DetailsView->GetSelectedObjects().Num())
 	{
 		BaseObjectNode = Cast<UCustomizableObjectNodeObject>( DetailsView->GetSelectedObjects()[0].Get() );
 	}
diff --git a/Engine/Plugins/Mutable/Source/CustomizableObjectEditor/Private/MuCOE/Nodes/CustomizableObjectNodeProjectorParameterDetails.cpp b/Engine/Plugins/Mutable/Source/CustomizableObjectEditor/Private/MuCOE/Nodes/CustomizableObjectNodeProjectorParameterDetails.cpp
--- a/Engine/Plugins/Mutable/Source/CustomizableObjectEditor/Private/MuCOE/Nodes/CustomizableObjectNodeProjectorParameterDetails.cpp
+++ b/Engine/Plugins/Mutable/Source/CustomizableObjectEditor/Private/MuCOE/Nodes/CustomizableObjectNodeProjectorParameterDetails.cpp
@@ -35,8 +35,8 @@ void FCustomizableObjectNodeProjectorParameterDetails::CustomizeDetails( IDetail
 	NodeParameter = nullptr;
 	DetailBuilderPtr = &DetailBuilder;
 
-	const IDetailsView* DetailsView = DetailBuilder.GetDetailsView();
-	if ( DetailsView->GetSelectedObjects().Num() )
+	TSharedPtr<const IDetailsView> DetailsView = DetailBuilder.GetDetailsViewSharedPtr();
+	if (DetailsView && DetailsView->GetSelectedObjects().Num())
 	{
 		if (DetailsView->GetSelectedObjects()[0].Get()->IsA(UCustomizableObjectNodeProjectorConstant::StaticClass()))
 		{
diff --git a/Engine/Plugins/Mutable/Source/CustomizableObjectEditor/Private/MuCOE/Nodes/CustomizableObjectNodeSkeletalMeshDetails.cpp b/Engine/Plugins/Mutable/Source/CustomizableObjectEditor/Private/MuCOE/Nodes/CustomizableObjectNodeSkeletalMeshDetails.cpp
--- a/Engine/Plugins/Mutable/Source/CustomizableObjectEditor/Private/MuCOE/Nodes/CustomizableObjectNodeSkeletalMeshDetails.cpp
+++ b/Engine/Plugins/Mutable/Source/CustomizableObjectEditor/Private/MuCOE/Nodes/CustomizableObjectNodeSkeletalMeshDetails.cpp
@@ -23,9 +23,9 @@ void FCustomizableObjectNodeSkeletalMeshDetails::CustomizeDetails(IDetailLayoutB
 {
 	FCustomizableObjectNodeDetails::CustomizeDetails(DetailBuilder);
 
-	const IDetailsView* DetailsView = DetailBuilder.GetDetailsView();
+	TSharedPtr<const IDetailsView> DetailsView = DetailBuilder.GetDetailsViewSharedPtr();
 
-    if (DetailsView->GetSelectedObjects().Num())
+    if (DetailsView && DetailsView->GetSelectedObjects().Num())
 	{
 		Node = Cast<UCustomizableObjectNodeSkeletalMesh>(DetailsView->GetSelectedObjects()[0].Get());
 	}
diff --git a/Engine/Plugins/Mutable/Source/CustomizableObjectEditor/Private/MuCOE/Nodes/CustomizableObjectNodeTableDetails.cpp b/Engine/Plugins/Mutable/Source/CustomizableObjectEditor/Private/MuCOE/Nodes/CustomizableObjectNodeTableDetails.cpp
--- a/Engine/Plugins/Mutable/Source/CustomizableObjectEditor/Private/MuCOE/Nodes/CustomizableObjectNodeTableDetails.cpp
+++ b/Engine/Plugins/Mutable/Source/CustomizableObjectEditor/Private/MuCOE/Nodes/CustomizableObjectNodeTableDetails.cpp
@@ -41,9 +41,9 @@ void FCustomizableObjectNodeTableDetails::CustomizeDetails(const TSharedPtr<IDet
 	Node = 0;
 	DetailBuilderPtr = DetailBuilder;
 
-	const IDetailsView* DetailsView = DetailBuilder->GetDetailsView();
+	TSharedPtr<const IDetailsView> DetailsView = DetailBuilder->GetDetailsViewSharedPtr();
 
-	if (DetailsView->GetSelectedObjects().Num())
+	if (DetailsView && DetailsView->GetSelectedObjects().Num())
 	{
 		Node = Cast<UCustomizableObjectNodeTable>(DetailsView->GetSelectedObjects()[0].Get());
 	}
diff --git a/Engine/Plugins/Mutable/Source/CustomizableObjectEditor/Private/MuCOE/Nodes/FCustomizableObjectNodeComponentMeshDetails.cpp b/Engine/Plugins/Mutable/Source/CustomizableObjectEditor/Private/MuCOE/Nodes/FCustomizableObjectNodeComponentMeshDetails.cpp
--- a/Engine/Plugins/Mutable/Source/CustomizableObjectEditor/Private/MuCOE/Nodes/FCustomizableObjectNodeComponentMeshDetails.cpp
+++ b/Engine/Plugins/Mutable/Source/CustomizableObjectEditor/Private/MuCOE/Nodes/FCustomizableObjectNodeComponentMeshDetails.cpp
@@ -1,4 +1,4 @@
-﻿// Copyright Epic Games, Inc. All Rights Reserved.
+// Copyright Epic Games, Inc. All Rights Reserved.
 
 #include "CustomizableObjectNodeComponentMeshDetails.h"
 
@@ -26,7 +26,12 @@ void FCustomizableObjectNodeComponentMeshDetails::CustomizeDetails(IDetailLayout
 	FCustomizableObjectNodeDetails::CustomizeDetails(DetailBuilder);
 
 	
-	const IDetailsView* DetailsView = DetailBuilder.GetDetailsView();
+	TSharedPtr<const IDetailsView> DetailsView = DetailBuilder.GetDetailsViewSharedPtr();
+	if (!DetailsView)
+	{
+		return;
+	}
+
 	if (DetailsView->GetSelectedObjects().IsEmpty())
 	{
 		return;
diff --git a/Engine/Plugins/Mutable/Source/CustomizableObjectEditor/Private/MuCOE/SCustomizableObjectLayoutEditor.cpp b/Engine/Plugins/Mutable/Source/CustomizableObjectEditor/Private/MuCOE/SCustomizableObjectLayoutEditor.cpp
--- a/Engine/Plugins/Mutable/Source/CustomizableObjectEditor/Private/MuCOE/SCustomizableObjectLayoutEditor.cpp
+++ b/Engine/Plugins/Mutable/Source/CustomizableObjectEditor/Private/MuCOE/SCustomizableObjectLayoutEditor.cpp
@@ -65,8 +65,8 @@ void FCustomizableObjectLayoutEditorDetailsBuilder::CustomizeDetails(IDetailLayo
 {
 	TWeakObjectPtr<UObject> Node = nullptr;
 
-	const IDetailsView* DetailsView = DetailsBuilder.GetDetailsView();
-	if (DetailsView->GetSelectedObjects().Num())
+	TSharedPtr<const IDetailsView> DetailsView = DetailsBuilder.GetDetailsViewSharedPtr();
+	if (DetailsView && DetailsView->GetSelectedObjects().Num())
 	{
 		Node = DetailsView->GetSelectedObjects()[0].Get();
 	}
diff --git a/Engine/Plugins/PCG/Source/PCGEditor/Private/Details/PCGEditableUserParameterDetails.cpp b/Engine/Plugins/PCG/Source/PCGEditor/Private/Details/PCGEditableUserParameterDetails.cpp
--- a/Engine/Plugins/PCG/Source/PCGEditor/Private/Details/PCGEditableUserParameterDetails.cpp
+++ b/Engine/Plugins/PCG/Source/PCGEditor/Private/Details/PCGEditableUserParameterDetails.cpp
@@ -49,7 +49,7 @@ void FPCGEditableUserParameterDetails::CustomizeDetails(IDetailLayoutBuilder& De
 	}
 
 	// Cache the parent details view to be used on tick
-	ParentDetailsView = StaticCastWeakPtr<IDetailsView>(DetailBuilder.GetDetailsView()->AsWeak());
+	ParentDetailsView = DetailBuilder.GetDetailsViewSharedPtr();
 
 	if (UPCGUserParameterGetSettings* Settings = Cast<UPCGUserParameterGetSettings>(SettingsObject.Get()))
 	{
diff --git a/Engine/Plugins/Runtime/RigVM/Source/RigVMEditor/Private/Editor/RigVMGraphDetailCustomization.cpp b/Engine/Plugins/Runtime/RigVM/Source/RigVMEditor/Private/Editor/RigVMGraphDetailCustomization.cpp
--- a/Engine/Plugins/Runtime/RigVM/Source/RigVMEditor/Private/Editor/RigVMGraphDetailCustomization.cpp
+++ b/Engine/Plugins/Runtime/RigVM/Source/RigVMEditor/Private/Editor/RigVMGraphDetailCustomization.cpp
@@ -2425,7 +2425,7 @@ void FRigVMWrappedNodeDetailCustomization::CustomizeDetails(IDetailLayoutBuilder
         	}
         }
 
-		TSharedPtr<SWidget> ValueWidget = PropertyHandle->CreatePropertyValueWidgetWithCustomization(DetailLayout.GetDetailsView());
+		TSharedPtr<SWidget> ValueWidget = PropertyHandle->CreatePropertyValueWidgetWithCustomization(DetailLayout.GetDetailsViewSharedPtr().Get());
 
 		constexpr bool bShowChildren = true;
 		Row->CustomWidget(bShowChildren)
diff --git a/Engine/Plugins/VirtualProduction/RemoteControl/Source/RemoteControlUI/Private/RemoteControlUIModule.cpp b/Engine/Plugins/VirtualProduction/RemoteControl/Source/RemoteControlUI/Private/RemoteControlUIModule.cpp
--- a/Engine/Plugins/VirtualProduction/RemoteControl/Source/RemoteControlUI/Private/RemoteControlUIModule.cpp
+++ b/Engine/Plugins/VirtualProduction/RemoteControl/Source/RemoteControlUI/Private/RemoteControlUIModule.cpp
@@ -1342,7 +1342,7 @@ UMeshComponent* FRemoteControlUIModule::GetSelectedMeshComponentToBeModified(UOb
 
 	if (TSharedPtr<IDetailTreeNode> OwnerTreeNode = WeakDetailsTreeNode.Pin())
 	{
-		if (IDetailsView* DetailsView = OwnerTreeNode->GetNodeDetailsView())
+		if (TSharedPtr<IDetailsView> DetailsView = OwnerTreeNode->GetNodeDetailsViewSharedPtr())
 		{
 			SelectedActors = DetailsView->GetSelectedActors();
 
@@ -1426,7 +1426,7 @@ void FRemoteControlUIModule::RefreshPanels()
 {
 	if (TSharedPtr<IDetailTreeNode> OwnerTreeNode = WeakDetailsTreeNode.Pin())
 	{
-		if (IDetailsView* DetailsView = OwnerTreeNode->GetNodeDetailsView())
+		if (TSharedPtr<IDetailsView> DetailsView = OwnerTreeNode->GetNodeDetailsViewSharedPtr())
 		{
 			DetailsView->ForceRefresh();
 		}
diff --git a/Engine/Plugins/VirtualProduction/Takes/Source/TakeRecorder/Private/Widgets/SLevelSequenceTakeEditor.cpp b/Engine/Plugins/VirtualProduction/Takes/Source/TakeRecorder/Private/Widgets/SLevelSequenceTakeEditor.cpp
--- a/Engine/Plugins/VirtualProduction/Takes/Source/TakeRecorder/Private/Widgets/SLevelSequenceTakeEditor.cpp
+++ b/Engine/Plugins/VirtualProduction/Takes/Source/TakeRecorder/Private/Widgets/SLevelSequenceTakeEditor.cpp
@@ -831,7 +831,7 @@ class FRecorderSourceObjectCustomization : public IDetailCustomization
 {
 	virtual void CustomizeDetails(IDetailLayoutBuilder& DetailBuilder) override
 	{
-		FText NewTitle = ComputeTitle(DetailBuilder.GetDetailsView());
+		FText NewTitle = ComputeTitle(DetailBuilder.GetDetailsViewSharedPtr());
 		if (!NewTitle.IsEmpty())
 		{
 			// Edit the category and add *all* properties for the object to it
@@ -850,8 +850,13 @@ class FRecorderSourceObjectCustomization : public IDetailCustomization
 		}
 	}
 
-	FText ComputeTitle(const IDetailsView* DetailsView) const
+	FText ComputeTitle(const TSharedPtr<const IDetailsView>& DetailsView) const
 	{
+		if (!DetailsView)
+		{
+			return FText();
+		}
+
 		// Compute the title for all the sources that this details panel is editing
 		static const FName CategoryName = "Category";
 
diff --git a/Engine/Source/Editor/DetailCustomizations/Private/ActorDetails.cpp b/Engine/Source/Editor/DetailCustomizations/Private/ActorDetails.cpp
--- a/Engine/Source/Editor/DetailCustomizations/Private/ActorDetails.cpp
+++ b/Engine/Source/Editor/DetailCustomizations/Private/ActorDetails.cpp
@@ -86,7 +86,7 @@ void FActorDetails::CustomizeDetails( IDetailLayoutBuilder& DetailLayout )
 	}
 
 	// These details only apply when adding an instance of the actor in a level
-	if (!DetailLayout.HasClassDefaultObject() && DetailLayout.GetDetailsView() && DetailLayout.GetDetailsView()->GetSelectedActorInfo().NumSelected > 0)
+	if (TSharedPtr<IDetailsView> DetailsView = DetailLayout.GetDetailsViewSharedPtr(); !DetailLayout.HasClassDefaultObject() && DetailsView && DetailsView->GetSelectedActorInfo().NumSelected > 0)
 	{
 		// Build up a list of unique blueprints in the selection set (recording the first actor in the set for each one)
 		TMap<UBlueprint*, UObject*> UniqueBlueprints;
@@ -466,7 +466,7 @@ void FActorDetails::AddLayersCategory( IDetailLayoutBuilder& DetailBuilder )
 
 void FActorDetails::AddTransformCategory( IDetailLayoutBuilder& DetailBuilder )
 {
-	const FSelectedActorInfo& SelectedActorInfo = DetailBuilder.GetDetailsView()->GetSelectedActorInfo();
+	const FSelectedActorInfo& SelectedActorInfo = DetailBuilder.GetDetailsViewSharedPtr()->GetSelectedActorInfo();
 
 	bool bAreBrushesSelected = SelectedActorInfo.bHaveBrush;
 	bool bIsOnlyWorldPropsSelected =  SelectedActors.Num() == 1 && SelectedActors[0].IsValid() && SelectedActors[0]->IsA<AWorldSettings>();
@@ -588,7 +588,7 @@ void FActorDetails::AddActorCategory( IDetailLayoutBuilder& DetailBuilder, const
 	const FLevelEditorCommands& Commands = LevelEditor.GetLevelEditorCommands();
 	TSharedRef<const FUICommandList> CommandBindings = LevelEditor.GetGlobalLevelEditorActions();
 
-	const FSelectedActorInfo& SelectedActorInfo = DetailBuilder.GetDetailsView()->GetSelectedActorInfo();
+	const FSelectedActorInfo& SelectedActorInfo = DetailBuilder.GetDetailsViewSharedPtr()->GetSelectedActorInfo();
 	TSharedPtr<SVerticalBox> LevelBox;
 
 	IDetailCategoryBuilder& ActorCategory = DetailBuilder.EditCategory("Actor", FText::GetEmpty(), ECategoryPriority::Uncommon );
diff --git a/Engine/Source/Editor/DetailCustomizations/Private/MeshComponentDetails.cpp b/Engine/Source/Editor/DetailCustomizations/Private/MeshComponentDetails.cpp
--- a/Engine/Source/Editor/DetailCustomizations/Private/MeshComponentDetails.cpp
+++ b/Engine/Source/Editor/DetailCustomizations/Private/MeshComponentDetails.cpp
@@ -23,7 +23,8 @@ void FMeshComponentDetails::CustomizeDetails( IDetailLayoutBuilder& DetailLayout
 	if( MaterialProperty->IsValidHandle() )
 	{
 		// Only show this in the advanced section of the category if we have selected actors (which will show a separate material section)
-		bool bIsAdvanced = DetailLayout.GetDetailsView() && DetailLayout.GetDetailsView()->GetSelectedActorInfo().NumSelected > 0;
+		TSharedPtr<IDetailsView> DetailsView = DetailLayout.GetDetailsViewSharedPtr();
+		bool bIsAdvanced = DetailsView && DetailsView->GetSelectedActorInfo().NumSelected > 0;
 
 		RenderingCategory->AddProperty( MaterialProperty, bIsAdvanced ? EPropertyLocation::Advanced : EPropertyLocation::Default );
 	}
diff --git a/Engine/Source/Editor/DetailCustomizations/Private/MovieSceneEventParametersCustomization.cpp b/Engine/Source/Editor/DetailCustomizations/Private/MovieSceneEventParametersCustomization.cpp
--- a/Engine/Source/Editor/DetailCustomizations/Private/MovieSceneEventParametersCustomization.cpp
+++ b/Engine/Source/Editor/DetailCustomizations/Private/MovieSceneEventParametersCustomization.cpp
@@ -90,18 +90,19 @@ void FMovieSceneEventParametersCustomization::CustomizeChildren(TSharedRef<IProp
 		}
 	}
 
-	TSharedRef<const SWidget> DetailsView = ChildBuilder.GetParentCategory().GetParentLayout().GetDetailsView()->AsShared();
-
-	PropertyUtilities->EnqueueDeferredAction(FSimpleDelegate::CreateLambda(
-		[DetailsView]
-		{
-			FWidgetPath WidgetPath;
-			bool bFound = FSlateApplication::Get().FindPathToWidget(DetailsView, WidgetPath);
-			if (bFound)
+	if (TSharedPtr<const SWidget> DetailsView = ChildBuilder.GetParentCategory().GetParentLayout().GetDetailsViewSharedPtr())
+	{
+		PropertyUtilities->EnqueueDeferredAction(FSimpleDelegate::CreateLambda(
+			[DetailsView]
 			{
-				FSlateApplication::Get().SetAllUserFocus(WidgetPath, EFocusCause::SetDirectly);
-			}
-		}));
+				FWidgetPath WidgetPath;
+				bool bFound = FSlateApplication::Get().FindPathToWidget(DetailsView.ToSharedRef(), WidgetPath);
+				if (bFound)
+				{
+					FSlateApplication::Get().SetAllUserFocus(WidgetPath, EFocusCause::SetDirectly);
+				}
+			}));
+	}
 }
 
 void FMovieSceneEventParametersCustomization::OnStructChanged(const FAssetData& AssetData)
diff --git a/Engine/Source/Editor/DetailCustomizations/Private/SceneComponentDetails.cpp b/Engine/Source/Editor/DetailCustomizations/Private/SceneComponentDetails.cpp
--- a/Engine/Source/Editor/DetailCustomizations/Private/SceneComponentDetails.cpp
+++ b/Engine/Source/Editor/DetailCustomizations/Private/SceneComponentDetails.cpp
@@ -242,10 +242,16 @@ void FSceneComponentDetails::CustomizeDetails( IDetailLayoutBuilder& DetailBuild
 	// Details panel for BP component will have the base class be the Actor due to how the SKismetInspector works, but in that case we
 	// have a class default object selected, so use that to infer that this is the component directly selected and since BPs do not do
 	// property flattening it all kind of works
-	if (DetailBuilder.GetBaseClass()->IsChildOf<AActor>() && DetailBuilder.GetDetailsView() && !DetailBuilder.GetDetailsView()->HasClassDefaultObject())
+	if (DetailBuilder.GetBaseClass()->IsChildOf<AActor>())
 	{
-		TSharedPtr<IPropertyHandle> ComponentHiddenInGameProperty = DetailBuilder.GetProperty(GET_MEMBER_NAME_CHECKED(USceneComponent, bHiddenInGame));
-		ComponentHiddenInGameProperty->MarkHiddenByCustomization();
+		if (TSharedPtr<IDetailsView> DetailsView = DetailBuilder.GetDetailsViewSharedPtr())
+		{
+			if (!DetailsView->HasClassDefaultObject())
+			{
+				TSharedPtr<IPropertyHandle> ComponentHiddenInGameProperty = DetailBuilder.GetProperty(GET_MEMBER_NAME_CHECKED(USceneComponent, bHiddenInGame));
+				ComponentHiddenInGameProperty->MarkHiddenByCustomization();
+			}
+		}
 	}
 }
 
@@ -254,7 +260,7 @@ void FSceneComponentDetails::MakeTransformDetails( IDetailLayoutBuilder& DetailB
 	int SelectedActorsNum = 0;
 	const FSelectedActorInfo* SelectedActorInfoPtr = nullptr;
 
-	const IDetailsView* DetailsView = DetailBuilder.GetDetailsView();
+	TSharedPtr<const IDetailsView> DetailsView = DetailBuilder.GetDetailsViewSharedPtr();
 	if (DetailsView)
 	{
 		SelectedActorsNum = DetailsView->GetSelectedActors().Num();
diff --git a/Engine/Source/Editor/Kismet/Private/BlueprintDetailsCustomization.cpp b/Engine/Source/Editor/Kismet/Private/BlueprintDetailsCustomization.cpp
--- a/Engine/Source/Editor/Kismet/Private/BlueprintDetailsCustomization.cpp
+++ b/Engine/Source/Editor/Kismet/Private/BlueprintDetailsCustomization.cpp
@@ -6529,7 +6529,7 @@ bool FBlueprintGraphActionDetails::IsAddNewInputOutputEnabled() const
 {
 	if (DetailsLayoutPtr)
 	{
-		if (const IDetailsView* DetailsView = DetailsLayoutPtr->GetDetailsView())
+		if (TSharedPtr<const IDetailsView> DetailsView = DetailsLayoutPtr->GetDetailsViewSharedPtr())
 		{
 			return DetailsView->IsPropertyEditingEnabled();
 		}
diff --git a/Engine/Source/Editor/Kismet/Private/UserDefinedStructureEditor.cpp b/Engine/Source/Editor/Kismet/Private/UserDefinedStructureEditor.cpp
--- a/Engine/Source/Editor/Kismet/Private/UserDefinedStructureEditor.cpp
+++ b/Engine/Source/Editor/Kismet/Private/UserDefinedStructureEditor.cpp
@@ -119,8 +119,10 @@ void FDefaultValueDetails::CustomizeDetails(class IDetailLayoutBuilder& DetailLa
 	{
 		UserDefinedStruct = CastChecked<UUserDefinedStruct>(Objects[0].Get());
 
-		const IDetailsView* DetailsView = DetailLayout.GetDetailsView();
-		DetailsView->OnFinishedChangingProperties().AddSP(this, &FDefaultValueDetails::OnFinishedChangingProperties);
+		if (TSharedPtr<const IDetailsView> DetailsView = DetailLayout.GetDetailsViewSharedPtr())
+		{
+			DetailsView->OnFinishedChangingProperties().AddSP(this, &FDefaultValueDetails::OnFinishedChangingProperties);
+		}
 
 		IDetailCategoryBuilder& StructureCategory = DetailLayout.EditCategory("DefaultValues", LOCTEXT("DefaultValues", "Default Values"));
 
diff --git a/Engine/Source/Editor/LocalizationDashboard/Private/LocalizationTargetDetailCustomization.cpp b/Engine/Source/Editor/LocalizationDashboard/Private/LocalizationTargetDetailCustomization.cpp
--- a/Engine/Source/Editor/LocalizationDashboard/Private/LocalizationTargetDetailCustomization.cpp
+++ b/Engine/Source/Editor/LocalizationDashboard/Private/LocalizationTargetDetailCustomization.cpp
@@ -924,7 +924,7 @@ void FLocalizationTargetDetailCustomization::GatherText()
 		}
 
 		// Execute gather.
-		const TSharedPtr<SWindow> ParentWindow = FSlateApplication::Get().FindWidgetWindow(DetailLayoutBuilder->GetDetailsView()->AsShared());
+		const TSharedPtr<SWindow> ParentWindow = FSlateApplication::Get().FindWidgetWindow(DetailLayoutBuilder->GetDetailsViewSharedPtr().ToSharedRef());
 		LocalizationCommandletTasks::GatherTextForTarget(ParentWindow.ToSharedRef(), LocalizationTarget.Get());
 
 		UpdateTargetFromReports();
@@ -942,7 +942,7 @@ void FLocalizationTargetDetailCustomization::ImportTextAllCultures()
 	if (LocalizationTarget.IsValid() && DesktopPlatform)
 	{
 		void* ParentWindowWindowHandle = NULL;
-		const TSharedPtr<SWindow> ParentWindow = FSlateApplication::Get().FindWidgetWindow(DetailLayoutBuilder->GetDetailsView()->AsShared());
+		const TSharedPtr<SWindow> ParentWindow = FSlateApplication::Get().FindWidgetWindow(DetailLayoutBuilder->GetDetailsViewSharedPtr().ToSharedRef());
 		if (ParentWindow.IsValid() && ParentWindow->GetNativeWindow().IsValid())
 		{
 			ParentWindowWindowHandle = ParentWindow->GetNativeWindow()->GetOSWindowHandle();
@@ -979,7 +979,7 @@ void FLocalizationTargetDetailCustomization::ExportTextAllCultures()
 	if (LocalizationTarget.IsValid() && DesktopPlatform)
 	{
 		void* ParentWindowWindowHandle = NULL;
-		const TSharedPtr<SWindow> ParentWindow = FSlateApplication::Get().FindWidgetWindow(DetailLayoutBuilder->GetDetailsView()->AsShared());
+		const TSharedPtr<SWindow> ParentWindow = FSlateApplication::Get().FindWidgetWindow(DetailLayoutBuilder->GetDetailsViewSharedPtr().ToSharedRef());
 		if (ParentWindow.IsValid() && ParentWindow->GetNativeWindow().IsValid())
 		{
 			ParentWindowWindowHandle = ParentWindow->GetNativeWindow()->GetOSWindowHandle();
@@ -1014,7 +1014,7 @@ void FLocalizationTargetDetailCustomization::ImportDialogueScriptAllCultures()
 	if (LocalizationTarget.IsValid() && DesktopPlatform)
 	{
 		void* ParentWindowWindowHandle = NULL;
-		const TSharedPtr<SWindow> ParentWindow = FSlateApplication::Get().FindWidgetWindow(DetailLayoutBuilder->GetDetailsView()->AsShared());
+		const TSharedPtr<SWindow> ParentWindow = FSlateApplication::Get().FindWidgetWindow(DetailLayoutBuilder->GetDetailsViewSharedPtr().ToSharedRef());
 		if (ParentWindow.IsValid() && ParentWindow->GetNativeWindow().IsValid())
 		{
 			ParentWindowWindowHandle = ParentWindow->GetNativeWindow()->GetOSWindowHandle();
@@ -1051,7 +1051,7 @@ void FLocalizationTargetDetailCustomization::ExportDialogueScriptAllCultures()
 	if (LocalizationTarget.IsValid() && DesktopPlatform)
 	{
 		void* ParentWindowWindowHandle = NULL;
-		const TSharedPtr<SWindow> ParentWindow = FSlateApplication::Get().FindWidgetWindow(DetailLayoutBuilder->GetDetailsView()->AsShared());
+		const TSharedPtr<SWindow> ParentWindow = FSlateApplication::Get().FindWidgetWindow(DetailLayoutBuilder->GetDetailsViewSharedPtr().ToSharedRef());
 		if (ParentWindow.IsValid() && ParentWindow->GetNativeWindow().IsValid())
 		{
 			ParentWindowWindowHandle = ParentWindow->GetNativeWindow()->GetOSWindowHandle();
@@ -1096,7 +1096,7 @@ void FLocalizationTargetDetailCustomization::ImportDialogueAllCultures()
 		}
 
 		// Execute import dialogue.
-		const TSharedPtr<SWindow> ParentWindow = FSlateApplication::Get().FindWidgetWindow(DetailLayoutBuilder->GetDetailsView()->AsShared());
+		const TSharedPtr<SWindow> ParentWindow = FSlateApplication::Get().FindWidgetWindow(DetailLayoutBuilder->GetDetailsViewSharedPtr().ToSharedRef());
 		LocalizationCommandletTasks::ImportDialogueForTarget(ParentWindow.ToSharedRef(), LocalizationTarget.Get());
 	}
 }
@@ -1110,7 +1110,7 @@ void FLocalizationTargetDetailCustomization::CountWords()
 {
 	if (LocalizationTarget.IsValid())
 	{
-		const TSharedPtr<SWindow> ParentWindow = FSlateApplication::Get().FindWidgetWindow(DetailLayoutBuilder->GetDetailsView()->AsShared());
+		const TSharedPtr<SWindow> ParentWindow = FSlateApplication::Get().FindWidgetWindow(DetailLayoutBuilder->GetDetailsViewSharedPtr().ToSharedRef());
 		LocalizationCommandletTasks::GenerateWordCountReportForTarget(ParentWindow.ToSharedRef(), LocalizationTarget.Get());
 
 		UpdateTargetFromReports();
@@ -1127,7 +1127,7 @@ void FLocalizationTargetDetailCustomization::CompileTextAllCultures()
 	if (LocalizationTarget.IsValid())
 	{
 		// Execute compile.
-		const TSharedPtr<SWindow> ParentWindow = FSlateApplication::Get().FindWidgetWindow(DetailLayoutBuilder->GetDetailsView()->AsShared());
+		const TSharedPtr<SWindow> ParentWindow = FSlateApplication::Get().FindWidgetWindow(DetailLayoutBuilder->GetDetailsViewSharedPtr().ToSharedRef());
 		LocalizationCommandletTasks::CompileTextForTarget(ParentWindow.ToSharedRef(), LocalizationTarget.Get());
 	}
 }
diff --git a/Engine/Source/Editor/LocalizationDashboard/Private/LocalizationTargetSetDetailCustomization.cpp b/Engine/Source/Editor/LocalizationDashboard/Private/LocalizationTargetSetDetailCustomization.cpp
--- a/Engine/Source/Editor/LocalizationDashboard/Private/LocalizationTargetSetDetailCustomization.cpp
+++ b/Engine/Source/Editor/LocalizationDashboard/Private/LocalizationTargetSetDetailCustomization.cpp
@@ -354,7 +354,7 @@ void FLocalizationTargetSetDetailCustomization::GatherTextAllTargets()
 	}
 
 	// Execute gather.
-	const TSharedPtr<SWindow> ParentWindow = FSlateApplication::Get().FindWidgetWindow(DetailLayoutBuilder->GetDetailsView()->AsShared());
+	const TSharedPtr<SWindow> ParentWindow = FSlateApplication::Get().FindWidgetWindow(DetailLayoutBuilder->GetDetailsViewSharedPtr().ToSharedRef());
 	LocalizationCommandletTasks::GatherTextForTargets(ParentWindow.ToSharedRef(), TargetObjectsToProcess);
 
 	for (ULocalizationTarget* const LocalizationTarget : TargetObjectsToProcess)
@@ -383,7 +383,7 @@ void FLocalizationTargetSetDetailCustomization::ImportTextAllTargets()
 	if (DesktopPlatform)
 	{
 		void* ParentWindowWindowHandle = NULL;
-		const TSharedPtr<SWindow> ParentWindow = FSlateApplication::Get().FindWidgetWindow(DetailLayoutBuilder->GetDetailsView()->AsShared());
+		const TSharedPtr<SWindow> ParentWindow = FSlateApplication::Get().FindWidgetWindow(DetailLayoutBuilder->GetDetailsViewSharedPtr().ToSharedRef());
 		if (ParentWindow.IsValid() && ParentWindow->GetNativeWindow().IsValid())
 		{
 			ParentWindowWindowHandle = ParentWindow->GetNativeWindow()->GetOSWindowHandle();
@@ -435,7 +435,7 @@ void FLocalizationTargetSetDetailCustomization::ExportTextAllTargets()
 	if (DesktopPlatform)
 	{
 		void* ParentWindowWindowHandle = NULL;
-		const TSharedPtr<SWindow> ParentWindow = FSlateApplication::Get().FindWidgetWindow(DetailLayoutBuilder->GetDetailsView()->AsShared());
+		const TSharedPtr<SWindow> ParentWindow = FSlateApplication::Get().FindWidgetWindow(DetailLayoutBuilder->GetDetailsViewSharedPtr().ToSharedRef());
 		if (ParentWindow.IsValid() && ParentWindow->GetNativeWindow().IsValid())
 		{
 			ParentWindowWindowHandle = ParentWindow->GetNativeWindow()->GetOSWindowHandle();
@@ -481,7 +481,7 @@ void FLocalizationTargetSetDetailCustomization::ImportDialogueScriptAllTargets()
 	if (DesktopPlatform)
 	{
 		void* ParentWindowWindowHandle = NULL;
-		const TSharedPtr<SWindow> ParentWindow = FSlateApplication::Get().FindWidgetWindow(DetailLayoutBuilder->GetDetailsView()->AsShared());
+		const TSharedPtr<SWindow> ParentWindow = FSlateApplication::Get().FindWidgetWindow(DetailLayoutBuilder->GetDetailsViewSharedPtr().ToSharedRef());
 		if (ParentWindow.IsValid() && ParentWindow->GetNativeWindow().IsValid())
 		{
 			ParentWindowWindowHandle = ParentWindow->GetNativeWindow()->GetOSWindowHandle();
@@ -533,7 +533,7 @@ void FLocalizationTargetSetDetailCustomization::ExportDialogueScriptAllTargets()
 	if (DesktopPlatform)
 	{
 		void* ParentWindowWindowHandle = NULL;
-		const TSharedPtr<SWindow> ParentWindow = FSlateApplication::Get().FindWidgetWindow(DetailLayoutBuilder->GetDetailsView()->AsShared());
+		const TSharedPtr<SWindow> ParentWindow = FSlateApplication::Get().FindWidgetWindow(DetailLayoutBuilder->GetDetailsViewSharedPtr().ToSharedRef());
 		if (ParentWindow.IsValid() && ParentWindow->GetNativeWindow().IsValid())
 		{
 			ParentWindowWindowHandle = ParentWindow->GetNativeWindow()->GetOSWindowHandle();
@@ -591,7 +591,7 @@ void FLocalizationTargetSetDetailCustomization::ImportDialogueAllTargets()
 	}
 
 	// Execute import dialogue.
-	const TSharedPtr<SWindow> ParentWindow = FSlateApplication::Get().FindWidgetWindow(DetailLayoutBuilder->GetDetailsView()->AsShared());
+	const TSharedPtr<SWindow> ParentWindow = FSlateApplication::Get().FindWidgetWindow(DetailLayoutBuilder->GetDetailsViewSharedPtr().ToSharedRef());
 	LocalizationCommandletTasks::ImportDialogueForTargets(ParentWindow.ToSharedRef(), TargetObjectsToProcess);
 }
 
@@ -621,7 +621,7 @@ void FLocalizationTargetSetDetailCustomization::CountWordsForAllTargets()
 	}
 
 	// Execute compile.
-	const TSharedPtr<SWindow> ParentWindow = FSlateApplication::Get().FindWidgetWindow(DetailLayoutBuilder->GetDetailsView()->AsShared());
+	const TSharedPtr<SWindow> ParentWindow = FSlateApplication::Get().FindWidgetWindow(DetailLayoutBuilder->GetDetailsViewSharedPtr().ToSharedRef());
 	LocalizationCommandletTasks::GenerateWordCountReportsForTargets(ParentWindow.ToSharedRef(), TargetObjectsToProcess);
 }
 
@@ -651,7 +651,7 @@ void FLocalizationTargetSetDetailCustomization::CompileTextAllTargets()
 	}
 
 	// Execute compile.
-	const TSharedPtr<SWindow> ParentWindow = FSlateApplication::Get().FindWidgetWindow(DetailLayoutBuilder->GetDetailsView()->AsShared());
+	const TSharedPtr<SWindow> ParentWindow = FSlateApplication::Get().FindWidgetWindow(DetailLayoutBuilder->GetDetailsViewSharedPtr().ToSharedRef());
 	LocalizationCommandletTasks::CompileTextForTargets(ParentWindow.ToSharedRef(), TargetObjectsToProcess);
 }
 
diff --git a/Engine/Source/Editor/MovieSceneTools/Private/Conditions/MovieSceneConditionCustomization.cpp b/Engine/Source/Editor/MovieSceneTools/Private/Conditions/MovieSceneConditionCustomization.cpp
--- a/Engine/Source/Editor/MovieSceneTools/Private/Conditions/MovieSceneConditionCustomization.cpp
+++ b/Engine/Source/Editor/MovieSceneTools/Private/Conditions/MovieSceneConditionCustomization.cpp
@@ -173,7 +173,7 @@ void FMovieSceneConditionCustomization::CustomizeChildren(TSharedRef<IPropertyHa
 	//IDetailCategoryBuilder& NewCategory = LayoutBuilder.EditCategory(TEXT("Condition"), FText::GetEmpty(), ECategoryPriority::TypeSpecific);
 
 	// Hold onto a reference to the details view to prevent it from being destroyed immediately when the menu goes away. 
-	DetailsView = StaticCastSharedPtr<IDetailsView>(LayoutBuilder.GetDetailsView()->AsShared().ToSharedPtr());
+	DetailsView = LayoutBuilder.GetDetailsViewSharedPtr();
 
 	// Customize and display the inner children of the Condition property itself as the children here.
 
diff --git a/Engine/Source/Editor/Persona/Private/AnimGraphNodeDetails.cpp b/Engine/Source/Editor/Persona/Private/AnimGraphNodeDetails.cpp
--- a/Engine/Source/Editor/Persona/Private/AnimGraphNodeDetails.cpp
+++ b/Engine/Source/Editor/Persona/Private/AnimGraphNodeDetails.cpp
@@ -117,7 +117,7 @@ void FAnimGraphNodeDetails::CustomizeDetails(class IDetailLayoutBuilder& DetailB
 	}
 
 	// In the case of property row generators, we dont have a details view, so we can skip this work
-	if(DetailBuilder.GetDetailsView() == nullptr)
+	if(DetailBuilder.GetDetailsViewSharedPtr() == nullptr)
 	{
 		return;
 	}
diff --git a/Engine/Source/Editor/PropertyEditor/Private/AsyncDetailViewDiff.cpp b/Engine/Source/Editor/PropertyEditor/Private/AsyncDetailViewDiff.cpp
--- a/Engine/Source/Editor/PropertyEditor/Private/AsyncDetailViewDiff.cpp
+++ b/Engine/Source/Editor/PropertyEditor/Private/AsyncDetailViewDiff.cpp
@@ -11,7 +11,7 @@ namespace AsyncDetailViewDiffHelpers
 	TArray<TWeakObjectPtr<UObject>> GetObjects(const TSharedPtr<FDetailTreeNode>& TreeNode)
 	{
 		TArray<UObject*> Result;
-		if (const IDetailsViewPrivate* DetailsView = TreeNode->GetDetailsView())
+		if (TSharedPtr<const IDetailsViewPrivate> DetailsView = TreeNode->GetDetailsViewSharedPtr())
 		{
 			return DetailsView->GetSelectedObjects();
 		}
diff --git a/Engine/Source/Editor/PropertyEditor/Private/CustomChildBuilder.cpp b/Engine/Source/Editor/PropertyEditor/Private/CustomChildBuilder.cpp
--- a/Engine/Source/Editor/PropertyEditor/Private/CustomChildBuilder.cpp
+++ b/Engine/Source/Editor/PropertyEditor/Private/CustomChildBuilder.cpp
@@ -191,7 +191,7 @@ TSharedRef<SWidget> FCustomChildrenBuilder::GenerateStructValueWidget( TSharedRe
 
 	FPropertyEditorModule& PropertyEditorModule = FModuleManager::GetModuleChecked<FPropertyEditorModule>("PropertyEditor");
 	
-	IDetailsViewPrivate* DetailsView = ParentCategory.Pin()->GetDetailsView();
+	TSharedPtr<IDetailsViewPrivate> DetailsView = ParentCategory.Pin()->GetDetailsViewSharedPtr();
 
 	FPropertyTypeLayoutCallback LayoutCallback = PropertyEditorModule.GetPropertyTypeCustomization(StructProperty, *StructPropertyHandle, DetailsView ? DetailsView->GetCustomPropertyTypeLayoutMap() : FCustomPropertyTypeLayoutMap() );
 	if (LayoutCallback.IsValid())
diff --git a/Engine/Source/Editor/PropertyEditor/Private/DetailAdvancedDropdownNode.cpp b/Engine/Source/Editor/PropertyEditor/Private/DetailAdvancedDropdownNode.cpp
--- a/Engine/Source/Editor/PropertyEditor/Private/DetailAdvancedDropdownNode.cpp
+++ b/Engine/Source/Editor/PropertyEditor/Private/DetailAdvancedDropdownNode.cpp
@@ -34,7 +34,6 @@ public:
 	void Construct( const FArguments& InArgs, IDetailsViewPrivate* InDetailsView, const TSharedRef<STableViewBase>& InOwnerTableView)
 	{
 		IsExpanded = InArgs._IsExpanded;
-		DetailsView = InDetailsView;
 		OnClicked = InArgs._OnClicked;
 
 		TAttribute<bool> IsVisibleBool = InArgs._IsVisible;
@@ -102,7 +101,7 @@ public:
 				this->GetRowBackgroundColor();
 		};
 
-		FDetailsViewStyle ViewStyle = DetailsView ? DetailsView->GetStyleKey() : SDetailsView::GetPrimaryDetailsViewStyleKey();
+		FDetailsViewStyle ViewStyle = InDetailsView ? InDetailsView->GetStyleKey() : SDetailsView::GetPrimaryDetailsViewStyleKey();
 
 		ChildSlot
 		[
@@ -198,12 +197,11 @@ private:
 	TSharedPtr<SButton> ExpanderButton;
 	FOnClicked OnClicked;
 	bool bDisplayShowAdvancedMessage;
-	IDetailsViewPrivate* DetailsView;
 };
 
 TSharedRef< ITableRow > FAdvancedDropdownNode::GenerateWidgetForTableView( const TSharedRef<STableViewBase>& OwnerTable, bool bAllowFavoriteSystem)
 {
-	return SNew(SAdvancedDropdownRow, ParentCategory.GetDetailsView(), OwnerTable)
+	return SNew(SAdvancedDropdownRow, ParentCategory.GetDetailsViewSharedPtr().Get(), OwnerTable)
 		.OnClicked(this, &FAdvancedDropdownNode::OnAdvancedDropDownClicked)
 		.IsButtonEnabled(IsEnabled)
 		.IsExpanded(IsExpanded)
@@ -225,7 +223,7 @@ FReply FAdvancedDropdownNode::OnAdvancedDropDownClicked()
 
 ENodeVisibility FAdvancedDropdownNode::GetVisibility() const
 {
-	IDetailsViewPrivate* DetailsView = GetDetailsView();
+	TSharedPtr<IDetailsViewPrivate> DetailsView = GetDetailsViewSharedPtr();
 	if (!DetailsView || DetailsView->IsCustomRowVisible(FName(NAME_None), FName(*GetParentCategory()->GetDisplayName().ToString())))
 	{
 		return ENodeVisibility::Visible;
diff --git a/Engine/Source/Editor/PropertyEditor/Private/DetailAdvancedDropdownNode.h b/Engine/Source/Editor/PropertyEditor/Private/DetailAdvancedDropdownNode.h
--- a/Engine/Source/Editor/PropertyEditor/Private/DetailAdvancedDropdownNode.h
+++ b/Engine/Source/Editor/PropertyEditor/Private/DetailAdvancedDropdownNode.h
@@ -26,8 +26,8 @@ public:
 
 private:
 	/** IDetailTreeNode Interface */
-	virtual IDetailsView* GetNodeDetailsView() const override { return ParentCategory.GetNodeDetailsView(); }
-	virtual IDetailsViewPrivate* GetDetailsView() const override { return ParentCategory.GetDetailsView(); }
+	virtual TSharedPtr<IDetailsView> GetNodeDetailsViewSharedPtr() const override { return ParentCategory.GetNodeDetailsViewSharedPtr(); }
+	virtual TSharedPtr<IDetailsViewPrivate> GetDetailsViewSharedPtr() const override { return ParentCategory.GetDetailsViewSharedPtr(); }
 	virtual TSharedRef< ITableRow > GenerateWidgetForTableView( const TSharedRef<STableViewBase>& OwnerTable, bool bAllowFavoriteSystem) override;
 	virtual bool GenerateStandaloneWidget(FDetailWidgetRow& OutRow) const override;
 	virtual void GetChildren(FDetailNodeList& OutChildren, const bool& bInIgnoreVisibility) override {}
diff --git a/Engine/Source/Editor/PropertyEditor/Private/DetailCategoryBuilderImpl.cpp b/Engine/Source/Editor/PropertyEditor/Private/DetailCategoryBuilderImpl.cpp
--- a/Engine/Source/Editor/PropertyEditor/Private/DetailCategoryBuilderImpl.cpp
+++ b/Engine/Source/Editor/PropertyEditor/Private/DetailCategoryBuilderImpl.cpp
@@ -291,22 +291,25 @@ void FDetailCategoryImpl::AddCustomBuilder(TSharedRef<IDetailCustomNodeBuilder>
 		NewCustomization.CustomBuilderRow->SetOriginalPath(PathToNode.ToString());
 	}
 
-	if (GetDetailsView() != nullptr && !bFavoriteCategory)
+	if (!bFavoriteCategory)
 	{
-		if (GetDetailsView()->IsCustomBuilderFavorite(NewCustomization.CustomBuilderRow->GetOriginalPath()))
+		if (TSharedPtr<IDetailsViewPrivate> DetailsView = GetDetailsViewSharedPtr())
 		{
-			TSharedPtr<FDetailLayoutBuilderImpl> ParentLayout = GetParentLayoutImpl();
-			if (ParentLayout.IsValid())
+			if (DetailsView->IsCustomBuilderFavorite(NewCustomization.CustomBuilderRow->GetOriginalPath()))
 			{
-				static const FName FavoritesCategoryName(TEXT("Favorites"));
-				FDetailCategoryImpl& FavoritesCategory = ParentLayout->DefaultCategory(FavoritesCategoryName);
-
-				FDetailLayoutCustomization FavoritesCustomization;
-				FavoritesCustomization.bCustom = true;
-				FavoritesCustomization.bAdvanced = false;
-				FavoritesCustomization.CustomBuilderRow = MakeShareable(new FDetailCustomBuilderRow(InCustomBuilder));
-				FavoritesCustomization.CustomBuilderRow->SetOriginalPath(NewCustomization.CustomBuilderRow->GetOriginalPath());
-				FavoritesCategory.AddCustomLayout(FavoritesCustomization);
+				TSharedPtr<FDetailLayoutBuilderImpl> ParentLayout = GetParentLayoutImpl();
+				if (ParentLayout.IsValid())
+				{
+					static const FName FavoritesCategoryName(TEXT("Favorites"));
+					FDetailCategoryImpl& FavoritesCategory = ParentLayout->DefaultCategory(FavoritesCategoryName);
+
+					FDetailLayoutCustomization FavoritesCustomization;
+					FavoritesCustomization.bCustom = true;
+					FavoritesCustomization.bAdvanced = false;
+					FavoritesCustomization.CustomBuilderRow = MakeShareable(new FDetailCustomBuilderRow(InCustomBuilder));
+					FavoritesCustomization.CustomBuilderRow->SetOriginalPath(NewCustomization.CustomBuilderRow->GetOriginalPath());
+					FavoritesCategory.AddCustomLayout(FavoritesCustomization);
+				}
 			}
 		}
 	}
@@ -385,9 +388,9 @@ void FDetailCategoryImpl::SetCategoryVisibility(bool bIsVisible)
 	{
 		bIsCategoryVisible = bIsVisible;
 
-		if (GetDetailsView())
+		if (TSharedPtr<IDetailsViewPrivate> DetailsView = GetDetailsViewSharedPtr())
 		{
-			GetDetailsView()->RerunCurrentFilter();
+			DetailsView->RerunCurrentFilter();
 		}
 
 		TSharedPtr<FDetailLayoutBuilderImpl> ParentLayout = GetParentLayoutImpl();
@@ -736,9 +739,9 @@ bool FDetailCategoryImpl::ShouldAdvancedBeVisible() const
 
 void FDetailCategoryImpl::RequestItemExpanded(TSharedRef<FDetailTreeNode> TreeNode, bool bShouldBeExpanded)
 {
-	if (GetDetailsView())
+	if (TSharedPtr<IDetailsViewPrivate> DetailsView = GetDetailsViewSharedPtr())
 	{		
-		GetDetailsView()->RequestItemExpanded(TreeNode, bShouldBeExpanded);
+		DetailsView->RequestItemExpanded(TreeNode, bShouldBeExpanded);
 	}
 }
 
@@ -781,9 +784,9 @@ void FDetailCategoryImpl::RefreshTreeInternal(bool bRefilterCategory)
 	}
 	else
 	{
-		if (GetDetailsView())
+		if (TSharedPtr<IDetailsViewPrivate> DetailsView = GetDetailsViewSharedPtr())
 		{
-			GetDetailsView()->RefreshTree();
+			DetailsView->RefreshTree();
 		}
 	}
 }
@@ -881,12 +884,12 @@ void FDetailCategoryImpl::SetDisplayName(FName InCategoryName, const FText& Loca
 	}
 }
 
-IDetailsViewPrivate* FDetailCategoryImpl::GetDetailsView() const
+TSharedPtr<IDetailsViewPrivate> FDetailCategoryImpl::GetDetailsViewSharedPtr() const
 {
 	TSharedPtr<FDetailLayoutBuilderImpl> ParentLayout = GetParentLayoutImpl();
 	if (ParentLayout.IsValid())
 	{
-		return ParentLayout->GetDetailsView();
+		return StaticCastSharedPtr<IDetailsViewPrivate>(ParentLayout->GetDetailsViewSharedPtr());
 	}
 
 	return nullptr;
@@ -1352,7 +1355,7 @@ void FDetailCategoryImpl::GenerateLayout()
 
 bool FDetailCategoryImpl::IsParentEnabled() const
 {
-	IDetailsViewPrivate* DetailsView = GetDetailsView();
+	TSharedPtr<IDetailsViewPrivate> DetailsView = GetDetailsViewSharedPtr();
 	return !DetailsView || DetailsView->IsPropertyEditingEnabled();
 }
 
diff --git a/Engine/Source/Editor/PropertyEditor/Private/DetailCategoryBuilderImpl.h b/Engine/Source/Editor/PropertyEditor/Private/DetailCategoryBuilderImpl.h
--- a/Engine/Source/Editor/PropertyEditor/Private/DetailCategoryBuilderImpl.h
+++ b/Engine/Source/Editor/PropertyEditor/Private/DetailCategoryBuilderImpl.h
@@ -196,8 +196,8 @@ public:
 	virtual TSharedPtr<FOnPasteFromText> OnPasteFromText() const override { return PasteFromTextDelegate; }
 
 	/** FDetailTreeNode interface */
-	virtual IDetailsView* GetNodeDetailsView() const override { return GetDetailsView(); }
-	virtual IDetailsViewPrivate* GetDetailsView() const override;
+	virtual TSharedPtr<IDetailsView> GetNodeDetailsViewSharedPtr() const override { return GetDetailsViewSharedPtr(); }
+	virtual TSharedPtr<IDetailsViewPrivate> GetDetailsViewSharedPtr() const override;
 	virtual TSharedRef< ITableRow > GenerateWidgetForTableView(const TSharedRef<STableViewBase>& OwnerTable, bool bAllowFavoriteSystem) override;
 	virtual bool GenerateStandaloneWidget(FDetailWidgetRow& OutRow) const override;
 
diff --git a/Engine/Source/Editor/PropertyEditor/Private/DetailCategoryGroupNode.h b/Engine/Source/Editor/PropertyEditor/Private/DetailCategoryGroupNode.h
--- a/Engine/Source/Editor/PropertyEditor/Private/DetailCategoryGroupNode.h
+++ b/Engine/Source/Editor/PropertyEditor/Private/DetailCategoryGroupNode.h
@@ -25,8 +25,8 @@ public:
 	bool GetHasSplitter() const { return bHasSplitter; }
 
 private:
-	virtual IDetailsView* GetNodeDetailsView() const override { return ParentCategory.GetNodeDetailsView(); }
-	virtual IDetailsViewPrivate* GetDetailsView() const override { return ParentCategory.GetDetailsView(); }
+	virtual TSharedPtr<IDetailsView> GetNodeDetailsViewSharedPtr() const override { return ParentCategory.GetNodeDetailsViewSharedPtr(); }
+	virtual TSharedPtr<IDetailsViewPrivate> GetDetailsViewSharedPtr() const override { return ParentCategory.GetDetailsViewSharedPtr(); }
 	virtual void OnItemExpansionChanged( bool bIsExpanded, bool bShouldSaveState) override {}
 	virtual bool ShouldBeExpanded() const override { return true; }
 	virtual ENodeVisibility GetVisibility() const override { return bShouldBeVisible ? ENodeVisibility::Visible : ENodeVisibility::HiddenDueToFiltering; }
diff --git a/Engine/Source/Editor/PropertyEditor/Private/DetailItemNode.cpp b/Engine/Source/Editor/PropertyEditor/Private/DetailItemNode.cpp
--- a/Engine/Source/Editor/PropertyEditor/Private/DetailItemNode.cpp
+++ b/Engine/Source/Editor/PropertyEditor/Private/DetailItemNode.cpp
@@ -577,7 +577,7 @@ static bool PassesAllFilters( FDetailItemNode* ItemNode, const FDetailLayoutCust
 		}
 		static bool ItemIsKeyable( FDetailItemNode *InItemNode, UClass *ObjectClass, TSharedPtr<FPropertyNode> PropertyNode)
 		{
-			IDetailsViewPrivate* DetailsView = InItemNode->GetDetailsView();
+			TSharedPtr<IDetailsViewPrivate> DetailsView = InItemNode->GetDetailsViewSharedPtr();
 			TSharedPtr<IDetailKeyframeHandler> KeyframeHandler = DetailsView->GetKeyframeHandler();
 			if (KeyframeHandler.IsValid())
 			{
@@ -592,7 +592,7 @@ static bool PassesAllFilters( FDetailItemNode* ItemNode, const FDetailLayoutCust
 
 		static bool ItemIsAnimated(FDetailItemNode *InItemNode, TSharedPtr<FPropertyNode> PropertyNode)
 		{
-			IDetailsViewPrivate* DetailsView = InItemNode->GetDetailsView();
+			TSharedPtr<IDetailsViewPrivate> DetailsView = InItemNode->GetDetailsViewSharedPtr();
 			TSharedPtr<IDetailKeyframeHandler> KeyframeHandler = DetailsView->GetKeyframeHandler();
 			if (KeyframeHandler.IsValid())
 			{
@@ -844,7 +844,7 @@ EVisibility FDetailItemNode::ComputeItemVisibility() const
 		}
 	}
 
-	const IDetailsViewPrivate* DetailsView = GetDetailsView();
+	TSharedPtr<const IDetailsViewPrivate> DetailsView = GetDetailsViewSharedPtr();
 	if (DetailsView != nullptr)
 	{
 		// check the details view's IsCustomRowVisible delegate if this isn't a property row
@@ -935,7 +935,7 @@ bool FDetailItemNode::IsPropertyEditingEnabledImpl() const
 {
 	bool bIsEnabled = IsParentEnabled.Get(true);
 
-	IDetailsViewPrivate* DetailsView = GetDetailsView();
+	TSharedPtr<IDetailsViewPrivate> DetailsView = GetDetailsViewSharedPtr();
 	if (DetailsView)
 	{
 		if (Customization.HasPropertyNode())
diff --git a/Engine/Source/Editor/PropertyEditor/Private/DetailItemNode.h b/Engine/Source/Editor/PropertyEditor/Private/DetailItemNode.h
--- a/Engine/Source/Editor/PropertyEditor/Private/DetailItemNode.h
+++ b/Engine/Source/Editor/PropertyEditor/Private/DetailItemNode.h
@@ -65,8 +65,8 @@ public:
 	void RefreshCachedVisibility(bool bCallChangeDelegate = false);
 
 	/** FDetailTreeNode interface */
-	virtual IDetailsView* GetNodeDetailsView() const override { TSharedPtr<FDetailCategoryImpl> PC = GetParentCategory(); return PC.IsValid() ? PC->GetNodeDetailsView() : nullptr; }
-	virtual IDetailsViewPrivate* GetDetailsView() const override{ TSharedPtr<FDetailCategoryImpl> PC = GetParentCategory(); return PC.IsValid() ? PC->GetDetailsView() : nullptr; }
+	virtual TSharedPtr<IDetailsView> GetNodeDetailsViewSharedPtr() const override { TSharedPtr<FDetailCategoryImpl> PC = GetParentCategory(); return PC.IsValid() ? PC->GetNodeDetailsViewSharedPtr() : nullptr; }
+	virtual TSharedPtr<IDetailsViewPrivate> GetDetailsViewSharedPtr() const override{ TSharedPtr<FDetailCategoryImpl> PC = GetParentCategory(); return PC.IsValid() ? PC->GetDetailsViewSharedPtr() : nullptr; }
 	virtual TSharedRef< ITableRow > GenerateWidgetForTableView( const TSharedRef<STableViewBase>& OwnerTable, bool bAllowFavoriteSystem) override;
 	virtual bool GenerateStandaloneWidget(FDetailWidgetRow& OutRow) const override;
 	virtual void GetChildren(FDetailNodeList& OutChildren, const bool& bInIgnoreVisibility = false) override;
diff --git a/Engine/Source/Editor/PropertyEditor/Private/DetailLayoutBuilderImpl.cpp b/Engine/Source/Editor/PropertyEditor/Private/DetailLayoutBuilderImpl.cpp
--- a/Engine/Source/Editor/PropertyEditor/Private/DetailLayoutBuilderImpl.cpp
+++ b/Engine/Source/Editor/PropertyEditor/Private/DetailLayoutBuilderImpl.cpp
@@ -17,7 +17,7 @@ FDetailLayoutBuilderImpl::FDetailLayoutBuilderImpl(TSharedPtr<FComplexPropertyNo
 	, PropertyMap( InPropertyMap )
 	, PropertyDetailsUtilities( InPropertyUtilities )
 	, PropertyGenerationUtilities( InPropertyGenerationUtilities )
-	, DetailsView( InDetailsView.Get() )
+	, DetailsView( InDetailsView )
 	, CurrentCustomizationClass( nullptr )
 	, bLayoutForExternalRoot(bIsExternal)
 {
@@ -582,14 +582,15 @@ void FDetailLayoutBuilderImpl::GenerateDetailLayout()
 	}
 
 	TSharedPtr<FComplexPropertyNode> RootNodePinned = RootNode.Pin();
-	if (DetailsView && DetailsView->GetRootObjectCustomization() && RootNodePinned->GetInstancesNum() && !bLayoutForExternalRoot)
+	TSharedPtr<IDetailsViewPrivate> DetailsViewPinned = DetailsView.Pin();
+	if (DetailsViewPinned && DetailsViewPinned->GetRootObjectCustomization() && RootNodePinned->GetInstancesNum() && !bLayoutForExternalRoot)
 	{
 		FObjectPropertyNode* ObjectNode = RootNodePinned->AsObjectNode();
 
-		TSharedPtr<IDetailRootObjectCustomization> RootObjectCustomization = DetailsView->GetRootObjectCustomization();
+		TSharedPtr<IDetailRootObjectCustomization> RootObjectCustomization = DetailsViewPinned->GetRootObjectCustomization();
 
 		// there are multiple objects in the details panel.  Separate each one with a unique object name node to differentiate them
-		TSharedRef<FDetailMultiTopLevelObjectRootNode> NewRootNode = MakeShared<FDetailMultiTopLevelObjectRootNode>(RootObjectCustomization, DetailsView, ObjectNode);
+		TSharedRef<FDetailMultiTopLevelObjectRootNode> NewRootNode = MakeShared<FDetailMultiTopLevelObjectRootNode>(RootObjectCustomization, DetailsViewPinned, ObjectNode);
 		NewRootNode->SetChildren(CategoryNodes);
 
 		AllRootTreeNodes.Add(NewRootNode);
@@ -615,9 +616,9 @@ void FDetailLayoutBuilderImpl::FilterDetailLayout( const FDetailFilter& InFilter
 		if(RootTreeNode->GetVisibility() == ENodeVisibility::Visible)
 		{
 			FilteredRootTreeNodes.Add(RootTreeNode);
-			if (DetailsView)
+			if (TSharedPtr<IDetailsViewPrivate> DetailsViewPinned = DetailsView.Pin())
 			{
-				DetailsView->RequestItemExpanded(RootTreeNode, RootTreeNode->ShouldBeExpanded());
+				DetailsViewPinned->RequestItemExpanded(RootTreeNode, RootTreeNode->ShouldBeExpanded());
 			}
 		}
 	}
@@ -830,9 +831,9 @@ TSharedRef<IPropertyHandle> FDetailLayoutBuilderImpl::GetPropertyHandle( TShared
 void FDetailLayoutBuilderImpl::AddExternalRootPropertyNode(TSharedRef<FComplexPropertyNode> InExternalRootNode)
 {
 	ExternalRootPropertyNodes.Add(InExternalRootNode);
-	if (GetDetailsView())
+	if (TSharedPtr<IDetailsViewPrivate> DetailsViewPinned = DetailsView.Pin())
 	{
-		GetDetailsView()->RestoreExpandedItems(InExternalRootNode);
+		DetailsViewPinned->RestoreExpandedItems(InExternalRootNode);
 	}
 }
 
@@ -843,19 +844,22 @@ void FDetailLayoutBuilderImpl::RemoveExternalRootPropertyNode(TSharedRef<FComple
 		return ExternalRootPropertyNode == InExternalRootNode;
 	});
 
-	if (NumRemoved > 0 && GetDetailsView())
+	if (NumRemoved > 0)
 	{
-		GetDetailsView()->SaveExpandedItems(InExternalRootNode);
+		if (TSharedPtr<IDetailsViewPrivate> DetailsViewPinned = DetailsView.Pin())
+		{
+			DetailsViewPinned->SaveExpandedItems(InExternalRootNode);
+		}
 	}
 }
 
 void FDetailLayoutBuilderImpl::ClearExternalRootPropertyNodes()
 {
-	if (GetDetailsView())
+	if (TSharedPtr<IDetailsViewPrivate> DetailsViewPinned = DetailsView.Pin())
 	{
 		for (const TSharedPtr<FComplexPropertyNode>& ExternalRootPropertyNode : ExternalRootPropertyNodes)
 		{
-			GetDetailsView()->SaveExpandedItems(ExternalRootPropertyNode.ToSharedRef());
+			DetailsViewPinned->SaveExpandedItems(ExternalRootPropertyNode.ToSharedRef());
 		}
 	}
 
@@ -906,20 +910,23 @@ TSharedPtr<FAssetThumbnailPool> FDetailLayoutBuilderImpl::GetThumbnailPool() con
 
 bool FDetailLayoutBuilderImpl::IsPropertyVisible( TSharedRef<IPropertyHandle> PropertyHandle ) const
 {
-	if (PropertyHandle->IsValidHandle() && DetailsView != nullptr)
+	if (PropertyHandle->IsValidHandle())
 	{
-		TSharedPtr<FPropertyNode> PropertyNode = StaticCastSharedRef<FPropertyHandleBase>(PropertyHandle)->GetPropertyNode();
-		const FCategoryPropertyNode* CategoryNode = PropertyNode.IsValid() ? PropertyNode->AsCategoryNode() : nullptr;
-		if (CategoryNode != nullptr)
+		if (TSharedPtr<IDetailsViewPrivate> DetailsViewPinned = DetailsView.Pin())
 		{
-			// this is a subcategory
-			FName CategoryName = CategoryNode->GetCategoryName();
-			return DetailsView->IsCustomRowVisible(FName(), CategoryName);
-		}
-		else if (PropertyHandle->GetProperty() != nullptr)
-		{
-			FPropertyAndParent PropertyAndParent(PropertyHandle);
-			return DetailsView->IsPropertyVisible(PropertyAndParent);
+			TSharedPtr<FPropertyNode> PropertyNode = StaticCastSharedRef<FPropertyHandleBase>(PropertyHandle)->GetPropertyNode();
+			const FCategoryPropertyNode* CategoryNode = PropertyNode.IsValid() ? PropertyNode->AsCategoryNode() : nullptr;
+			if (CategoryNode != nullptr)
+			{
+				// this is a subcategory
+				FName CategoryName = CategoryNode->GetCategoryName();
+				return DetailsViewPinned->IsCustomRowVisible(FName(), CategoryName);
+			}
+			else if (PropertyHandle->GetProperty() != nullptr)
+			{
+				FPropertyAndParent PropertyAndParent(PropertyHandle);
+				return DetailsViewPinned->IsPropertyVisible(PropertyAndParent);
+			}
 		}
 	}
 
@@ -928,7 +935,8 @@ bool FDetailLayoutBuilderImpl::IsPropertyVisible( TSharedRef<IPropertyHandle> Pr
 
 bool FDetailLayoutBuilderImpl::IsPropertyVisible( const struct FPropertyAndParent& PropertyAndParent ) const
 {
-	return DetailsView ? DetailsView->IsPropertyVisible( PropertyAndParent ) : true;
+	TSharedPtr<IDetailsViewPrivate> DetailsViewPinned = DetailsView.Pin();
+	return DetailsViewPinned ? DetailsViewPinned->IsPropertyVisible( PropertyAndParent ) : true;
 }
 
 void FDetailLayoutBuilderImpl::HideCategory( FName CategoryName )
@@ -936,9 +944,9 @@ void FDetailLayoutBuilderImpl::HideCategory( FName CategoryName )
 	ForceHiddenCategories.Add(CategoryName);
 }
 
-const IDetailsView* FDetailLayoutBuilderImpl::GetDetailsView() const
+TSharedPtr<const IDetailsView> FDetailLayoutBuilderImpl::GetDetailsViewSharedPtr() const
 { 
-	return DetailsView; 
+	return DetailsView.Pin();
 }
 
 void FDetailLayoutBuilderImpl::GetObjectsBeingCustomized( TArray< TWeakObjectPtr<UObject> >& OutObjects ) const
@@ -1083,15 +1091,16 @@ void FDetailLayoutBuilderImpl::RemoveTickableNode( FDetailTreeNode& TickableNode
 
 void FDetailLayoutBuilderImpl::SaveExpansionState( const FString& NodePath, bool bIsExpanded )
 {
-	if (DetailsView)
+	if (TSharedPtr<IDetailsViewPrivate> DetailsViewPinned = DetailsView.Pin())
 	{
-		DetailsView->SaveCustomExpansionState(NodePath, bIsExpanded);
+		DetailsViewPinned->SaveCustomExpansionState(NodePath, bIsExpanded);
 	}
 }
 
 bool FDetailLayoutBuilderImpl::GetSavedExpansionState( const FString& NodePath ) const
 {
-	return DetailsView ? DetailsView->GetCustomSavedExpansionState(NodePath) : false;
+	TSharedPtr<IDetailsViewPrivate> DetailsViewPinned = DetailsView.Pin();
+	return DetailsViewPinned ? DetailsViewPinned->GetCustomSavedExpansionState(NodePath) : false;
 }
 
 void FDetailLayoutBuilderImpl::RegisterInstancedCustomPropertyTypeLayout(FName PropertyTypeName, FOnGetPropertyTypeCustomizationInstance PropertyTypeLayoutDelegate, TSharedPtr<IPropertyTypeIdentifier> Identifier)
@@ -1152,15 +1161,17 @@ void FDetailLayoutBuilderImpl::DisableInstancedReference(TSharedRef<IPropertyHan
 
 bool FDetailLayoutBuilderImpl::AddEmptyCategoryIfNeeded(TSharedPtr<FComplexPropertyNode> Node)
 {
+	TSharedPtr<IDetailsViewPrivate> DetailsViewPinned = DetailsView.Pin();
+
 	const bool bHasNoValidCategories = DefaultCategoryMap.IsEmpty();
-	const bool bHasValidDisplayManager = DetailsView && DetailsView->GetDisplayManager().IsValid();
+	const bool bHasValidDisplayManager = DetailsViewPinned && DetailsViewPinned->GetDisplayManager().IsValid();
 	const bool bHasValidPropertyNode = Node.IsValid();
 	
 	if ( bHasNoValidCategories &&
 		 bHasValidDisplayManager &&
 		 bHasValidPropertyNode )
 	{
- 		return DetailsView->GetDisplayManager()->AddEmptyCategoryToDetailLayoutIfNeeded(Node.ToSharedRef(), SharedThis(this));
+ 		return DetailsViewPinned->GetDisplayManager()->AddEmptyCategoryToDetailLayoutIfNeeded(Node.ToSharedRef(), SharedThis(this));
 	}
 	return false;
 }
diff --git a/Engine/Source/Editor/PropertyEditor/Private/DetailLayoutBuilderImpl.h b/Engine/Source/Editor/PropertyEditor/Private/DetailLayoutBuilderImpl.h
--- a/Engine/Source/Editor/PropertyEditor/Private/DetailLayoutBuilderImpl.h
+++ b/Engine/Source/Editor/PropertyEditor/Private/DetailLayoutBuilderImpl.h
@@ -28,7 +28,7 @@ public:
 	virtual ~FDetailLayoutBuilderImpl() override;
 
 	/** IDetailLayoutBuilder Interface */
-	virtual const IDetailsView* GetDetailsView() const override;
+	virtual TSharedPtr<const IDetailsView> GetDetailsViewSharedPtr() const override;
 	virtual void GetObjectsBeingCustomized(TArray< TWeakObjectPtr<UObject> >& OutObjects) const override;
 	virtual void GetStructsBeingCustomized(TArray< TSharedPtr<FStructOnScope> >& OutStructs) const override;
 	virtual IDetailCategoryBuilder& EditCategory(FName CategoryName, const FText& NewLocalizedDisplayName = FText::GetEmpty(), ECategoryPriority::Type CategoryType = ECategoryPriority::Default) override;
@@ -211,7 +211,7 @@ public:
 	void ClearExternalRootPropertyNodes();
 
 	/** @return The details view that owns this layout */
-	IDetailsViewPrivate* GetDetailsView() { return DetailsView; }
+	TSharedPtr<IDetailsView> GetDetailsViewSharedPtr() override { return DetailsView.Pin(); }
 	/** @return The root node for this customization */
 	TSharedPtr<FComplexPropertyNode> GetRootNode() const { return RootNode.Pin(); }
 
@@ -285,7 +285,7 @@ private:
 	/** Internal utilities for generating property layouts. */
 	const TWeakPtr<IPropertyGenerationUtilities> PropertyGenerationUtilities;
 	/** The view where this detail customizer resides */
-	class IDetailsViewPrivate* DetailsView;
+	TWeakPtr<IDetailsViewPrivate> DetailsView;
 	/** The current class being customized */
 	UStruct* CurrentCustomizationClass;
 	/** True if the layout is for an external root property node and not in the main set of objects the details panel is observing */
diff --git a/Engine/Source/Editor/PropertyEditor/Private/DetailMultiTopLevelObjectRootNode.cpp b/Engine/Source/Editor/PropertyEditor/Private/DetailMultiTopLevelObjectRootNode.cpp
--- a/Engine/Source/Editor/PropertyEditor/Private/DetailMultiTopLevelObjectRootNode.cpp
+++ b/Engine/Source/Editor/PropertyEditor/Private/DetailMultiTopLevelObjectRootNode.cpp
@@ -72,7 +72,7 @@ FReply SDetailMultiTopLevelObjectTableRow::OnMouseButtonDoubleClick( const FGeom
 }
 
 
-FDetailMultiTopLevelObjectRootNode::FDetailMultiTopLevelObjectRootNode(const TSharedPtr<IDetailRootObjectCustomization>& InRootObjectCustomization, IDetailsViewPrivate* InDetailsView, const FObjectPropertyNode* RootNode)
+FDetailMultiTopLevelObjectRootNode::FDetailMultiTopLevelObjectRootNode(const TSharedPtr<IDetailRootObjectCustomization>& InRootObjectCustomization, const TSharedPtr<IDetailsViewPrivate>& InDetailsView, const FObjectPropertyNode* RootNode)
 	: DetailsView(InDetailsView)
 	, RootObjectCustomization(InRootObjectCustomization)
 	, bShouldBeVisible(false)
@@ -199,9 +199,9 @@ void FDetailMultiTopLevelObjectRootNode::FilterNode( const FDetailFilter& InFilt
 		{
 			bShouldBeVisible = true;
 
-			if (DetailsView)
+			if (TSharedPtr<IDetailsViewPrivate> DetailsViewPinned = DetailsView.Pin())
 			{
-				DetailsView->RequestItemExpanded(Child, Child->ShouldBeExpanded());
+				DetailsViewPinned->RequestItemExpanded(Child, Child->ShouldBeExpanded());
 			}
 		}
 	}
diff --git a/Engine/Source/Editor/PropertyEditor/Private/DetailMultiTopLevelObjectRootNode.h b/Engine/Source/Editor/PropertyEditor/Private/DetailMultiTopLevelObjectRootNode.h
--- a/Engine/Source/Editor/PropertyEditor/Private/DetailMultiTopLevelObjectRootNode.h
+++ b/Engine/Source/Editor/PropertyEditor/Private/DetailMultiTopLevelObjectRootNode.h
@@ -45,11 +45,11 @@ private:
 class FDetailMultiTopLevelObjectRootNode : public FDetailTreeNode, public TSharedFromThis<FDetailMultiTopLevelObjectRootNode>
 {
 public:
-	FDetailMultiTopLevelObjectRootNode(const TSharedPtr<IDetailRootObjectCustomization>& RootObjectCustomization, IDetailsViewPrivate* InDetailsView, const FObjectPropertyNode* RootNode);
+	FDetailMultiTopLevelObjectRootNode(const TSharedPtr<IDetailRootObjectCustomization>& RootObjectCustomization, const TSharedPtr<IDetailsViewPrivate>& InDetailsView, const FObjectPropertyNode* RootNode);
 	void SetChildren(const FDetailNodeList& InChildNodes);
 private:
-	virtual IDetailsView* GetNodeDetailsView() const override { return DetailsView; }
-	virtual IDetailsViewPrivate* GetDetailsView() const override { return DetailsView; }
+	virtual TSharedPtr<IDetailsView> GetNodeDetailsViewSharedPtr() const override { return DetailsView.Pin(); }
+	virtual TSharedPtr<IDetailsViewPrivate> GetDetailsViewSharedPtr() const override { return DetailsView.Pin(); }
 	virtual void OnItemExpansionChanged(bool bIsExpanded, bool bShouldSaveState) override;
 	virtual bool ShouldBeExpanded() const override;
 	virtual ENodeVisibility GetVisibility() const override;
@@ -66,7 +66,7 @@ private:
 
 private:
 	FDetailNodeList ChildNodes;
-	IDetailsViewPrivate* DetailsView;
+	TWeakPtr<IDetailsViewPrivate> DetailsView;
 	TWeakPtr<IDetailRootObjectCustomization> RootObjectCustomization;
 	FDetailsObjectSet RootObjectSet;
 	const UClass* CommonBaseClass;
diff --git a/Engine/Source/Editor/PropertyEditor/Private/DetailPropertyRow.cpp b/Engine/Source/Editor/PropertyEditor/Private/DetailPropertyRow.cpp
--- a/Engine/Source/Editor/PropertyEditor/Private/DetailPropertyRow.cpp
+++ b/Engine/Source/Editor/PropertyEditor/Private/DetailPropertyRow.cpp
@@ -81,7 +81,7 @@ FDetailPropertyRow::FDetailPropertyRow(TSharedPtr<FPropertyNode> InPropertyNode,
 		if (PropertyNode->AsComplexNode() && ExternalRootNode.IsValid()) // AsComplexNode works both for objects and structs
 		{
 			// We are showing an entirely different object inline.  Generate a layout for it now.
-			if (IDetailsViewPrivate* DetailsView = InParentCategory->GetDetailsView())
+			if (TSharedPtr<IDetailsViewPrivate> DetailsView = InParentCategory->GetDetailsViewSharedPtr())
 			{
 				ExternalObjectLayout = MakeShared<FDetailLayoutData>();
 				DetailsView->UpdateSinglePropertyMap(InExternalRootNode, *ExternalObjectLayout, true);
@@ -96,7 +96,7 @@ FDetailPropertyRow::FDetailPropertyRow(TSharedPtr<FPropertyNode> InPropertyNode,
 			if (ComplexChildNode.IsValid())
 			{
 				// We are showing an entirely different object inline.  Generate a layout for it now.
-				if (IDetailsViewPrivate* DetailsView = InParentCategory->GetDetailsView())
+				if (TSharedPtr<IDetailsViewPrivate> DetailsView = InParentCategory->GetDetailsViewSharedPtr())
 				{
 					ExternalObjectLayout = MakeShared<FDetailLayoutData>();
 					DetailsView->UpdateSinglePropertyMap(ComplexChildNode, *ExternalObjectLayout, true);
@@ -396,7 +396,8 @@ void FDetailPropertyRow::OnItemNodeInitialized( TSharedRef<FDetailCategoryImpl>
 		}
 	}
 
-	IDetailPropertyChildrenCustomizationHandler* CustomizationHandler = InParentCategory->GetDetailsView() ? InParentCategory->GetDetailsView()->GetChildrenCustomizationHandler().Get() : nullptr;
+	TSharedPtr<IDetailsView> DetailsView = InParentCategory->GetDetailsViewSharedPtr();
+	IDetailPropertyChildrenCustomizationHandler* CustomizationHandler = DetailsView ? DetailsView->GetChildrenCustomizationHandler().Get() : nullptr;
 	if (CustomizationHandler && CustomizationHandler->ShouldCustomizeChildren(PropertyHandle.ToSharedRef()))
 	{
 		PropertyTypeLayoutBuilder = MakeShared<FCustomChildrenBuilder>(InParentCategory, InParentGroup);
diff --git a/Engine/Source/Editor/PropertyEditor/Private/DetailTreeNode.cpp b/Engine/Source/Editor/PropertyEditor/Private/DetailTreeNode.cpp
--- a/Engine/Source/Editor/PropertyEditor/Private/DetailTreeNode.cpp
+++ b/Engine/Source/Editor/PropertyEditor/Private/DetailTreeNode.cpp
@@ -114,7 +114,7 @@ const UStruct* FDetailTreeNode::GetParentBaseStructure() const
 	}
 	// If category fails, use the first selected object. This should be rare, if it's even possible
 	// (maybe a custom builder appearing in a detail property row or something?)
-	for (TWeakObjectPtr<UObject> Selected : GetDetailsView()->GetSelectedObjects())
+	for (TWeakObjectPtr<UObject> Selected : GetDetailsViewSharedPtr()->GetSelectedObjects())
 	{
 		if (UObject* Object = Selected.Get())
 		{
diff --git a/Engine/Source/Editor/PropertyEditor/Private/SDetailCategoryTableRow.cpp b/Engine/Source/Editor/PropertyEditor/Private/SDetailCategoryTableRow.cpp
--- a/Engine/Source/Editor/PropertyEditor/Private/SDetailCategoryTableRow.cpp
+++ b/Engine/Source/Editor/PropertyEditor/Private/SDetailCategoryTableRow.cpp
@@ -30,7 +30,7 @@ void SDetailCategoryTableRow::Construct(const FArguments& InArgs, TSharedRef<FDe
 	bShowBorder = InArgs._ShowBorder;
 	bIsEmpty = InArgs._IsEmpty;
 
-	IDetailsViewPrivate* DetailsView = InOwnerTreeNode->GetDetailsView();
+	TSharedPtr<IDetailsViewPrivate> DetailsView = InOwnerTreeNode->GetDetailsViewSharedPtr();
 	FDetailColumnSizeData& ColumnSizeData = DetailsView->GetColumnSizeData();
 	DisplayManager = DetailsView->GetDisplayManager();
 
@@ -453,7 +453,7 @@ void SDetailCategoryTableRow::OnPasteCategory()
 				*FString::Join(PropertiesNotPasted, TEXT("\n")));
 		}
 
-		IDetailsViewPrivate* DetailsView = OwnerTreeNode.Pin()->GetDetailsView();
+		TSharedPtr<IDetailsViewPrivate> DetailsView = OwnerTreeNode.Pin()->GetDetailsViewSharedPtr();
 		DetailsView->ForceRefresh();
 	}
 }
diff --git a/Engine/Source/Editor/PropertyEditor/Private/SDetailSingleItemRow.cpp b/Engine/Source/Editor/PropertyEditor/Private/SDetailSingleItemRow.cpp
--- a/Engine/Source/Editor/PropertyEditor/Private/SDetailSingleItemRow.cpp
+++ b/Engine/Source/Editor/PropertyEditor/Private/SDetailSingleItemRow.cpp
@@ -106,7 +106,7 @@ bool SDetailSingleItemRow::CheckValidDrop(const TSharedPtr<SDetailSingleItemRow>
 
 		if (OriginalIndex != NewIndex)
 		{
-			IDetailsViewPrivate* DetailsView = OwnerTreeNode.Pin()->GetDetailsView();
+			TSharedPtr<IDetailsViewPrivate> DetailsView = OwnerTreeNode.Pin()->GetDetailsViewSharedPtr();
 			TSharedPtr<IPropertyHandle> SwappingHandle = PropertyEditorHelpers::GetPropertyHandle(SwappingPropertyNode.ToSharedRef(), DetailsView->GetNotifyHook(), DetailsView->GetPropertyUtilities());
 			TSharedPtr<IPropertyHandleArray> ParentHandle = SwappingHandle->GetParentHandle()->AsArray();
 
@@ -138,7 +138,7 @@ FReply SDetailSingleItemRow::OnArrayAcceptDrop(const FDragDropEvent& DragDropEve
 		return FReply::Unhandled();
 	}
 
-	IDetailsViewPrivate* DetailsView = OwnerTreeNode.Pin()->GetDetailsView();
+	TSharedPtr<IDetailsViewPrivate> DetailsView = OwnerTreeNode.Pin()->GetDetailsViewSharedPtr();
 
 	TSharedPtr<FPropertyNode> SwappingPropertyNode = RowPtr->SwappablePropertyNode;
 	TSharedPtr<IPropertyHandle> SwappingHandle = PropertyEditorHelpers::GetPropertyHandle(SwappingPropertyNode.ToSharedRef(), DetailsView->GetNotifyHook(), DetailsView->GetPropertyUtilities());
@@ -257,7 +257,7 @@ TSharedPtr<IPropertyHandle> SDetailSingleItemRow::GetPropertyHandle() const
 	{
 		if (const TSharedPtr<FDetailTreeNode> OwnerTreeNodePtr = OwnerTreeNode.Pin())
 		{
-			if (IDetailsViewPrivate* DetailsView = OwnerTreeNodePtr->GetDetailsView())
+			if (TSharedPtr<IDetailsViewPrivate> DetailsView = OwnerTreeNodePtr->GetDetailsViewSharedPtr())
 			{
 				Handle = PropertyEditorHelpers::GetPropertyHandle(PropertyNode.ToSharedRef(), DetailsView->GetNotifyHook(), DetailsView->GetPropertyUtilities());
 			}
@@ -307,7 +307,7 @@ void SDetailSingleItemRow::Construct( const FArguments& InArgs, FDetailLayoutCus
 	FOnAcceptDrop AcceptDropDelegate;
 	FOnCanAcceptDrop CanAcceptDropDelegate;
 
-	IDetailsViewPrivate* DetailsView = InOwnerTreeNode->GetDetailsView();
+	TSharedPtr<IDetailsViewPrivate> DetailsView = InOwnerTreeNode->GetDetailsViewSharedPtr();
 	const FDetailColumnSizeData& ColumnSizeData = DetailsView ? DetailsView->GetColumnSizeData() : FDetailColumnSizeData();
 
 	PulseAnimation.AddCurve(0.0f, UE::PropertyEditor::Private::PulseAnimationLength, ECurveEaseFunction::CubicInOut);
@@ -828,14 +828,14 @@ FReply SDetailSingleItemRow::OnMouseButtonUp(const FGeometry& MyGeometry, const
 			bIsHandled = true;
 		}
 		// Paste is disabled if property editing is disabled
-		else if (PasteAction.CanExecute() && MouseEvent.GetEffectingButton() == EKeys::LeftMouseButton && OwnerTreeNode.Pin()->GetDetailsView()->IsPropertyEditingEnabled())
+		else if (PasteAction.CanExecute() && MouseEvent.GetEffectingButton() == EKeys::LeftMouseButton && OwnerTreeNode.Pin()->GetDetailsViewSharedPtr()->IsPropertyEditingEnabled())
 		{
 			PasteAction.Execute();
 			PulseAnimation.Play(SharedThis(this));
 			if(TSharedPtr<FPropertyNode> PropertyNode = GetPropertyNode())
 			{
 				// Mark property node as animating so we will animate after any potential re-construction
-				IDetailsViewPrivate* DetailsView = OwnerTreeNode.Pin()->GetDetailsView();
+				TSharedPtr<IDetailsViewPrivate> DetailsView = OwnerTreeNode.Pin()->GetDetailsViewSharedPtr();
 				DetailsView->MarkNodeAnimating(PropertyNode, UE::PropertyEditor::Private::PulseAnimationLength);
 			}
 			bIsHandled = true;
@@ -853,7 +853,7 @@ FReply SDetailSingleItemRow::OnMouseButtonUp(const FGeometry& MyGeometry, const
 
 		if (PropertyNode.IsValid() && PropertyHandle.IsValid())
 		{
-			IDetailsViewPrivate* DetailsView = OwnerTreeNode.Pin()->GetDetailsView();
+			TSharedPtr<IDetailsViewPrivate> DetailsView = OwnerTreeNode.Pin()->GetDetailsViewSharedPtr();
 			TSharedPtr<IPropertyHandle> Handle = PropertyEditorHelpers::GetPropertyHandle(PropertyNode.ToSharedRef(), DetailsView->GetNotifyHook(), DetailsView->GetPropertyUtilities());
 
 			FString Value;
@@ -1104,10 +1104,10 @@ void SDetailSingleItemRow::PopulateContextMenu(UToolMenu* ToolMenu)
 {
 	SDetailTableRowBase::PopulateContextMenu(ToolMenu);
 
-	IDetailsViewPrivate* OwningDetailsView = nullptr;
+	TSharedPtr<IDetailsViewPrivate> OwningDetailsView;
 	if (TSharedPtr<FDetailTreeNode> OwnerTreeNodePtr = OwnerTreeNode.Pin())
 	{
-		OwningDetailsView = OwnerTreeNodePtr->GetDetailsView();
+		OwningDetailsView = OwnerTreeNodePtr->GetDetailsViewSharedPtr();
 	}
 	
 	FToolMenuSection& EditSection = ToolMenu->FindOrAddSection(TEXT("Edit"));
@@ -1217,7 +1217,7 @@ void SDetailSingleItemRow::PopulateContextMenu(UToolMenu* ToolMenu)
 
 		// Favorite
 		{
-			if (OwnerTreeNode.Pin()->GetDetailsView()->IsFavoritingEnabled())
+			if (OwnerTreeNode.Pin()->GetDetailsViewSharedPtr()->IsFavoritingEnabled())
 			{
 				FUIAction FavoriteAction;
 				FavoriteAction.ExecuteAction = FExecuteAction::CreateSP(this, &SDetailSingleItemRow::OnFavoriteMenuToggle);
@@ -1323,7 +1323,7 @@ void SDetailSingleItemRow::OnCopyProperty()
 		TSharedPtr<FPropertyNode> PropertyNode = GetPropertyNode();
 		if (PropertyNode.IsValid())
 		{
-			IDetailsViewPrivate* DetailsView = OwnerTreeNode.Pin()->GetDetailsView();
+			TSharedPtr<IDetailsViewPrivate> DetailsView = OwnerTreeNode.Pin()->GetDetailsViewSharedPtr();
 			TSharedPtr<IPropertyHandle> Handle = PropertyEditorHelpers::GetPropertyHandle(PropertyNode.ToSharedRef(), DetailsView->GetNotifyHook(), DetailsView->GetPropertyUtilities());
 
 			FString Value;
@@ -1433,16 +1433,19 @@ TSharedRef<SWidget> SDetailSingleItemRow::GetNameWidget(TSharedRef<SWidget> Name
      {
 	 	const TSharedPtr<FDetailTreeNode> DetailTreeNodeSP = OwnerTreeNode.Pin();
 	 	
-	     if (DetailTreeNodeSP.IsValid() && DetailTreeNodeSP->GetDetailsView())
+	     if (DetailTreeNodeSP.IsValid())
 	     {
-	     	const TSharedPtr<FDetailsNameWidgetOverrideCustomization> DetailsNameWidgetOverrideCustomization =
-				DetailTreeNodeSP->GetDetailsView()->GetDetailsNameWidgetOverrideCustomization();
-
-	     	if ( DetailsNameWidgetOverrideCustomization.IsValid() )
-	     	{
-	     		const TSharedRef< FPropertyPath > Path = FPropertyNode::CreatePropertyPath(Node.ToSharedRef());
-	     		return DetailsNameWidgetOverrideCustomization->CustomizeName(NameWidget, Path.Get());
-	     	}
+			 if (TSharedPtr<IDetailsViewPrivate> DetailsView = DetailTreeNodeSP->GetDetailsViewSharedPtr())
+			 {
+				 const TSharedPtr<FDetailsNameWidgetOverrideCustomization> DetailsNameWidgetOverrideCustomization =
+					 DetailsView->GetDetailsNameWidgetOverrideCustomization();
+
+				 if (DetailsNameWidgetOverrideCustomization.IsValid())
+				 {
+					 const TSharedRef< FPropertyPath > Path = FPropertyNode::CreatePropertyPath(Node.ToSharedRef());
+					 return DetailsNameWidgetOverrideCustomization->CustomizeName(NameWidget, Path.Get());
+				 }
+			 }
 	     }
      }
 	return NameWidget;
@@ -1528,7 +1531,7 @@ void SDetailSingleItemRow::OnPasteProperty()
 	if (PasteFromText(TEXT(""), ClipboardContent))
 	{
 		// Mark property node as animating so we will animate after re-construction
-		IDetailsViewPrivate* DetailsView = OwnerTreeNode.Pin()->GetDetailsView();
+		TSharedPtr<IDetailsViewPrivate> DetailsView = OwnerTreeNode.Pin()->GetDetailsViewSharedPtr();
 		DetailsView->MarkNodeAnimating(PropertyNode, UE::PropertyEditor::Private::PulseAnimationLength);
 
 		// Need to refresh the details panel in case a property was pasted over another.
@@ -1544,7 +1547,7 @@ void SDetailSingleItemRow::OnPasteFromText(const FString& InTag, const FString&
 		if (PropertyNode.IsValid())
 		{
 			// Mark property node as animating so we will animate after re-construction
-			IDetailsViewPrivate* DetailsView = OwnerTreeNode.Pin()->GetDetailsView();
+			TSharedPtr<IDetailsViewPrivate> DetailsView = OwnerTreeNode.Pin()->GetDetailsViewSharedPtr();
 			DetailsView->MarkNodeAnimating(PropertyNode, UE::PropertyEditor::Private::PulseAnimationLength, InOperationId);
 		}
 	}
@@ -1701,7 +1704,7 @@ void SDetailSingleItemRow::PopulateExtensionWidget()
 	TSharedPtr<FDetailTreeNode> OwnerTreeNodePinned = OwnerTreeNode.Pin();
 	if (OwnerTreeNodePinned.IsValid())
 	{ 
-		IDetailsViewPrivate* DetailsView = OwnerTreeNodePinned->GetDetailsView();
+		TSharedPtr<IDetailsViewPrivate> DetailsView = OwnerTreeNodePinned->GetDetailsViewSharedPtr();
 		TSharedPtr<IDetailPropertyExtensionHandler> ExtensionHandler = DetailsView->GetExtensionHandler();
 		if (Customization->HasPropertyNode() && ExtensionHandler.IsValid())
 		{
@@ -1761,7 +1764,7 @@ bool SDetailSingleItemRow::IsFavorite() const
 			const FString& OriginalPath = Customization->CustomBuilderRow->GetOriginalPath();
 			if (!OriginalPath.IsEmpty())
 			{
-				return OwnerTreeNodePinned->GetDetailsView()->IsCustomBuilderFavorite(OriginalPath);
+				return OwnerTreeNodePinned->GetDetailsViewSharedPtr()->IsCustomBuilderFavorite(OriginalPath);
 			}
 		}
 	}
@@ -1782,7 +1785,7 @@ void SDetailSingleItemRow::OnFavoriteMenuToggle()
 		return;
 	}
 
-	IDetailsViewPrivate* DetailsView = OwnerTreeNodePinned->GetDetailsView();
+	TSharedPtr<IDetailsViewPrivate> DetailsView = OwnerTreeNodePinned->GetDetailsViewSharedPtr();
 
 	bool bNewValue = !IsFavorite();
 	if (Customization->HasPropertyNode())
diff --git a/Engine/Source/Editor/PropertyEditor/Private/SDetailTableRowBase.cpp b/Engine/Source/Editor/PropertyEditor/Private/SDetailTableRowBase.cpp
--- a/Engine/Source/Editor/PropertyEditor/Private/SDetailTableRowBase.cpp
+++ b/Engine/Source/Editor/PropertyEditor/Private/SDetailTableRowBase.cpp
@@ -29,7 +29,7 @@ FReply SDetailTableRowBase::OnMouseButtonUp(const FGeometry& MyGeometry, const F
 
 					UDetailRowMenuContext* RowMenuContext = NewObject<UDetailRowMenuContext>();					
 					RowMenuContext->PropertyHandles = GetPropertyHandles(true);
-					RowMenuContext->DetailsView = OwnerTreeNode.Pin()->GetDetailsView();
+					RowMenuContext->DetailsView = OwnerTreeNode.Pin()->GetDetailsViewSharedPtr();
 					RowMenuContext->ForceRefreshWidget().AddSPLambda(this, [this]
 					{
 						ForceRefresh();						
@@ -64,13 +64,16 @@ int32 SDetailTableRowBase::GetIndentLevelForBackgroundColor() const
 		IndentLevel = GetIndentLevel() - 1;
 	}
 
-	TSharedPtr<FDetailTreeNode> DetailTreeNode = OwnerTreeNode.Pin();
-	if (DetailTreeNode.IsValid() && 
-		DetailTreeNode->GetDetailsView() != nullptr && 
-		DetailTreeNode->GetDetailsView()->ContainsMultipleTopLevelObjects())
+	if (TSharedPtr<FDetailTreeNode> DetailTreeNode = OwnerTreeNode.Pin())
 	{
-		// if the row is in a multiple top level object display (eg. Project Settings), don't display an indent for the initial level
-		--IndentLevel;
+		if (TSharedPtr<IDetailsViewPrivate> DetailsView = DetailTreeNode->GetDetailsViewSharedPtr())
+		{
+			if (DetailsView->ContainsMultipleTopLevelObjects())
+			{
+				// if the row is in a multiple top level object display (eg. Project Settings), don't display an indent for the initial level
+				--IndentLevel;
+			}
+		}
 	}
 
 	return FMath::Max(0, IndentLevel);
@@ -183,7 +186,7 @@ void SDetailTableRowBase::ForceRefresh()
 	if (const TSharedPtr<FDetailTreeNode> Owner = OwnerTreeNode.Pin();
 		Owner.IsValid())
 	{
-		if (IDetailsViewPrivate* DetailsView = Owner->GetDetailsView())
+		if (TSharedPtr<IDetailsViewPrivate> DetailsView = Owner->GetDetailsViewSharedPtr())
 		{
 			DetailsView->ForceRefresh();
 		}
diff --git a/Engine/Source/Editor/PropertyEditor/Private/SDetailTableRowBase.h b/Engine/Source/Editor/PropertyEditor/Private/SDetailTableRowBase.h
--- a/Engine/Source/Editor/PropertyEditor/Private/SDetailTableRowBase.h
+++ b/Engine/Source/Editor/PropertyEditor/Private/SDetailTableRowBase.h
@@ -60,7 +60,7 @@ private:
 		{
 			const bool bRecursive = true;
 			const bool bIsExpanded = true;
-			OwnerTreeNodePin->GetDetailsView()->SetNodeExpansionState( OwnerTreeNodePin.ToSharedRef(), bIsExpanded, bRecursive );
+			OwnerTreeNodePin->GetDetailsViewSharedPtr()->SetNodeExpansionState( OwnerTreeNodePin.ToSharedRef(), bIsExpanded, bRecursive );
 		}
 	}
 
@@ -71,7 +71,7 @@ private:
 		{
 			const bool bRecursive = true;
 			const bool bIsExpanded = false;
-			OwnerTreeNodePin->GetDetailsView()->SetNodeExpansionState( OwnerTreeNodePin.ToSharedRef(), bIsExpanded, bRecursive );
+			OwnerTreeNodePin->GetDetailsViewSharedPtr()->SetNodeExpansionState( OwnerTreeNodePin.ToSharedRef(), bIsExpanded, bRecursive );
 		}
 	}
 	
diff --git a/Engine/Source/Editor/PropertyEditor/Private/SDetailsSplitter.cpp b/Engine/Source/Editor/PropertyEditor/Private/SDetailsSplitter.cpp
--- a/Engine/Source/Editor/PropertyEditor/Private/SDetailsSplitter.cpp
+++ b/Engine/Source/Editor/PropertyEditor/Private/SDetailsSplitter.cpp
@@ -16,7 +16,7 @@ namespace DetailsSplitterHelpers
 {
 	const TArray<TWeakObjectPtr<UObject>>* GetObjects(const TSharedPtr<FDetailTreeNode>& TreeNode)
 	{
-		if (const IDetailsViewPrivate* DetailsView = TreeNode->GetDetailsView())
+		if (TSharedPtr<const IDetailsViewPrivate> DetailsView = TreeNode->GetDetailsViewSharedPtr())
 		{
 			return &DetailsView->GetSelectedObjects();
 		}
@@ -491,7 +491,7 @@ void SDetailsSplitter::HighlightFromMergeResults(const TMap<FString, TMap<FPrope
 		const TSharedPtr<FDetailTreeNode> DetailNode = DiffNode->ValueA.IsValid() ? DiffNode->ValueA.Pin() : DiffNode->ValueB.Pin();
 		
 		FPropertyPath Path = DetailNode->GetPropertyPath();
-		const UObject* OwningObject = DetailNode->GetDetailsView()->GetSelectedObjects()[0].Get();
+		const UObject* OwningObject = DetailNode->GetDetailsViewSharedPtr()->GetSelectedObjects()[0].Get();
 		if (OwningObject && Path.IsValid())
 		{
 			FPropertyPath PathFromSubObject;
diff --git a/Engine/Source/Editor/PropertyEditor/Public/DetailLayoutBuilder.h b/Engine/Source/Editor/PropertyEditor/Public/DetailLayoutBuilder.h
--- a/Engine/Source/Editor/PropertyEditor/Public/DetailLayoutBuilder.h
+++ b/Engine/Source/Editor/PropertyEditor/Public/DetailLayoutBuilder.h
@@ -58,12 +58,24 @@ public:
 	/**
 	 * @return the parent detail view for this layout builder
 	 */
-	virtual const IDetailsView* GetDetailsView() const = 0;
+	virtual TSharedPtr<const IDetailsView> GetDetailsViewSharedPtr() const = 0;
+
+	UE_DEPRECATED(5.5, "Use GetDetailsViewSharedPtr() instead.")
+	const IDetailsView* GetDetailsView() const
+	{
+		return GetDetailsViewSharedPtr().Get();
+	}
 
 	/**
 	 * @return the parent detail view for this layout builder
 	 */
-	virtual IDetailsView* GetDetailsView() = 0;
+	virtual TSharedPtr<IDetailsView> GetDetailsViewSharedPtr() = 0;
+
+	UE_DEPRECATED(5.5, "Use GetDetailsViewSharedPtr() instead.")
+	IDetailsView* GetDetailsView()
+	{
+		return GetDetailsViewSharedPtr().Get();
+	}
 
 	/**
 	 * @return The base class of the objects being customized in this detail layout
diff --git a/Engine/Source/Editor/PropertyEditor/Public/DetailRowMenuContext.h b/Engine/Source/Editor/PropertyEditor/Public/DetailRowMenuContext.h
--- a/Engine/Source/Editor/PropertyEditor/Public/DetailRowMenuContext.h
+++ b/Engine/Source/Editor/PropertyEditor/Public/DetailRowMenuContext.h
@@ -1,4 +1,4 @@
-﻿// Copyright Epic Games, Inc. All Rights Reserved.
+// Copyright Epic Games, Inc. All Rights Reserved.
 
 #pragma once
 
@@ -15,11 +15,6 @@ class PROPERTYEDITOR_API UDetailRowMenuContext : public UObject
 	GENERATED_BODY()
 
 public:
-	virtual ~UDetailRowMenuContext() override
-	{
-		DetailsView = nullptr;		
-	}
-	
 	/** Optionally invoke to refresh the widget. */
 	TMulticastDelegate<void()>& ForceRefreshWidget() { return ForceRefreshWidgetDelegate; } 
 	
@@ -27,7 +22,7 @@ public:
 	TArray<TSharedPtr<IPropertyHandle>> PropertyHandles;
 
 	/** Containing DetailsView. */
-	IDetailsView* DetailsView = nullptr;
+	TWeakPtr<IDetailsView> DetailsView;
 	
 private:
 	/** Optionally invoke to refresh the widget. */
diff --git a/Engine/Source/Editor/PropertyEditor/Public/DetailTreeNode.h b/Engine/Source/Editor/PropertyEditor/Public/DetailTreeNode.h
--- a/Engine/Source/Editor/PropertyEditor/Public/DetailTreeNode.h
+++ b/Engine/Source/Editor/PropertyEditor/Public/DetailTreeNode.h
@@ -44,7 +44,13 @@ public:
 	virtual bool GetInitiallyCollapsed() const override { return false; }
 
 	/** @return The details view that this node is in */
-	virtual IDetailsViewPrivate* GetDetailsView() const = 0;
+	virtual TSharedPtr<IDetailsViewPrivate> GetDetailsViewSharedPtr() const = 0;
+
+	UE_DEPRECATED(5.5, "Use GetDetailsViewSharedPtr() instead.")
+	IDetailsViewPrivate* GetDetailsView() const
+	{
+		return GetDetailsViewSharedPtr().Get();
+	}
 
 	/**
 	 * Generates the widget representing this node
diff --git a/Engine/Source/Editor/PropertyEditor/Public/IDetailTreeNode.h b/Engine/Source/Editor/PropertyEditor/Public/IDetailTreeNode.h
--- a/Engine/Source/Editor/PropertyEditor/Public/IDetailTreeNode.h
+++ b/Engine/Source/Editor/PropertyEditor/Public/IDetailTreeNode.h
@@ -123,7 +123,13 @@ public:
 	/** 
 	 * @return Get the details view that this node is contained in.
 	 */
-	virtual class IDetailsView* GetNodeDetailsView() const = 0;
+	virtual TSharedPtr<class IDetailsView> GetNodeDetailsViewSharedPtr() const = 0;
+
+	UE_DEPRECATED(5.5, "Use GetNodeDetailsViewSharedPtr() instead.")
+	class IDetailsView* GetNodeDetailsView() const
+	{
+		return GetNodeDetailsViewSharedPtr().Get();
+	}
 
 	/**
 	 * @return The type of this node.  Should be used to determine any external styling to apply to the generated r ow
diff --git a/Engine/Source/Editor/Sequencer/Private/SequencerModule.cpp b/Engine/Source/Editor/Sequencer/Private/SequencerModule.cpp
--- a/Engine/Source/Editor/Sequencer/Private/SequencerModule.cpp
+++ b/Engine/Source/Editor/Sequencer/Private/SequencerModule.cpp
@@ -99,7 +99,7 @@ static TSharedPtr<IDetailKeyframeHandler> GetKeyframeHandler(TWeakPtr<IDetailTre
 		return TSharedPtr<IDetailKeyframeHandler>();
 	}
 
-	IDetailsView* DetailsView = OwnerTreeNodePtr->GetNodeDetailsView();
+	TSharedPtr<IDetailsView> DetailsView = OwnerTreeNodePtr->GetNodeDetailsViewSharedPtr();
 	if (DetailsView == nullptr)
 	{
 		return TSharedPtr<IDetailKeyframeHandler>();
diff --git a/Engine/Source/Editor/UniversalObjectLocatorEditor/Private/UniversalObjectLocatorCustomization.cpp b/Engine/Source/Editor/UniversalObjectLocatorEditor/Private/UniversalObjectLocatorCustomization.cpp
--- a/Engine/Source/Editor/UniversalObjectLocatorEditor/Private/UniversalObjectLocatorCustomization.cpp
+++ b/Engine/Source/Editor/UniversalObjectLocatorEditor/Private/UniversalObjectLocatorCustomization.cpp
@@ -246,7 +246,7 @@ void FUniversalObjectLocatorCustomization::CustomizeHeader(TSharedRef<IPropertyH
 
 void FUniversalObjectLocatorCustomization::CustomizeChildren(TSharedRef<IPropertyHandle> StructPropertyHandle, IDetailChildrenBuilder& StructBuilder, IPropertyTypeCustomizationUtils& StructCustomizationUtils)
 {
-	if (IDetailsView* DetailsView = StructBuilder.GetParentCategory().GetParentLayout().GetDetailsView())
+	if (TSharedPtr<IDetailsView> DetailsView = StructBuilder.GetParentCategory().GetParentLayout().GetDetailsViewSharedPtr())
 	{
 		DetailsView->OnFinishedChangingProperties().AddSPLambda(this, [this](const FPropertyChangedEvent&)
 			{
