Property Editor - Add an api to the FDetailsDisplayManager class to allow implmenters to supply a widget which will be overlayed over the detail tree portion of the property editor.  This functionality is being added to support drag/drop onto the detail tree initiated from external code.

diff --git a/Engine/Source/Editor/PropertyEditor/Private/DetailsDisplayManager.cpp b/Engine/Source/Editor/PropertyEditor/Private/DetailsDisplayManager.cpp
--- a/Engine/Source/Editor/PropertyEditor/Private/DetailsDisplayManager.cpp
+++ b/Engine/Source/Editor/PropertyEditor/Private/DetailsDisplayManager.cpp
@@ -70,6 +70,11 @@ FMargin FDetailsDisplayManager::GetTablePadding() const
 	return Style ? Style->GetTablePadding(bIsScrollBarNeeded) : 0;
 }
 
+TSharedRef<SWidget> FDetailsDisplayManager::ConstructDetailTreeOverlay()
+{
+	return SNullWidget::NullWidget;
+}
+
 void FDetailsDisplayManager::UpdatePropertyForCategory(FName InCategoryObjectName, FProperty* Property, bool bAddProperty)
 {
 	TSet<FProperty*> Set;
diff --git a/Engine/Source/Editor/PropertyEditor/Private/SDetailsView.cpp b/Engine/Source/Editor/PropertyEditor/Private/SDetailsView.cpp
--- a/Engine/Source/Editor/PropertyEditor/Private/SDetailsView.cpp
+++ b/Engine/Source/Editor/PropertyEditor/Private/SDetailsView.cpp
@@ -474,6 +474,15 @@ void SDetailsView::Construct(const FArguments& InArgs, const FDetailsViewArgs& I
 			.Image(FAppStyle::GetBrush("Searching.SearchActiveBorder"))
 			.Visibility_Lambda([this]() { return (this->GetFilterBoxVisibility() == EVisibility::Visible) && this->HasActiveSearch() ? EVisibility::HitTestInvisible : EVisibility::Collapsed; })
 		]
+		+ SOverlay::Slot()
+		[
+			SAssignNew(ExternalOverlayContainer, SBox)
+			.Visibility(DisplayManager.IsValid() && DisplayManager->SupportsDetailTreeOverlay() ? EVisibility::SelfHitTestInvisible : EVisibility::Collapsed)
+			.Content()
+			[
+				DisplayManager.IsValid() && DisplayManager->SupportsDetailTreeOverlay() ? DisplayManager->ConstructDetailTreeOverlay() : SNullWidget::NullWidget
+			]
+		]
 	];
 
 	ChildSlot
@@ -1605,6 +1614,7 @@ const FDetailsViewStyleKey& SDetailsView::GetStyleKey()
 
 void SDetailsView::RefreshDisplayManager()
 {
+	TSharedPtr<FDetailsDisplayManager> OldDisplayManager = DisplayManager;
 	if (ObjectFilter.IsValid())
 	{
 		DisplayManager = ObjectFilter->GetDisplayManager();
@@ -1613,6 +1623,19 @@ void SDetailsView::RefreshDisplayManager()
 	{
 		DisplayManager = MakeShared<FDetailsDisplayManager>();
 	}
+	if (OldDisplayManager != DisplayManager && ExternalOverlayContainer.IsValid())
+	{
+		if (DisplayManager->SupportsDetailTreeOverlay())
+		{
+			ExternalOverlayContainer->SetVisibility(EVisibility::SelfHitTestInvisible);
+			ExternalOverlayContainer->SetContent(DisplayManager->ConstructDetailTreeOverlay());
+		}
+		else
+		{
+			ExternalOverlayContainer->SetVisibility(EVisibility::Collapsed);
+			ExternalOverlayContainer->SetContent(SNullWidget::NullWidget);
+		}
+	}
 }
 
 TSharedPtr<FDetailsDisplayManager> SDetailsView::GetDisplayManager()
diff --git a/Engine/Source/Editor/PropertyEditor/Private/SDetailsViewBase.h b/Engine/Source/Editor/PropertyEditor/Private/SDetailsViewBase.h
--- a/Engine/Source/Editor/PropertyEditor/Private/SDetailsViewBase.h
+++ b/Engine/Source/Editor/PropertyEditor/Private/SDetailsViewBase.h
@@ -452,6 +452,8 @@ protected:
 	TSharedPtr<SDetailTree> DetailTree;
 	/** Root tree nodes visible in the tree */
 	FDetailNodeList RootTreeNodes;
+	/** A container for external overlays provided by display managers. */
+	TSharedPtr<SBox> ExternalOverlayContainer;
 	/** Delegate executed to determine if a property should be visible */
 	FIsPropertyVisible IsPropertyVisibleDelegate;
 	/** Delegate executed to determine if a custom row should be visible. */
diff --git a/Engine/Source/Editor/PropertyEditor/Public/DetailsDisplayManager.h b/Engine/Source/Editor/PropertyEditor/Public/DetailsDisplayManager.h
--- a/Engine/Source/Editor/PropertyEditor/Public/DetailsDisplayManager.h
+++ b/Engine/Source/Editor/PropertyEditor/Public/DetailsDisplayManager.h
@@ -92,6 +92,17 @@ public:
 	 */
 	FMargin GetTablePadding() const;
 
+	/**
+	 * Returns whether or not this display manager provides a widget to overlay over the detail tree.  This can be used to facilitate interactions like drag and drop.
+	 */
+	PROPERTYEDITOR_API virtual bool SupportsDetailTreeOverlay() const { return false; }
+
+	/**
+	 * Returns a widget which will be put in an overlay slot which is over the detail tree portion of the widget.
+	 * @note Implementations must manage hit test visibility of the supplied widget carefully to prevent blocking UI interactions with the detail tree widgets.
+	 */
+	PROPERTYEDITOR_API virtual TSharedRef<SWidget> ConstructDetailTreeOverlay();
+
 	/**
 	 * If TSharedRef<FComplexPropertyNode> Node has a valid UObject instance, add an empty properties
 	 * Category to the DetailLayoutBuilder
