SDetailSingleItemRow::OnResetToDefaultClicked to handle optional properties

Background:
The UI for an optional is shared with its value.
It can be reset to default via the OnResetToDefaultClicked callback in some cases (notably if DetailsDisplayManager has been set or if using PropertyRowExtensions).
It can be reset to default via other methods - notably FPropertyEditor will instantiate a callback to FPropertyEditor::OnClearOptionalValue

What changed:
* OnResetToDefaultClicked will now handle the case where the property is optional.  It could walk up to the parent to check if it is within an optional container similar to FPropertyEditor::OnClearOptionalValue to handle the case where the property handle is the value.

diff --git a/Engine/Source/Editor/PropertyEditor/Private/SDetailSingleItemRow.cpp b/Engine/Source/Editor/PropertyEditor/Private/SDetailSingleItemRow.cpp
--- a/Engine/Source/Editor/PropertyEditor/Private/SDetailSingleItemRow.cpp
+++ b/Engine/Source/Editor/PropertyEditor/Private/SDetailSingleItemRow.cpp
@@ -887,7 +887,19 @@ void SDetailSingleItemRow::OnResetToDefaultClicked() const
 	}
 	else if (PropertyHandle.IsValid())
 	{
+		TSharedPtr<IPropertyHandleOptional> OptionalHandle = PropertyHandle->AsOptional();
+		if (!OptionalHandle.IsValid())
+		{
+			// This could be the value calling as we share UI between options and their value
+			OptionalHandle = PropertyHandle->GetParentHandle()->AsOptional();
+			if (OptionalHandle.IsValid())
+			{
+				PropertyHandle = PropertyHandle->GetParentHandle();
+			}
+		}
+		
 		PropertyHandle->ResetToDefault();
+		
 	}
 }
 
