[PropertyNode] BuildPropertyChain should allow the chain to go up to the first parent UObject and not the first complex node (can be a struct) since the rest of the code expects the path to start on a UObject.

diff --git a/Engine/Source/Editor/PropertyEditor/Private/PropertyNode.cpp b/Engine/Source/Editor/PropertyEditor/Private/PropertyNode.cpp
--- a/Engine/Source/Editor/PropertyEditor/Private/PropertyNode.cpp
+++ b/Engine/Source/Editor/PropertyEditor/Private/PropertyNode.cpp
@@ -3529,7 +3529,7 @@ TSharedRef<FEditPropertyChain> FPropertyNode::BuildPropertyChain( FProperty* InP
 
 	const FPropertyNode* ItemNode = this;
 
-	const FComplexPropertyNode* ComplexNode = FindComplexParent();
+	const FComplexPropertyNode* ComplexNode = FindObjectItemParent();
 	FProperty* MemberProperty = InProperty;
 
 	do
