Property Editor: Add utility function to create a custom property editor button

diff --git a/Engine/Source/Editor/PropertyEditor/Private/PropertyCustomizationHelpers.cpp b/Engine/Source/Editor/PropertyEditor/Private/PropertyCustomizationHelpers.cpp
--- a/Engine/Source/Editor/PropertyEditor/Private/PropertyCustomizationHelpers.cpp
+++ b/Engine/Source/Editor/PropertyEditor/Private/PropertyCustomizationHelpers.cpp
@@ -104,6 +104,19 @@ namespace PropertyCustomizationHelpers
 		FSimpleDelegate OnClickAction;
 	};
 
+	TSharedRef<SWidget> MakeCustomButton(const FSlateBrush* IconBrush, FSimpleDelegate OnClicked,
+		TAttribute<FText> ToolTipText, TAttribute<bool> IsEnabled, TAttribute<EVisibility> IsVisible)
+	{
+		return
+			SNew(SPropertyEditorButton)
+			.Text(ToolTipText)
+			.Image(IconBrush)
+			.OnClickAction(OnClicked)
+			.IsEnabled(IsEnabled)
+			.Visibility(IsVisible)
+			.IsFocusable(false);
+	}
+
 	TSharedRef<SWidget> MakeResetButton(FSimpleDelegate OnResetClicked, TAttribute<FText> OptionalToolTipText /*= FText()*/, TAttribute<bool> IsEnabled /*= true*/)
 	{
 		return
diff --git a/Engine/Source/Editor/PropertyEditor/Public/PropertyCustomizationHelpers.h b/Engine/Source/Editor/PropertyEditor/Public/PropertyCustomizationHelpers.h
--- a/Engine/Source/Editor/PropertyEditor/Public/PropertyCustomizationHelpers.h
+++ b/Engine/Source/Editor/PropertyEditor/Public/PropertyCustomizationHelpers.h
@@ -153,6 +153,7 @@ struct FPropertyFunctionCallDelegates
 
 namespace PropertyCustomizationHelpers
 {
+	PROPERTYEDITOR_API TSharedRef<SWidget> MakeCustomButton(const FSlateBrush* IconBrush, FSimpleDelegate OnClicked, TAttribute<FText> ToolTipText = FText(), TAttribute<bool> IsEnabled = true, TAttribute<EVisibility> IsVisible = EVisibility::Visible);
 	PROPERTYEDITOR_API TSharedRef<SWidget> MakeResetButton(FSimpleDelegate OnResetClicked, TAttribute<FText> OptionalToolTipText = FText(), TAttribute<bool> IsEnabled = true);
 	PROPERTYEDITOR_API TSharedRef<SWidget> MakeAddButton( FSimpleDelegate OnAddClicked, TAttribute<FText> OptionalToolTipText = FText(), TAttribute<bool> IsEnabled = true );
 	PROPERTYEDITOR_API TSharedRef<SWidget> MakeRemoveButton( FSimpleDelegate OnRemoveClicked, TAttribute<FText> OptionalToolTipText = FText(), TAttribute<bool> IsEnabled = true );
