UE-230360 - "TOptional is not filtering actor workflows correctly"

- add support for soft class/object pts to optional 'picker' logic

diff --git a/Engine/Source/Editor/PropertyEditor/Private/PropertyCustomizationHelpers.cpp b/Engine/Source/Editor/PropertyEditor/Private/PropertyCustomizationHelpers.cpp
--- a/Engine/Source/Editor/PropertyEditor/Private/PropertyCustomizationHelpers.cpp
+++ b/Engine/Source/Editor/PropertyEditor/Private/PropertyCustomizationHelpers.cpp
@@ -241,26 +241,25 @@ namespace PropertyCustomizationHelpers
 		const UClass* RequiredInterface = nullptr;
 		const UClass* SelectedClass = nullptr;
 
-		if (const FClassProperty* ClassProperty = CastField<FClassProperty>(ValueProperty))
+		if(const FInterfaceProperty* InterfaceProperty = CastField<FInterfaceProperty>(ValueProperty))
 		{
-			MetaClass = ClassProperty->MetaClass;
-			RequiredInterface = ClassProperty->GetOwnerProperty()->GetClassMetaData(TEXT("MustImplement"));
+			RequiredInterface = InterfaceProperty->InterfaceClass;
 		}
-		else if(const FInterfaceProperty* InterfaceProperty = CastField<FInterfaceProperty>(ValueProperty))
+		else if (ValueProperty->IsA(FClassProperty::StaticClass()) || ValueProperty->IsA(FSoftClassProperty::StaticClass()))
 		{
-			RequiredInterface = InterfaceProperty->InterfaceClass;
+			MetaClass = ValueProperty->IsA(FClassProperty::StaticClass())
+				? CastFieldChecked<FClassProperty>(ValueProperty)->MetaClass
+				: CastFieldChecked<FSoftClassProperty>(ValueProperty)->MetaClass;
+			RequiredInterface = ValueProperty->GetOwnerProperty()->GetClassMetaData(TEXT("MustImplement"));
 		}
-		else if (const FObjectPropertyBase* ObjectProperty = CastField<FObjectProperty>(ValueProperty))
+		else if (ValueProperty->IsA(FObjectProperty::StaticClass()) || ValueProperty->IsA(FSoftObjectProperty::StaticClass()))
 		{
-			MetaClass = ObjectProperty->PropertyClass;
+			MetaClass = ValueProperty->IsA(FObjectProperty::StaticClass())
+				? CastFieldChecked<FObjectProperty>(ValueProperty)->PropertyClass
+				: CastFieldChecked<FSoftObjectProperty>(ValueProperty)->PropertyClass;
 
 			// Verse generated interfaces are not fully supported yet, for now they are stored in FObjectProperty metadata
-			RequiredInterface = ObjectProperty->GetOwnerProperty()->GetClassMetaData(TEXT("ObjectMustImplement"));
-
-			if (!RequiredInterface && ObjectProperty->PropertyClass->HasAnyClassFlags(EClassFlags::CLASS_Interface))
-			{
-				RequiredInterface = ObjectProperty->PropertyClass;
-			}
+			RequiredInterface = ValueProperty->GetOwnerProperty()->GetClassMetaData(TEXT("ObjectMustImplement"));
 		}
 
 		TWeakPtr<FPropertyNode> WeakPropertyNode = PropertyNode;
diff --git a/Engine/Source/Editor/PropertyEditor/Private/PropertyHandleImpl.cpp b/Engine/Source/Editor/PropertyEditor/Private/PropertyHandleImpl.cpp
--- a/Engine/Source/Editor/PropertyEditor/Private/PropertyHandleImpl.cpp
+++ b/Engine/Source/Editor/PropertyEditor/Private/PropertyHandleImpl.cpp
@@ -5373,13 +5373,17 @@ FPropertyAccess::Result FPropertyHandleOptional::SetOptionalValue(FProperty* New
 			OptionalProperty->SetValueProperty(NewProperty);
 		}
 
+		// TODO: This should be re-worked to leverage using 'import_text' instead of directly setting new values as we are duplicating logic somewhat here (non-trivial task)
 		FProperty* ValueProperty = OptionalProperty->GetValueProperty();
 		if (ValueProperty && NewObjectClass)
 		{
-			if (FClassProperty* ClassProperty = CastField<FClassProperty>(ValueProperty))
+			if (ValueProperty->IsA(FClassProperty::StaticClass()) || ValueProperty->IsA(FSoftClassProperty::StaticClass()))
 			{
-				void* ClassPropertyValuePtr = ClassProperty->ContainerPtrToValuePtr<void>(Addr);
-				ClassProperty->SetObjectPropertyValue(ClassPropertyValuePtr, const_cast<UClass*>(NewObjectClass));
+				if (FObjectPropertyBase* ClassProperty = CastField<FObjectPropertyBase>(ValueProperty))
+				{
+					void* ClassPropertyValuePtr = ClassProperty->ContainerPtrToValuePtr<void>(Addr);
+					ClassProperty->SetObjectPropertyValue(ClassPropertyValuePtr, const_cast<UClass*>(NewObjectClass));
+				}
 			}
 			else if (NewObjectClass->HasAnyClassFlags(EClassFlags::CLASS_Abstract))
 			{
@@ -5402,7 +5406,7 @@ FPropertyAccess::Result FPropertyHandleOptional::SetOptionalValue(FProperty* New
 				UObject* Obj = NewObject<UObject>(Outer, NewObjectClass);
 				if (Obj)
 				{
-					if (FObjectProperty* ObjectProperty = CastField<FObjectProperty>(ValueProperty))
+					if (FObjectPropertyBase* ObjectProperty = CastField<FObjectPropertyBase>(ValueProperty))
 					{
 						void* ObjectPropertyValuePtr = ObjectProperty->ContainerPtrToValuePtr<void>(Addr);
 						ObjectProperty->SetObjectPropertyValue(ObjectPropertyValuePtr, Obj);
