Fixed TOptional properties not being fully hidden when using MarkHiddenByCustomization because the flag was not applied to the eventual child optional node (also fixed MarkResetToDefaultCustomized and ClearResetToDefaultCustomized)

diff --git a/Engine/Source/Editor/PropertyEditor/Private/PropertyHandleImpl.cpp b/Engine/Source/Editor/PropertyEditor/Private/PropertyHandleImpl.cpp
--- a/Engine/Source/Editor/PropertyEditor/Private/PropertyHandleImpl.cpp
+++ b/Engine/Source/Editor/PropertyEditor/Private/PropertyHandleImpl.cpp
@@ -2295,25 +2295,43 @@ FText FPropertyHandleBase::GetResetToDefaultLabel() const
 
 void FPropertyHandleBase::MarkHiddenByCustomization()
 {
-	if( Implementation->GetPropertyNode().IsValid() )
+	if(TSharedPtr<FPropertyNode> PropertyNode = Implementation->GetPropertyNode(); PropertyNode.IsValid() )
 	{
-		Implementation->GetPropertyNode()->SetNodeFlags( EPropertyNodeFlags::IsCustomized, true );
+		PropertyNode->SetNodeFlags( EPropertyNodeFlags::IsCustomized, true );
+
+		// If this is a set optional property, propagate the flag change 
+		if (TSharedPtr<FPropertyNode> OptionalValueNode = PropertyNode->GetOptionalValueNode(); OptionalValueNode.IsValid())
+		{
+			OptionalValueNode->SetNodeFlags(EPropertyNodeFlags::IsCustomized, true);
+		}
 	}
 }
 
 void FPropertyHandleBase::MarkResetToDefaultCustomized(bool bCustomized/* = true*/)
 {
-	if (Implementation->GetPropertyNode().IsValid())
+	if (TSharedPtr<FPropertyNode> PropertyNode = Implementation->GetPropertyNode(); PropertyNode.IsValid())
 	{
-		Implementation->GetPropertyNode()->SetNodeFlags(EPropertyNodeFlags::HasCustomResetToDefault, bCustomized);
+		PropertyNode->SetNodeFlags(EPropertyNodeFlags::HasCustomResetToDefault, bCustomized);
+
+		// If this is a set optional property, propagate the flag change 
+		if (TSharedPtr<FPropertyNode> OptionalValueNode = PropertyNode->GetOptionalValueNode(); OptionalValueNode.IsValid())
+		{
+			OptionalValueNode->SetNodeFlags(EPropertyNodeFlags::HasCustomResetToDefault, bCustomized);
+		}
 	}
 }
 
 void FPropertyHandleBase::ClearResetToDefaultCustomized()
 {
-	if (Implementation->GetPropertyNode().IsValid())
+	if (TSharedPtr<FPropertyNode> PropertyNode = Implementation->GetPropertyNode(); PropertyNode.IsValid())
 	{
-		Implementation->GetPropertyNode()->SetNodeFlags(EPropertyNodeFlags::HasCustomResetToDefault, false);
+		PropertyNode->SetNodeFlags(EPropertyNodeFlags::HasCustomResetToDefault, true);
+
+		// If this is a set optional property, propagate the flag change 
+		if (TSharedPtr<FPropertyNode> OptionalValueNode = PropertyNode->GetOptionalValueNode(); OptionalValueNode.IsValid())
+		{
+			OptionalValueNode->SetNodeFlags(EPropertyNodeFlags::HasCustomResetToDefault, true);
+		}
 	}
 }
 
