Apply "AllowEditInlineCustomization" metadata only to UVerseAsset object properties
* Alleviates one part of a performance issue with setting all ObjectProperties that have non-native PropertyClass types with EditInline.
* Works by avoiding building out unecessarily the details tree and applying customizations to these nodes in FDetailsItemNode
* Allow overriding the creation of Category Nodes at the root level of the DetailsView when objects are set.  Removes UX clutter of redudant category node but disables categories for scene graph

diff --git a/Engine/Source/Editor/PropertyEditor/Private/DetailsDisplayManager.cpp b/Engine/Source/Editor/PropertyEditor/Private/DetailsDisplayManager.cpp
--- a/Engine/Source/Editor/PropertyEditor/Private/DetailsDisplayManager.cpp
+++ b/Engine/Source/Editor/PropertyEditor/Private/DetailsDisplayManager.cpp
@@ -28,6 +28,11 @@ FDetailsDisplayManager::~FDetailsDisplayManager()
 	OnDetailsNeedsUpdate.Unbind();
 }
 
+TOptional<bool> FDetailsDisplayManager::OverrideCreateCategoryNodes() const
+{
+	return TOptional<bool>();
+}
+
 bool FDetailsDisplayManager::ShouldHideComponentEditor()
 {
 	return !GetForceShowSubObjectEditor();
diff --git a/Engine/Source/Editor/PropertyEditor/Private/SDetailsView.cpp b/Engine/Source/Editor/PropertyEditor/Private/SDetailsView.cpp
--- a/Engine/Source/Editor/PropertyEditor/Private/SDetailsView.cpp
+++ b/Engine/Source/Editor/PropertyEditor/Private/SDetailsView.cpp
@@ -1046,6 +1046,10 @@ void SDetailsView::PostSetObject(const TArray<FDetailsViewObjectRoot>& Roots)
 	InitParams.ArrayOffset = 0;
 	InitParams.ArrayIndex = INDEX_NONE;
 	InitParams.bAllowChildren = true;
+	if (DisplayManager && DisplayManager->OverrideCreateCategoryNodes().IsSet())
+	{
+		InitParams.bCreateCategoryNodes = DisplayManager->OverrideCreateCategoryNodes().GetValue();
+	}
 	InitParams.bForceHiddenPropertyVisibility = 
 		FPropertySettings::Get().ShowHiddenProperties() || 
 		bShowHiddenPropertiesWhilePlaying ||
diff --git a/Engine/Source/Editor/PropertyEditor/Public/DetailsDisplayManager.h b/Engine/Source/Editor/PropertyEditor/Public/DetailsDisplayManager.h
--- a/Engine/Source/Editor/PropertyEditor/Public/DetailsDisplayManager.h
+++ b/Engine/Source/Editor/PropertyEditor/Public/DetailsDisplayManager.h
@@ -29,6 +29,8 @@ public:
 
 	PROPERTYEDITOR_API virtual ~FDetailsDisplayManager();
 
+	virtual TOptional<bool> OverrideCreateCategoryNodes() const;
+
 	/**
 	 * Returns a boolean indicating if the Component Editor should be hidden 
 	 */
