Property Editor: Fix custom rows did not respect centered vertical alignment when extension buttons had a larger height than the property widget

diff --git a/Engine/Source/Editor/PropertyEditor/Private/DetailPropertyRow.cpp b/Engine/Source/Editor/PropertyEditor/Private/DetailPropertyRow.cpp
--- a/Engine/Source/Editor/PropertyEditor/Private/DetailPropertyRow.cpp
+++ b/Engine/Source/Editor/PropertyEditor/Private/DetailPropertyRow.cpp
@@ -1113,6 +1113,7 @@ void FDetailPropertyRow::MakeValueWidget( FDetailWidgetRow& Row, const TSharedPt
 		MaxWidth = InCustomRow->ValueWidget.MaxWidth;
 
 		ValueWidget->AddSlot()
+			.VAlign(VerticalAlignment)
 		[
 			InCustomRow->ValueWidget.Widget
 		];
