ResetToDefault will no longer unconditionally set property initialization state to false
* State is set explicitly to false if there is no archetype.  Initialization state is untracked for values that are not CPF_RequiredParm so this operation is a noop in those cases.
* State is otherwise inherited from the archetype

diff --git a/Engine/Source/Editor/PropertyEditor/Private/PropertyHandleImpl.cpp b/Engine/Source/Editor/PropertyEditor/Private/PropertyHandleImpl.cpp
--- a/Engine/Source/Editor/PropertyEditor/Private/PropertyHandleImpl.cpp
+++ b/Engine/Source/Editor/PropertyEditor/Private/PropertyHandleImpl.cpp
@@ -40,6 +40,9 @@
 #include "PropertyEditorModule.h"
 #include "Internationalization/TextPackageNamespaceUtil.h"
 #include "UObject/NonNullPropertyUtils.h"
+#include "UObject/PropertyPathName.h"
+#include "UObject/PropertyPathFunctions.h"
+#include "UObject/PropertyStateTracking.h"
 
 #define LOCTEXT_NAMESPACE "PropertyHandleImplementation"
 
@@ -493,6 +496,38 @@ FPropertyAccess::Result FPropertyValueImpl::ImportText( const TArray<FObjectBase
 			EPropertyPortFlags PortFlags = (Flags & EPropertyValueSetFlags::InstanceObjects) != 0 ? PPF_InstanceSubobjects : PPF_None;
 			FPropertyTextUtilities::TextToPropertyHelper(*NewValue, InPropertyNode, NodeProperty, Cur, PortFlags);
 
+			if (Flags & EPropertyValueSetFlags::ResetToDefault)
+            {
+            	if (CurObject)
+            	{
+            		UObject* Archetype = CurObject->GetArchetype();
+            		if (Archetype)
+            		{
+            			TSharedRef< FPropertyPath > PropertyPath = FPropertyNode::CreatePropertyPath(InPropertyNode->AsShared());
+            			
+            			const TArray<UE::FPropertyPathName> PathNames = PropertyPath->ToPropertyPathName();
+			            
+            			for (int32 i = 0; i < PathNames.Num(); i++)
+            			{
+            				const UE::FPropertyPathName& PathName = PathNames[i];
+			            
+            				const UE::FPropertyValueInContainer ArchetypeValue = UE::TryResolvePropertyPath(PathName, Archetype);
+            				const UE::FPropertyValueInContainer CurrObjValue = UE::TryResolvePropertyPath(PathName, CurObject);
+			            
+            				if (ArchetypeValue && CurrObjValue)
+            				{
+            					bool bWasSet = UE::FInitializedPropertyValueState(ArchetypeValue.Struct, ArchetypeValue.Container).IsSet(ArchetypeValue.Property, ArchetypeValue.ArrayIndex);
+            				 	UE::FInitializedPropertyValueState(CurrObjValue.Struct, CurrObjValue.Container).SetTo(bWasSet, CurrObjValue.Property, CurrObjValue.ArrayIndex);
+            				}
+            				else if (CurrObjValue)
+            				{
+            					UE::FInitializedPropertyValueState(CurrObjValue.Struct, CurrObjValue.Container).Clear(CurrObjValue.Property, CurrObjValue.ArrayIndex);
+            				}
+            			}
+            		}
+            	};
+            }
+
 			// If we are an FOptionalProperty, this check tells us whether we have changed between set/unset and 
 			// thus require a re-draw of our node in the details panel UI which is accomplished by having our parent rebuild us.
 			// 
diff --git a/Engine/Source/Runtime/CoreUObject/Private/UObject/Obj.cpp b/Engine/Source/Runtime/CoreUObject/Private/UObject/Obj.cpp
--- a/Engine/Source/Runtime/CoreUObject/Private/UObject/Obj.cpp
+++ b/Engine/Source/Runtime/CoreUObject/Private/UObject/Obj.cpp
@@ -476,11 +476,10 @@ void UObject::PostEditChangeProperty(FPropertyChangedEvent& PropertyChangedEvent
 	{
 		const EPropertyChangeType::Type SetPropertyValueInitializedStates =
 			EPropertyChangeType::ArrayAdd | EPropertyChangeType::ArrayRemove | EPropertyChangeType::ArrayClear | EPropertyChangeType::ArrayMove | EPropertyChangeType::ValueSet | EPropertyChangeType::Interactive;
-		if (PropertyChangedEvent.ChangeType & EPropertyChangeType::ResetToDefault)
-		{
-			UE::FInitializedPropertyValueState(this).Clear(Property);
-		}
-		else if (PropertyChangedEvent.ChangeType & SetPropertyValueInitializedStates)
+		const bool bIsResetToDefaultOperation = (PropertyChangedEvent.ChangeType & EPropertyChangeType::ResetToDefault) == EPropertyChangeType::ResetToDefault;
+		const bool bChangeTypeSetsInitState = (PropertyChangedEvent.ChangeType & SetPropertyValueInitializedStates) != 0;
+		
+		if (!bIsResetToDefaultOperation && bChangeTypeSetsInitState)
 		{
 			UE::FInitializedPropertyValueState(this).Set(Property);
 		}
