[Feature] EditorWorkflow tests for IDOs

#virtualized

diff --git a/Engine/Source/Editor/PropertyEditor/Private/SDetailTableRowBase.cpp b/Engine/Source/Editor/PropertyEditor/Private/SDetailTableRowBase.cpp
--- a/Engine/Source/Editor/PropertyEditor/Private/SDetailTableRowBase.cpp
+++ b/Engine/Source/Editor/PropertyEditor/Private/SDetailTableRowBase.cpp
@@ -113,7 +113,9 @@ void SDetailTableRowBase::PopulateContextMenu(UToolMenu* ToolMenu)
 				 NSLOCTEXT("PropertyView", "ExpandAll", "Expand All"),
 				NSLOCTEXT("PropertyView", "ExpandAll_ToolTip", "Expands this item and all children"),
 				FSlateIcon(),
-				ExpandAllAction);
+				ExpandAllAction,
+				EUserInterfaceActionType::Button,
+				TEXT("DetailsContextMenu_ExpandAll"));
 		}
 	}
 }
