Propagate outer object archetype flags to inner optional subobject archetypes when instanced at edit time in order to allow packages containing one or more instances of those subobjects to import them at save time.

diff --git a/Engine/Source/Editor/PropertyEditor/Private/PropertyHandleImpl.cpp b/Engine/Source/Editor/PropertyEditor/Private/PropertyHandleImpl.cpp
--- a/Engine/Source/Editor/PropertyEditor/Private/PropertyHandleImpl.cpp
+++ b/Engine/Source/Editor/PropertyEditor/Private/PropertyHandleImpl.cpp
@@ -5455,8 +5455,17 @@ FPropertyAccess::Result FPropertyHandleOptional::SetOptionalValue(FProperty* New
 				}
 				TopLevelObjects.Add(Outer);
 
+				UObject* UseOuter = (NewObjectClass->IsChildOf(UClass::StaticClass()) ? Cast<UClass>(Outer)->GetDefaultObject() : Outer);
+				EObjectFlags MaskedOuterFlags = UseOuter ? UseOuter->GetMaskedFlags(RF_PropagateToSubObjects) : RF_NoFlags;
+				if (UseOuter && UseOuter->HasAnyFlags(RF_ClassDefaultObject | RF_ArchetypeObject))
+				{
+					MaskedOuterFlags |= RF_ArchetypeObject;
+				}
+
 				// TODO: Construct object using `NewObjectClass`
-				UObject* Obj = NewObject<UObject>(Outer, NewObjectClass);
+				// Note: Using 'NAME_None' here for now since it looks like we don't support using the default value as a template for a new instance override.
+				// This means the new name may not match the template name if we start utilizing defaults here, which could result in archetype lookup failures.
+				UObject* Obj = NewObject<UObject>(UseOuter, NewObjectClass, NAME_None, MaskedOuterFlags);
 				if (Obj)
 				{
 					if (FObjectPropertyBase* ObjectProperty = CastField<FObjectPropertyBase>(ValueProperty))
