[PropertyEditor] Row animations correctly track

diff --git a/Engine/Source/Editor/PropertyEditor/Private/SDetailsViewBase.cpp b/Engine/Source/Editor/PropertyEditor/Private/SDetailsViewBase.cpp
--- a/Engine/Source/Editor/PropertyEditor/Private/SDetailsViewBase.cpp
+++ b/Engine/Source/Editor/PropertyEditor/Private/SDetailsViewBase.cpp
@@ -1670,7 +1670,7 @@ void SDetailsViewBase::MarkNodeAnimating(TSharedPtr<FPropertyNode> InNode, float
 		// Otherwise add it
 		if (AnimatingNode == nullptr)
 		{
-			AnimatingNode = &CurrentlyAnimatingNodeCollections.Add_GetRef({{NodePropertyPath}});
+			AnimatingNode = &CurrentlyAnimatingNodeCollections.Add_GetRef(FAnimatingNodeCollection({ NodePropertyPath }));
 			if (InAnimationBatchId.IsSet())
 			{
 				AnimatingNode->BatchId = InAnimationBatchId.GetValue();
diff --git a/Engine/Source/Editor/PropertyEditor/Private/SDetailsViewBase.h b/Engine/Source/Editor/PropertyEditor/Private/SDetailsViewBase.h
--- a/Engine/Source/Editor/PropertyEditor/Private/SDetailsViewBase.h
+++ b/Engine/Source/Editor/PropertyEditor/Private/SDetailsViewBase.h
@@ -208,12 +208,24 @@ protected:
 	/** Contains one or more nodes being animated by a single timer. */
 	struct FAnimatingNodeCollection
 	{
+		FAnimatingNodeCollection()
+			: BatchId(FGuid::NewGuid())
+		{
+		}
+
+		explicit FAnimatingNodeCollection(const TConstArrayView<TSharedPtr<FPropertyPath>>& InNodePaths, const FGuid& InBatchId = FGuid::NewGuid())
+			: NodePaths(InNodePaths)
+			, BatchId(InBatchId)
+		{
+		}
+
 		TArray<TSharedPtr<FPropertyPath>> NodePaths;
-		FTimerHandle NodeTimer;
 
-		/** Optionally identify by an id, where multiple nodes can be part of the same animation. */
+		/** Required id for tracking and removal (Property Nodes aren't reliable), which also allow multiple nodes to be part of the same animation. */
 		FGuid BatchId;
 
+		FTimerHandle NodeTimer;
+
 		/** Get validity of this AnimatingNode, determined by it's node path. */
 		bool IsValid() const;
 
@@ -224,9 +236,9 @@ protected:
 				return false;
 			}
 
-			if (BatchId.IsValid() && BatchId == Other.BatchId)
+			if (BatchId.IsValid())
 			{
-				return true;
+				return BatchId == Other.BatchId;
 			}
 
 			for (int32 Idx = 0; Idx < NodePaths.Num(); ++Idx)
@@ -239,8 +251,8 @@ protected:
 
 				if (!FPropertyPath::AreEqual(NodePaths[Idx].ToSharedRef(), Other.NodePaths[Idx].ToSharedRef()))
 				{
-					return false;					
-				}				
+					return false;
+				}
 			}
 			
 			return true;
