Fixed sub-categories having unlocalized display names

diff --git a/Engine/Source/Editor/PropertyEditor/Private/CategoryPropertyNode.cpp b/Engine/Source/Editor/PropertyEditor/Private/CategoryPropertyNode.cpp
--- a/Engine/Source/Editor/PropertyEditor/Private/CategoryPropertyNode.cpp
+++ b/Engine/Source/Editor/PropertyEditor/Private/CategoryPropertyNode.cpp
@@ -24,9 +24,24 @@ bool FCategoryPropertyNode::IsSubcategory() const
 
 FText FCategoryPropertyNode::GetDisplayName() const 
 {
-	FString SubcategoryName = GetSubcategoryName();
-	return FText::FromString( SubcategoryName );
-}	
+	FText DisplayName = FObjectEditorUtils::GetCategoryText(CategoryName);
+	if (IsSubcategory())
+	{
+		// The category name may actually contain a path of categories.  When displaying this category
+		// in the property window, we only want the leaf part of the path
+		const FString& CategoryPath = DisplayName.ToString();
+		if (int32 LastDelimiterCharIndex = INDEX_NONE;
+			CategoryPath.FindLastChar(FPropertyNodeConstants::CategoryDelimiterChar, LastDelimiterCharIndex))
+		{
+			// Grab the last sub-category from the path
+			FStringView SubCategoryName = FStringView(CategoryPath).Mid(LastDelimiterCharIndex + 1);
+			SubCategoryName.TrimStartAndEndInline();
+			DisplayName = FText::AsCultureInvariant(SubCategoryName);
+		}
+	}
+	return DisplayName;
+}
+
 /**
  * Overridden function for special setup
  */
@@ -136,27 +151,3 @@ bool FCategoryPropertyNode::GetQualifiedName( FString& PathPlusIndex, const bool
 
 	return bAddedAnything;
 }
-
-FString FCategoryPropertyNode::GetSubcategoryName() const 
-{
-	FString SubcategoryName;
-	{
-		// The category name may actually contain a path of categories.  When displaying this category
-		// in the property window, we only want the leaf part of the path
-		const FString& CategoryPath = GetCategoryName().ToString();
-		FString CategoryDelimiterString;
-		CategoryDelimiterString.AppendChar( FPropertyNodeConstants::CategoryDelimiterChar );  
-		const int32 LastDelimiterCharIndex = CategoryPath.Find( CategoryDelimiterString, ESearchCase::IgnoreCase, ESearchDir::FromEnd );
-		if( LastDelimiterCharIndex != INDEX_NONE )
-		{
-			// Grab the last sub-category from the path
-			SubcategoryName = CategoryPath.Mid( LastDelimiterCharIndex + 1 );
-		}
-		else
-		{
-			// No sub-categories, so just return the original (clean) category path
-			SubcategoryName = CategoryPath;
-		}
-	}
-	return SubcategoryName;
-}
diff --git a/Engine/Source/Editor/PropertyEditor/Private/CategoryPropertyNode.h b/Engine/Source/Editor/PropertyEditor/Private/CategoryPropertyNode.h
--- a/Engine/Source/Editor/PropertyEditor/Private/CategoryPropertyNode.h
+++ b/Engine/Source/Editor/PropertyEditor/Private/CategoryPropertyNode.h
@@ -47,7 +47,6 @@ protected:
 	 */
 	virtual bool GetQualifiedName( FString& PathPlusIndex, const bool bWithArrayIndex, const FPropertyNode* StopParent = nullptr, bool bIgnoreCategories = false ) const override;
 
-	FString GetSubcategoryName() const;
 	/**
 	 * Stored category name for the window
 	 */
