Allow a custom EditCondition to be treated as an additional condition rather than an override

This is controlled by a new ECustomEditConditionMode argument to EditCondition, and makes it easier to add additional EditCondition to a property via a detail customization, without stomping over any native EditCondition that was set on the property itself

diff --git a/Engine/Source/Editor/PropertyEditor/Private/DetailPropertyRow.cpp b/Engine/Source/Editor/PropertyEditor/Private/DetailPropertyRow.cpp
--- a/Engine/Source/Editor/PropertyEditor/Private/DetailPropertyRow.cpp
+++ b/Engine/Source/Editor/PropertyEditor/Private/DetailPropertyRow.cpp
@@ -154,10 +154,11 @@ IDetailPropertyRow& FDetailPropertyRow::ShowPropertyButtons( bool bInShowPropert
 	return *this;
 }
 
-IDetailPropertyRow& FDetailPropertyRow::EditCondition( TAttribute<bool> EditConditionValue, FOnBooleanValueChanged OnEditConditionValueChanged )
+IDetailPropertyRow& FDetailPropertyRow::EditCondition( TAttribute<bool> EditConditionValue, FOnBooleanValueChanged OnEditConditionValueChanged, ECustomEditConditionMode EditConditionMode )
 {
 	CustomEditConditionValue = EditConditionValue;
 	CustomEditConditionValueChanged = OnEditConditionValueChanged;
+	CustomEditConditionMode = EditConditionMode;
 	return *this;
 }
 
@@ -876,20 +877,28 @@ bool FDetailPropertyRow::GetEnabledState() const
 
 bool FDetailPropertyRow::IsEditConditionMet() const
 {
+	bool bIsEditConditionMet = true;
+
 	if (HasEditCondition())
 	{
 		if (CustomEditConditionValue.IsSet())
 		{
-			return CustomEditConditionValue.Get();
+			bIsEditConditionMet = bIsEditConditionMet && CustomEditConditionValue.Get();
+			
+			// In override mode, we ship checking the native property edit condition
+			if (CustomEditConditionMode == ECustomEditConditionMode::Override)
+			{
+				return bIsEditConditionMet;
+			}
 		}
-		else if (PropertyEditor.IsValid())
+
+		if (PropertyEditor.IsValid())
 		{
-			return PropertyEditor->IsEditConditionMet();
+			bIsEditConditionMet = bIsEditConditionMet && PropertyEditor->IsEditConditionMet();
 		}
 	}
 
-	// Default to True, matching the default output of FPropertyNode::IsEditConditionMet().
-	return true;
+	return bIsEditConditionMet;
 }
 
 bool FDetailPropertyRow::IsOnlyVisibleWhenEditConditionMet() const
@@ -939,10 +948,10 @@ static void ExecuteCustomEditConditionToggle(bool bValue, FOnBooleanValueChanged
 void FDetailPropertyRow::SetWidgetRowProperties(FDetailWidgetRow& Row) const
 {
 	// set edit condition handlers - use customized if provided
-	TAttribute<bool> EditConditionValue = CustomEditConditionValue;
-	if (!EditConditionValue.IsSet() && PropertyEditor.IsValid())
+	TAttribute<bool> EditConditionValue;
+	if (HasEditCondition())
 	{
-		EditConditionValue = TAttribute<bool>(PropertyEditor.ToSharedRef(), &FPropertyEditor::IsEditConditionMet);
+		EditConditionValue = TAttribute<bool>::CreateSP(this, &FDetailPropertyRow::IsEditConditionMet);
 	}
 
 	FOnBooleanValueChanged OnEditConditionValueChanged;
diff --git a/Engine/Source/Editor/PropertyEditor/Private/DetailPropertyRow.h b/Engine/Source/Editor/PropertyEditor/Private/DetailPropertyRow.h
--- a/Engine/Source/Editor/PropertyEditor/Private/DetailPropertyRow.h
+++ b/Engine/Source/Editor/PropertyEditor/Private/DetailPropertyRow.h
@@ -30,7 +30,7 @@ public:
 	virtual IDetailPropertyRow& DisplayName( const FText& InDisplayName ) override;
 	virtual IDetailPropertyRow& ToolTip( const FText& InToolTip ) override;
 	virtual IDetailPropertyRow& ShowPropertyButtons( bool bInShowPropertyButtons ) override;
-	virtual IDetailPropertyRow& EditCondition( TAttribute<bool> EditConditionValue, FOnBooleanValueChanged OnEditConditionValueChanged ) override;
+	virtual IDetailPropertyRow& EditCondition( TAttribute<bool> EditConditionValue, FOnBooleanValueChanged OnEditConditionValueChanged, ECustomEditConditionMode EditConditionMode = ECustomEditConditionMode::Override ) override;
 	virtual IDetailPropertyRow& EditConditionHides(bool bEditConditionHideValue) override;
 	virtual IDetailPropertyRow& IsEnabled(TAttribute<bool> InIsEnabled) override;
 	virtual IDetailPropertyRow& ShouldAutoExpand(bool bForceExpansion) override;
@@ -234,6 +234,8 @@ private:
 	TAttribute<bool> CustomEditConditionValue;
 	/** User customized edit condition change handler. */
 	FOnBooleanValueChanged CustomEditConditionValueChanged;
+	/** Should this custom edit condition override or merge with the existing edit condition of the property (if any)? */
+	ECustomEditConditionMode CustomEditConditionMode = ECustomEditConditionMode::Override;
 	/** User customized edit condition hides specifier */
 	bool bCustomEditConditionHides = false;
 	/** User customized reset to default */
diff --git a/Engine/Source/Editor/PropertyEditor/Public/IDetailPropertyRow.h b/Engine/Source/Editor/PropertyEditor/Public/IDetailPropertyRow.h
--- a/Engine/Source/Editor/PropertyEditor/Public/IDetailPropertyRow.h
+++ b/Engine/Source/Editor/PropertyEditor/Public/IDetailPropertyRow.h
@@ -15,6 +15,14 @@ class IDetailDragDropHandler;
 DECLARE_DELEGATE_RetVal_OneParam(bool, FIsResetToDefaultVisible, TSharedPtr<IPropertyHandle> /* PropertyHandle */);
 DECLARE_DELEGATE_OneParam(FResetToDefaultHandler, TSharedPtr<IPropertyHandle> /* PropertyHandle*/);
 
+enum class ECustomEditConditionMode : uint8
+{
+	/** The custom edit condition overrides the edit condition set directly on the property (if any). Only the custom edit condition must evaluate to true for the property to be editable. */
+	Override,
+	/** The custom edit condition merges with the existing edit condition set directly on the property (if any). Both edit conditions must evaluate to true for the property to be editable. */
+	Merge,
+};
+
 /**
  * Structure describing the delegates needed to override the behavior of reset to default in detail properties
  */
@@ -204,8 +212,9 @@ public:
 	 *
 	 * @param EditConditionValue			Attribute for the value of the edit condition (true indicates that the edit condition passes)
 	 * @param OnEditConditionValueChanged	Delegate called when the edit condition value changes
+	 * @param EditConditionMode				Should this custom edit condition override or merge with the existing edit condition of the property (if any)?
 	 */
-	virtual IDetailPropertyRow& EditCondition( TAttribute<bool> EditConditionValue, FOnBooleanValueChanged OnEditConditionValueChanged ) = 0;
+	virtual IDetailPropertyRow& EditCondition( TAttribute<bool> EditConditionValue, FOnBooleanValueChanged OnEditConditionValueChanged, ECustomEditConditionMode EditConditionMode = ECustomEditConditionMode::Override ) = 0;
 
 	/**
 	 * Sets whether or not the edit condition for this property should affect its visibility.  If the edit condition fails, the property will be hidden outright.
