Fix a regression that causes editinline subobject override names to no longer map to default subobject archetypes.

Notes:
- Introduced at CL [REDACTED_PHONE]).

diff --git a/Engine/Source/Editor/PropertyEditor/Private/UserInterface/PropertyEditor/SPropertyEditorEditInline.cpp b/Engine/Source/Editor/PropertyEditor/Private/UserInterface/PropertyEditor/SPropertyEditorEditInline.cpp
--- a/Engine/Source/Editor/PropertyEditor/Private/UserInterface/PropertyEditor/SPropertyEditorEditInline.cpp
+++ b/Engine/Source/Editor/PropertyEditor/Private/UserInterface/PropertyEditor/SPropertyEditorEditInline.cpp
@@ -374,26 +374,6 @@ void SPropertyEditorEditInline::OnClassPicked(UClass* InClass, TSharedRef<IPrope
 	{
 		GEditor->BeginTransaction(TEXT("PropertyEditor"), NSLOCTEXT("PropertyEditor", "OnClassPicked", "Set Class"), nullptr /* PropertyNode->GetProperty()) */ );
 
-		auto ExtractClassAndObjectNames = [](FStringView PathName, FStringView& ClassName, FStringView& ObjectName)
-		{
-			int32 ClassEnd;
-			PathName.FindChar(TCHAR('\''), ClassEnd);
-			if (ensure(ClassEnd != INDEX_NONE))
-			{
-				ClassName = PathName.Left(ClassEnd);
-			}
-
-			int32 LastPeriod, LastColon;
-			PathName.FindLastChar(TCHAR('.'), LastPeriod);
-			PathName.FindLastChar(TCHAR(':'), LastColon);
-			const int32 ObjectNameStart = FMath::Max(LastPeriod, LastColon);
-
-			if (ensure(ObjectNameStart != INDEX_NONE))
-			{
-				ObjectName = PathName.RightChop(ObjectNameStart + 1).LeftChop(1);
-			}
-		};
-
 		FString NewObjectName;
 		UObject* NewObjectTemplate = nullptr;
 
@@ -403,14 +383,16 @@ void SPropertyEditorEditInline::OnClassPicked(UClass* InClass, TSharedRef<IPrope
 			FString DefaultValue = PropertyNode->GetDefaultValueAsString(/*bUseDisplayName=*/false);
 			if (!DefaultValue.IsEmpty() && DefaultValue != FName(NAME_None).ToString())
 			{
-				FStringView ClassName, ObjectName;
-				ExtractClassAndObjectNames(DefaultValue, ClassName, ObjectName);
-				if (InClass->GetName() == ClassName)
+				FStringView ClassPath, ObjectPath;
+				if (FPackageName::ParseExportTextPath(DefaultValue, &ClassPath, &ObjectPath))
 				{
-					NewObjectName = ObjectName;
+					if (InClass->GetPathName() == ClassPath)
+					{
+						NewObjectName = FPackageName::ObjectPathToObjectName(ObjectPath);
 
-					ConstructorHelpers::StripObjectClass(DefaultValue);
-					NewObjectTemplate = StaticFindObject(InClass, nullptr, *DefaultValue);
+						ConstructorHelpers::StripObjectClass(DefaultValue);
+						NewObjectTemplate = StaticFindObject(InClass, nullptr, *DefaultValue);
+					}
 				}
 			}
 		}
@@ -428,13 +410,17 @@ void SPropertyEditorEditInline::OnClassPicked(UClass* InClass, TSharedRef<IPrope
 			FString NewValue;
 			if (InClass)
 			{
-				FStringView CurClassName, CurObjectName;
+				FStringView CurClassPath, CurObjectName;
 				if (PrevPerObjectValues.IsValidIndex(NewValues.Num()) && PrevPerObjectValues[NewValues.Num()] != FName(NAME_None).ToString())
 				{
-					ExtractClassAndObjectNames(PrevPerObjectValues[NewValues.Num()], CurClassName, CurObjectName);
+					FStringView CurObjectPath;
+					if (FPackageName::ParseExportTextPath(PrevPerObjectValues[NewValues.Num()], &CurClassPath, &CurObjectPath))
+					{
+						CurObjectName = FPackageName::ObjectPathToObjectName(CurObjectPath);
+					}
 				}
 
-				if (CurObjectName == NewObjectName && InClass->GetName() == CurClassName)
+				if (CurObjectName == NewObjectName && InClass->GetPathName() == CurClassPath)
 				{
 					NewValue = MoveTemp(PrevPerObjectValues[NewValues.Num()]);
 					PrevPerObjectValues[NewValues.Num()].Reset();
