Property Customization Helpers: support transient actors in the actor picker menu
- spawnable actors in sequencer are set as transient

#rnx

diff --git a/Engine/Source/Editor/PropertyEditor/Private/PropertyCustomizationHelpers.cpp b/Engine/Source/Editor/PropertyEditor/Private/PropertyCustomizationHelpers.cpp
--- a/Engine/Source/Editor/PropertyEditor/Private/PropertyCustomizationHelpers.cpp
+++ b/Engine/Source/Editor/PropertyEditor/Private/PropertyCustomizationHelpers.cpp
@@ -720,7 +720,7 @@ namespace PropertyCustomizationHelpers
 			.OnUseSelected(OnUseSelected);
 	}
 
-	TSharedRef<SWidget> MakeActorPickerWithMenu(AActor* const InitialActor, const bool AllowClear, const bool AllowPickingLevelInstanceContent, FOnShouldFilterActor ActorFilter, FOnActorSelected OnSet, FSimpleDelegate OnClose, FSimpleDelegate OnUseSelected, bool bDisplayUseSelected)
+	TSharedRef<SWidget> MakeActorPickerWithMenu(AActor* const InitialActor, const bool AllowClear, const bool AllowPickingLevelInstanceContent, FOnShouldFilterActor ActorFilter, FOnActorSelected OnSet, FSimpleDelegate OnClose, FSimpleDelegate OnUseSelected, bool bDisplayUseSelected, bool bShowTransient)
 	{
 		return
 			SNew(SPropertyMenuActorPicker)
@@ -731,7 +731,8 @@ namespace PropertyCustomizationHelpers
 			.OnSet(OnSet)
 			.OnClose(OnClose)
 			.OnUseSelected(OnUseSelected)
-			.DisplayUseSelected(bDisplayUseSelected);
+			.DisplayUseSelected(bDisplayUseSelected)
+			.ShowTransient(bShowTransient);
 	}
 
 	TSharedRef<SWidget> MakeComponentPickerWithMenu( UActorComponent* const InitialComponent, const bool AllowClear, FOnShouldFilterActor ActorFilter, FOnShouldFilterComponent ComponentFilter, FOnComponentSelected OnSet, FSimpleDelegate OnClose )
diff --git a/Engine/Source/Editor/PropertyEditor/Private/UserInterface/PropertyEditor/SPropertyMenuActorPicker.cpp b/Engine/Source/Editor/PropertyEditor/Private/UserInterface/PropertyEditor/SPropertyMenuActorPicker.cpp
--- a/Engine/Source/Editor/PropertyEditor/Private/UserInterface/PropertyEditor/SPropertyMenuActorPicker.cpp
+++ b/Engine/Source/Editor/PropertyEditor/Private/UserInterface/PropertyEditor/SPropertyMenuActorPicker.cpp
@@ -32,6 +32,7 @@ void SPropertyMenuActorPicker::Construct( const FArguments& InArgs )
 	OnClose = InArgs._OnClose;
 	OnUseSelected = InArgs._OnUseSelected;
 	bDisplayUseSelected = InArgs._DisplayUseSelected;
+	bShowTransient = InArgs._ShowTransient;
 
 	FMenuBuilder MenuBuilder(true, NULL);
 
@@ -92,6 +93,7 @@ void SPropertyMenuActorPicker::Construct( const FArguments& InArgs )
 		FSceneOutlinerInitializationOptions InitOptions;
 		InitOptions.Filters->AddFilterPredicate<FActorTreeItem>(ActorFilter);
 		InitOptions.bFocusSearchBoxWhenOpened = true;
+		InitOptions.bShowTransient = bShowTransient;
 
 		InitOptions.ColumnMap.Add(FSceneOutlinerBuiltInColumnTypes::Label(), FSceneOutlinerColumnInfo(ESceneOutlinerColumnVisibility::Visible, 0));
 		InitOptions.ColumnMap.Add(FSceneOutlinerBuiltInColumnTypes::ActorInfo(), FSceneOutlinerColumnInfo(ESceneOutlinerColumnVisibility::Visible, 10));
diff --git a/Engine/Source/Editor/PropertyEditor/Private/UserInterface/PropertyEditor/SPropertyMenuActorPicker.h b/Engine/Source/Editor/PropertyEditor/Private/UserInterface/PropertyEditor/SPropertyMenuActorPicker.h
--- a/Engine/Source/Editor/PropertyEditor/Private/UserInterface/PropertyEditor/SPropertyMenuActorPicker.h
+++ b/Engine/Source/Editor/PropertyEditor/Private/UserInterface/PropertyEditor/SPropertyMenuActorPicker.h
@@ -18,6 +18,7 @@ public:
 		, _AllowPickingLevelInstanceContent(false)
 		, _ActorFilter()
 		, _DisplayUseSelected(true)
+		, _ShowTransient(false)
 	{}
 		SLATE_ARGUMENT( AActor*, InitialActor )
 		SLATE_ARGUMENT( bool, AllowClear )
@@ -27,6 +28,7 @@ public:
 		SLATE_EVENT( FOnActorSelected, OnSet )
 		SLATE_EVENT( FSimpleDelegate, OnClose )
 		SLATE_EVENT( FSimpleDelegate, OnUseSelected )
+		SLATE_ARGUMENT(bool, ShowTransient)
 	SLATE_END_ARGS()
 
 	/**
@@ -103,4 +105,7 @@ private:
 
 	/* Whether to display the 'Use Selected' item*/
 	bool bDisplayUseSelected;
+
+	/** Whether to show transient objects */
+	bool bShowTransient = false;
 };
diff --git a/Engine/Source/Editor/PropertyEditor/Public/PropertyCustomizationHelpers.h b/Engine/Source/Editor/PropertyEditor/Public/PropertyCustomizationHelpers.h
--- a/Engine/Source/Editor/PropertyEditor/Public/PropertyCustomizationHelpers.h
+++ b/Engine/Source/Editor/PropertyEditor/Public/PropertyCustomizationHelpers.h
@@ -173,7 +173,7 @@ namespace PropertyCustomizationHelpers
 	PROPERTYEDITOR_API TSharedRef<SWidget> MakeAssetPickerWithMenu( const FAssetData& InitialObject, const bool AllowClear, const bool AllowCopyPaste, const TArray<const UClass*>& AllowedClasses, const TArray<const UClass*>& DisallowedClasses, const TArray<UFactory*>& NewAssetFactories, FOnShouldFilterAsset OnShouldFilterAsset, FOnAssetSelected OnSet, FSimpleDelegate OnClose, const TSharedPtr<IPropertyHandle>& PropertyHandle = TSharedPtr<IPropertyHandle>(), const TArray<FAssetData>& OwnerAssetArray = TArray<FAssetData>());
 	PROPERTYEDITOR_API TSharedRef<SWidget> MakeActorPickerAnchorButton( FOnGetActorFilters OnGetActorFilters, FOnActorSelected OnActorSelectedFromPicker );
 	PROPERTYEDITOR_API TSharedRef<SWidget> MakeActorPickerWithMenu( AActor* const InitialActor, const bool AllowClear, FOnShouldFilterActor ActorFilter, FOnActorSelected OnSet, FSimpleDelegate OnClose, FSimpleDelegate OnUseSelected );
-	PROPERTYEDITOR_API TSharedRef<SWidget> MakeActorPickerWithMenu(AActor* const InitialActor, const bool AllowClear, const bool AllowPickingLevelInstanceContent, FOnShouldFilterActor ActorFilter, FOnActorSelected OnSet, FSimpleDelegate OnClose, FSimpleDelegate OnUseSelected, bool bDisplayUseSelected=true);
+	PROPERTYEDITOR_API TSharedRef<SWidget> MakeActorPickerWithMenu(AActor* const InitialActor, const bool AllowClear, const bool AllowPickingLevelInstanceContent, FOnShouldFilterActor ActorFilter, FOnActorSelected OnSet, FSimpleDelegate OnClose, FSimpleDelegate OnUseSelected, bool bDisplayUseSelected=true, bool bShowTransient=false);
 	PROPERTYEDITOR_API TSharedRef<SWidget> MakeComponentPickerWithMenu( UActorComponent* const InitialComponent, const bool AllowClear, FOnShouldFilterActor ActorFilter, FOnShouldFilterComponent ComponentFilter, FOnComponentSelected OnSet, FSimpleDelegate OnClose );
 	PROPERTYEDITOR_API TSharedRef<SWidget> MakeInteractiveActorPicker(FOnGetAllowedClasses OnGetAllowedClasses, FOnShouldFilterActor OnShouldFilterActor, FOnActorSelected OnActorSelectedFromPicker);
 	PROPERTYEDITOR_API TSharedRef<SWidget> MakeSceneDepthPicker(FOnSceneDepthLocationSelected OnSceneDepthLocationSelected);
