Fix copy paste of array of stings with spaces
#rnx

diff --git a/Engine/Source/Editor/PropertyEditor/Private/PropertyHandleImpl.cpp b/Engine/Source/Editor/PropertyEditor/Private/PropertyHandleImpl.cpp
--- a/Engine/Source/Editor/PropertyEditor/Private/PropertyHandleImpl.cpp
+++ b/Engine/Source/Editor/PropertyEditor/Private/PropertyHandleImpl.cpp
@@ -5299,7 +5299,7 @@ FPropertyAccess::Result FPropertyHandleArray::SetValueFromFormattedString(const
 		TArray<FString> Values;
 
 		// Remove firsts and lasts parenthesis, and split the string 
-		InValue.LeftChop(1).RightChop(1).ParseIntoArrayWS(Values, TEXT(","), true);
+		InValue.LeftChop(1).RightChop(1).ParseIntoArray(Values, TEXT(","), true);
 
 		for (const FString& Value : Values)
 		{
@@ -5962,7 +5962,7 @@ FPropertyAccess::Result FPropertyHandleSet::SetValueFromFormattedString(const FS
 		TArray<FString> Values;
 
 		// Remove firsts and lasts parenthesis  
-		InValue.LeftChop(1).RightChop(1).ParseIntoArrayWS(Values, TEXT(","), true);
+		InValue.LeftChop(1).RightChop(1).ParseIntoArray(Values, TEXT(","), true);
 		for (const FString& Value : Values)
 		{
 			AddItem();
