New editor subsystem for handling property visibility override.

diff --git a/Engine/Source/Editor/PropertyEditor/Private/DetailItemNode.cpp b/Engine/Source/Editor/PropertyEditor/Private/DetailItemNode.cpp
--- a/Engine/Source/Editor/PropertyEditor/Private/DetailItemNode.cpp
+++ b/Engine/Source/Editor/PropertyEditor/Private/DetailItemNode.cpp
@@ -1,6 +1,7 @@
 // Copyright Epic Games, Inc. All Rights Reserved.
 
 #include "DetailItemNode.h"
+
 #include "CategoryPropertyNode.h"
 #include "DetailCategoryGroupNode.h"
 #include "DetailGroup.h"
@@ -8,10 +9,11 @@
 #include "IDetailKeyframeHandler.h"
 #include "ObjectPropertyNode.h"
 #include "PropertyHandleImpl.h"
-#include "SConstrainedBox.h"
 #include "PropertyPermissionList.h"
+#include "SConstrainedBox.h"
 #include "SDetailCategoryTableRow.h"
 #include "SDetailSingleItemRow.h"
+#include "Subsystems/PropertyVisibilityOverrideSubsystem.h"
 #include "UObject/PropertyOptional.h"
 
 static const FName NAME_IsLooseMetadata = TEXT("IsLoose");
@@ -1010,7 +1012,7 @@ void FDetailItemNode::FilterNode(const FDetailFilter& InFilter)
 	}
 
 	// set bForceHidden if this node is loose and loose properties are hidden
-	if( TSharedPtr<FPropertyNode> PropertyNodePin = Customization.GetPropertyNode())
+	if (TSharedPtr<FPropertyNode> PropertyNodePin = Customization.GetPropertyNode())
 	{
 		if (LIKELY(!InFilter.bShowLooseProperties))
 		{
@@ -1022,6 +1024,7 @@ void FDetailItemNode::FilterNode(const FDetailFilter& InFilter)
 				}
 			}
 		}
+
 		if (!bForceHidden && InFilter.ShouldForceHideProperty.IsBound())
 		{
 			if (InFilter.ShouldForceHideProperty.Execute(PropertyNodePin.ToSharedRef()))
@@ -1029,8 +1032,22 @@ void FDetailItemNode::FilterNode(const FDetailFilter& InFilter)
 				bForceHidden = true;
 			}
 		}
-	}
 
+		if (!bForceHidden)
+		{
+			if (FProperty* Property = PropertyNodePin->GetProperty())
+			{
+				static const FName Name_PropertyVisibilityOverride("PropertyVisibilityOverride"); // TODO: move name definition to centralized place
+				if (Property->GetBoolMetaData(Name_PropertyVisibilityOverride))
+				{
+					if (UPropertyVisibilityOverrideSubsystem* PropertyVisibilityOverrideSubsystem = UPropertyVisibilityOverrideSubsystem::Get())
+					{
+						bForceHidden = PropertyVisibilityOverrideSubsystem->ShouldHideProperty(Property);
+					}
+				}
+			}
+		}
+	}
 
 	bShouldBeVisibleDueToChildFiltering = false;
 
diff --git a/Engine/Source/Editor/UnrealEd/Private/Subsystems/PropertyVisibilityOverrideSubsystem.cpp b/Engine/Source/Editor/UnrealEd/Private/Subsystems/PropertyVisibilityOverrideSubsystem.cpp
new file mode 100644
--- /dev/null
+++ b/Engine/Source/Editor/UnrealEd/Private/Subsystems/PropertyVisibilityOverrideSubsystem.cpp
@@ -0,0 +1,23 @@
+// Copyright Epic Games, Inc. All Rights Reserved.
+
+#include "Subsystems/PropertyVisibilityOverrideSubsystem.h"
+
+#include "Editor.h"
+#include "GameFramework/Actor.h"
+
+UPropertyVisibilityOverrideSubsystem* UPropertyVisibilityOverrideSubsystem::Get()
+{
+	if (GEditor)
+	{
+		return GEditor->GetEditorSubsystem<UPropertyVisibilityOverrideSubsystem>();
+	}
+
+	return nullptr;
+}
+
+bool UPropertyVisibilityOverrideSubsystem::ShouldHideProperty(const FProperty* Property) const
+{
+	return false;
+}
+
+
diff --git a/Engine/Source/Editor/UnrealEd/Public/Subsystems/PropertyVisibilityOverrideSubsystem.h b/Engine/Source/Editor/UnrealEd/Public/Subsystems/PropertyVisibilityOverrideSubsystem.h
new file mode 100644
--- /dev/null
+++ b/Engine/Source/Editor/UnrealEd/Public/Subsystems/PropertyVisibilityOverrideSubsystem.h
@@ -0,0 +1,23 @@
+// Copyright Epic Games, Inc. All Rights Reserved.
+
+#pragma once
+
+#include "EditorSubsystem.h"
+
+#include "PropertyVisibilityOverrideSubsystem.generated.h"
+
+class FProperty;
+class UObject;
+
+UCLASS(MinimalAPI)
+class UPropertyVisibilityOverrideSubsystem : public UEditorSubsystem
+{
+	GENERATED_BODY()
+
+public:
+	UNREALED_API static UPropertyVisibilityOverrideSubsystem* Get();
+
+	virtual bool ShouldCreateSubsystem(UObject* Outer) const override { return false; }
+
+	UNREALED_API virtual bool ShouldHideProperty(const FProperty* Property) const;
+};
