New files for FNames that should have a single place of definition.

diff --git a/Engine/Source/Editor/PropertyEditor/Private/DetailItemNode.cpp b/Engine/Source/Editor/PropertyEditor/Private/DetailItemNode.cpp
--- a/Engine/Source/Editor/PropertyEditor/Private/DetailItemNode.cpp
+++ b/Engine/Source/Editor/PropertyEditor/Private/DetailItemNode.cpp
@@ -14,6 +14,7 @@
 #include "SDetailCategoryTableRow.h"
 #include "SDetailSingleItemRow.h"
 #include "Subsystems/PropertyVisibilityOverrideSubsystem.h"
+#include "UObject/PropertyNames.h"
 #include "UObject/PropertyOptional.h"
 
 static const FName NAME_IsLooseMetadata = TEXT("IsLoose");
@@ -1037,8 +1038,7 @@ void FDetailItemNode::FilterNode(const FDetailFilter& InFilter)
 		{
 			if (FProperty* Property = PropertyNodePin->GetProperty())
 			{
-				static const FName Name_PropertyVisibilityOverride("PropertyVisibilityOverride"); // TODO: move name definition to centralized place
-				if (Property->GetBoolMetaData(Name_PropertyVisibilityOverride))
+				if (Property->GetBoolMetaData(PropertyNames::PropertyVisibilityOverrideName))
 				{
 					if (UPropertyVisibilityOverrideSubsystem* PropertyVisibilityOverrideSubsystem = UPropertyVisibilityOverrideSubsystem::Get())
 					{
diff --git a/Engine/Source/Runtime/CoreUObject/Private/UObject/PropertyNames.cpp b/Engine/Source/Runtime/CoreUObject/Private/UObject/PropertyNames.cpp
new file mode 100644
--- /dev/null
+++ b/Engine/Source/Runtime/CoreUObject/Private/UObject/PropertyNames.cpp
@@ -0,0 +1,5 @@
+// Copyright Epic Games, Inc. All Rights Reserved.
+
+#include "UObject/PropertyNames.h"
+
+extern const FName PropertyNames::PropertyVisibilityOverrideName(TEXT("PropertyVisibilityOverride"));
diff --git a/Engine/Source/Runtime/CoreUObject/Public/UObject/PropertyNames.h b/Engine/Source/Runtime/CoreUObject/Public/UObject/PropertyNames.h
new file mode 100644
--- /dev/null
+++ b/Engine/Source/Runtime/CoreUObject/Public/UObject/PropertyNames.h
@@ -0,0 +1,9 @@
+// Copyright Epic Games, Inc. All Rights Reserved.
+#pragma once
+
+#include "UObject/NameTypes.h"
+
+namespace PropertyNames
+{
+	COREUOBJECT_API extern const FName PropertyVisibilityOverrideName;
+}
