Optional property: Fix node collapsing on rebuild
#rnx

diff --git a/Engine/Source/Editor/PropertyEditor/Private/ItemPropertyNode.cpp b/Engine/Source/Editor/PropertyEditor/Private/ItemPropertyNode.cpp
--- a/Engine/Source/Editor/PropertyEditor/Private/ItemPropertyNode.cpp
+++ b/Engine/Source/Editor/PropertyEditor/Private/ItemPropertyNode.cpp
@@ -126,6 +126,7 @@ TSharedPtr<FPropertyNode>& FItemPropertyNode::GetOrCreateOptionalValueNode(void)
 					InitParams.bCreateDisableEditOnInstanceNodes = !!HasNodeFlags(EPropertyNodeFlags::ShouldShowDisableEditOnInstance);
 
 					OptionalValueNode->InitNode(InitParams);
+					AddChildNode(OptionalValueNode);
 				}
 			}
 		}
