Property Editor: Standalone Widgets now use visibility attribute.
- The GenerateStandaloneWidget method now copies the visiblity attribute when using a customisation.

diff --git a/Engine/Source/Editor/PropertyEditor/Private/DetailItemNode.cpp b/Engine/Source/Editor/PropertyEditor/Private/DetailItemNode.cpp
--- a/Engine/Source/Editor/PropertyEditor/Private/DetailItemNode.cpp
+++ b/Engine/Source/Editor/PropertyEditor/Private/DetailItemNode.cpp
@@ -370,6 +370,7 @@ bool FDetailItemNode::GenerateStandaloneWidget(FDetailWidgetRow& OutRow) const
 
 		OutRow.CustomResetToDefault = Row.CustomResetToDefault;
 		OutRow.IsEnabledAttr = Row.IsEnabledAttr;
+		OutRow.VisibilityAttr = Row.VisibilityAttr;
 		OutRow.EditConditionValue = Row.EditConditionValue;
 		OutRow.OnEditConditionValueChanged = Row.OnEditConditionValueChanged;
 
