[PropertyHandle] added validation for EditConst in CanResetToDefault to match ResetToDefault implementation

diff --git a/Engine/Source/Editor/PropertyEditor/Private/PropertyHandleImpl.cpp b/Engine/Source/Editor/PropertyEditor/Private/PropertyHandleImpl.cpp
--- a/Engine/Source/Editor/PropertyEditor/Private/PropertyHandleImpl.cpp
+++ b/Engine/Source/Editor/PropertyEditor/Private/PropertyHandleImpl.cpp
@@ -3603,7 +3603,7 @@ bool FPropertyHandleBase::CanResetToDefault() const
 
 	// Should not be able to reset fixed size arrays
 	const bool bFixedSized = (Property->PropertyFlags & CPF_EditFixedSize) != 0;
-	const bool bCanResetToDefault = (Property->PropertyFlags & CPF_Config) == 0;
+	const bool bCanResetToDefault = (Property->PropertyFlags & CPF_Config) == 0 && !IsEditConst();
 
 	return bCanResetToDefault && !bFixedSized && DiffersFromDefault();
 }
