Expand container elements on add in the details panel
* Legacy behaviour can be reverted using `DetailsPanel.ExpandAddedItems=0`
* FPropertyHandleItemAddResult created to be a tuple of legacy FPropertyAccess::Result and the newly passed logical index of the added item.  Conversion operators retain compilation compatibility.

diff --git a/Engine/Source/Editor/PropertyEditor/Private/Presentation/PropertyEditor/PropertyEditor.cpp b/Engine/Source/Editor/PropertyEditor/Private/Presentation/PropertyEditor/PropertyEditor.cpp
--- a/Engine/Source/Editor/PropertyEditor/Private/Presentation/PropertyEditor/PropertyEditor.cpp
+++ b/Engine/Source/Editor/PropertyEditor/Private/Presentation/PropertyEditor/PropertyEditor.cpp
@@ -31,6 +31,12 @@ DEFINE_LOG_CATEGORY_STATIC(LogPropertyEditor, Log, All);
 
 namespace UE::Private
 {
+	static FAutoConsoleVariable CVarExpandAddedItem(
+		TEXT("DetailsPanel.ExpandAddedItems"),
+		true,
+		TEXT("Automatically expands items added to container in the details panel"),
+		ECVF_Default);
+	
 	static UClass* GetMetaClass(const FProperty* ForProperty)
 	{
 		if(!ForProperty)
@@ -227,22 +233,33 @@ void FPropertyEditor::OnAddItem()
 
 	check(ArrayHandle.IsValid() || SetHandle.IsValid() || MapHandle.IsValid());
 
+	FPropertyHandleItemAddResult Result;
+
 	if (ArrayHandle.IsValid())
 	{
-		ArrayHandle->AddItem();
+		Result = ArrayHandle->AddItem();
 	}
 	else if (SetHandle.IsValid())
 	{
-		SetHandle->AddItem();
+		Result = SetHandle->AddItem();
 	}
 	else if (MapHandle.IsValid())
 	{
-		MapHandle->AddItem();
+		Result = MapHandle->AddItem();
 	}
 
 	// Expand containers when an item is added to them
 	PropertyNode->SetNodeFlags(EPropertyNodeFlags::Expanded, true);
 
+	if (Result.GetAccessResult() == FPropertyAccess::Success && UE::Private::CVarExpandAddedItem->GetBool())
+	{
+		TSharedPtr<FPropertyNode>& ChildNode = PropertyNode->GetChildNode(Result.GetIndex());
+		if (ChildNode)
+		{
+			ChildNode->SetNodeFlags(EPropertyNodeFlags::Expanded, true);
+		}
+	}
+	
 	//In case the property is show in the favorite category refresh the whole tree
 	if (PropertyNode->IsFavorite())
 	{
diff --git a/Engine/Source/Editor/PropertyEditor/Private/PropertyHandleImpl.cpp b/Engine/Source/Editor/PropertyEditor/Private/PropertyHandleImpl.cpp
--- a/Engine/Source/Editor/PropertyEditor/Private/PropertyHandleImpl.cpp
+++ b/Engine/Source/Editor/PropertyEditor/Private/PropertyHandleImpl.cpp
@@ -5180,17 +5180,18 @@ bool FPropertyHandleArray::Supports( TSharedRef<FPropertyNode> PropertyNode )
 }
 
 
-FPropertyAccess::Result FPropertyHandleArray::AddItem()
+FPropertyHandleItemAddResult FPropertyHandleArray::AddItem()
 {
 	FPropertyAccess::Result Result = FPropertyAccess::Fail;
+	int32 Index = INDEX_NONE;
 	if( IsEditable() )
 	{
-		Implementation->AddChild();
+		Index = Implementation->AddChild();
 		Implementation->GetPropertyNode()->RebuildChildren();
 		Result = FPropertyAccess::Success;
+		return FPropertyHandleItemAddResult(Index, Result);
 	}
-
-	return Result;
+	return FPropertyHandleItemAddResult();
 }
 
 FPropertyAccess::Result FPropertyHandleArray::EmptyArray()
@@ -5779,7 +5780,7 @@ bool FPropertyHandleSet::HasDefaultElement()
 	return false;
 }
 
-FPropertyAccess::Result FPropertyHandleSet::AddItem()
+FPropertyHandleItemAddResult FPropertyHandleSet::AddItem()
 {
 	FPropertyAccess::Result Result = FPropertyAccess::Fail;
 	if (IsEditable())
@@ -5869,7 +5870,7 @@ FPropertyAccess::Result FPropertyHandleSet::AddItem()
 									// Set the new entry to have the element value.
 									Implementation->ImportText(Enum->GetNameStringByIndex(EnumIndex), ChildNode.Get(), EPropertyValueSetFlags::DefaultFlags);
 									Result = FPropertyAccess::Success;
-									break;
+									return FPropertyHandleItemAddResult(ChildNodeIndex, Result);
 								}
 							}
 						}
@@ -5879,9 +5880,10 @@ FPropertyAccess::Result FPropertyHandleSet::AddItem()
 				{
 					if (!HasDefaultElement())
 					{
-						Implementation->AddChild();
+						int32 Index = Implementation->AddChild();
 						Implementation->GetPropertyNode()->RebuildChildren();
 						Result = FPropertyAccess::Success;
+						return FPropertyHandleItemAddResult(Index, Result);
 					}
 				}
 			}
@@ -5893,7 +5895,7 @@ FPropertyAccess::Result FPropertyHandleSet::AddItem()
 		Implementation->ShowInvalidOperationError(LOCTEXT("DuplicateSetElement_Add", "Cannot add a new element to the set while an element with the same value exists"));
 	}
 
-	return Result;
+	return FPropertyHandleItemAddResult();
 }
 
 FPropertyAccess::Result FPropertyHandleSet::Empty()
@@ -6026,7 +6028,7 @@ bool FPropertyHandleMap::HasDefaultKey()
 	return false;
 }
 
-FPropertyAccess::Result FPropertyHandleMap::AddItem()
+FPropertyHandleItemAddResult FPropertyHandleMap::AddItem()
 {
 	FPropertyAccess::Result Result = FPropertyAccess::Fail;
 	if (IsEditable())
@@ -6120,7 +6122,7 @@ FPropertyAccess::Result FPropertyHandleMap::AddItem()
 										// Set the new entry to have the key value.
 										Implementation->ImportText(Enum->GetNameStringByIndex(EnumIndex), ChildKeyNode.Get(), EPropertyValueSetFlags::DefaultFlags);
 										Result = FPropertyAccess::Success;
-										break;
+										return FPropertyHandleItemAddResult(ChildNodeIndex, Result);
 									}
 								}
 							}
@@ -6131,8 +6133,9 @@ FPropertyAccess::Result FPropertyHandleMap::AddItem()
 				{
 					if (!HasDefaultKey())
 					{
-						Implementation->AddChild();
+						const int32 Index = Implementation->AddChild();
 						Result = FPropertyAccess::Success;
+						return FPropertyHandleItemAddResult(Index, Result);
 					}
 				}
 			}
@@ -6144,7 +6147,7 @@ FPropertyAccess::Result FPropertyHandleMap::AddItem()
 		Implementation->ShowInvalidOperationError(LOCTEXT("DuplicateMapKey_Add", "Cannot add a new key to the map while a key with the same value exists"));
 	}
 
-	return Result;
+	return FPropertyHandleItemAddResult();
 }
 
 FPropertyAccess::Result FPropertyHandleMap::Empty()
diff --git a/Engine/Source/Editor/PropertyEditor/Private/PropertyHandleImpl.h b/Engine/Source/Editor/PropertyEditor/Private/PropertyHandleImpl.h
--- a/Engine/Source/Editor/PropertyEditor/Private/PropertyHandleImpl.h
+++ b/Engine/Source/Editor/PropertyEditor/Private/PropertyHandleImpl.h
@@ -717,7 +717,7 @@ public:
 	FPropertyHandleArray( TSharedRef<FPropertyNode> PropertyNode, FNotifyHook* NotifyHook, TSharedPtr<IPropertyUtilities> PropertyUtilities );
 	static bool Supports( TSharedRef<FPropertyNode> PropertyNode );
 	/** IPropertyHandleArray interface */
-	virtual FPropertyAccess::Result AddItem() override;
+	virtual FPropertyHandleItemAddResult AddItem() override;
 	virtual FPropertyAccess::Result EmptyArray() override;
 	virtual FPropertyAccess::Result Insert( int32 Index ) override;
 	virtual FPropertyAccess::Result DuplicateItem( int32 Index ) override;
@@ -771,7 +771,7 @@ public:
 
 	/** IPropertyHandleSet Interface */
 	virtual bool HasDefaultElement() override;
-	virtual FPropertyAccess::Result AddItem() override;
+	virtual FPropertyHandleItemAddResult AddItem() override;
 	virtual FPropertyAccess::Result Empty() override;
 	virtual FPropertyAccess::Result DeleteItem(int32 Index) override;
 	virtual FPropertyAccess::Result GetNumElements(uint32& OutNumElements) override;
@@ -795,7 +795,7 @@ public:
 
 	/** IPropertyHandleMap Interface */
 	virtual bool HasDefaultKey() override;
-	virtual FPropertyAccess::Result AddItem() override;
+	virtual FPropertyHandleItemAddResult AddItem() override;
 	virtual FPropertyAccess::Result Empty() override;
 	virtual FPropertyAccess::Result DeleteItem(int32 Index) override;
 	virtual FPropertyAccess::Result GetNumElements(uint32& OutNumElements) override;
diff --git a/Engine/Source/Editor/PropertyEditor/Public/PropertyHandle.h b/Engine/Source/Editor/PropertyEditor/Public/PropertyHandle.h
--- a/Engine/Source/Editor/PropertyEditor/Public/PropertyHandle.h
+++ b/Engine/Source/Editor/PropertyEditor/Public/PropertyHandle.h
@@ -784,6 +784,25 @@ public:
 	virtual bool IsCategoryHandle() const = 0;
 };
 
+class FPropertyHandleItemAddResult
+{
+public:
+	FPropertyHandleItemAddResult() = default;
+	FPropertyHandleItemAddResult(int32 InIndex, FPropertyAccess::Result InResult)
+		: Index(InIndex)
+		, PropertyAccessResult(InResult)
+	{}
+
+	int32 GetIndex() const { return Index; }
+	FPropertyAccess::Result GetAccessResult() const { return PropertyAccessResult; }
+
+	operator FPropertyAccess::Result() const { return PropertyAccessResult; }
+	
+private:
+	int32 Index = INDEX_NONE;
+	FPropertyAccess::Result PropertyAccessResult = FPropertyAccess::Fail;
+};
+
 /**
  * A handle to an array property which allows you to manipulate the array                                                              
  */
@@ -796,7 +815,7 @@ public:
 	 * Adds an item to the end of the array
 	 * @return Whether or not this was successful
 	 */
-	virtual FPropertyAccess::Result AddItem() = 0;
+	virtual FPropertyHandleItemAddResult AddItem() = 0;
 
 	/**
 	 * Empty the array
@@ -873,7 +892,7 @@ public:
 	 * Adds an item to the set.
 	 * @return Whether or not this was successful
 	 */
-	virtual FPropertyAccess::Result AddItem() = 0;
+	virtual FPropertyHandleItemAddResult AddItem() = 0;
 
 	/**
 	 * Empties the set
@@ -926,7 +945,7 @@ public:
 	 * Adds an item to the map.
 	 * @return Whether or not this was successful
 	 */
-	virtual FPropertyAccess::Result AddItem() = 0;
+	virtual FPropertyHandleItemAddResult AddItem() = 0;
 
 	/**
 	 * Empties the map
