[PropertyEditor] Row animations correctly resume on recreation (rather than restarting)

diff --git a/Engine/Source/Editor/PropertyEditor/Private/IDetailsViewPrivate.h b/Engine/Source/Editor/PropertyEditor/Private/IDetailsViewPrivate.h
--- a/Engine/Source/Editor/PropertyEditor/Private/IDetailsViewPrivate.h
+++ b/Engine/Source/Editor/PropertyEditor/Private/IDetailsViewPrivate.h
@@ -166,6 +166,9 @@ public:
 	/** @return true if node is animating */
 	virtual bool IsNodeAnimating(TSharedPtr<FPropertyNode> InNode) = 0;
 
+	/** @return the time/progress of the given node (time since creation/play), if it's animating. The time is in seconds. */
+	virtual float GetNodeAnimationTime(TSharedPtr<FPropertyNode> InNode) = 0;
+
 	/** Column width accessibility */
 	virtual class FDetailColumnSizeData& GetColumnSizeData() = 0;
 
diff --git a/Engine/Source/Editor/PropertyEditor/Private/SDetailSingleItemRow.cpp b/Engine/Source/Editor/PropertyEditor/Private/SDetailSingleItemRow.cpp
--- a/Engine/Source/Editor/PropertyEditor/Private/SDetailSingleItemRow.cpp
+++ b/Engine/Source/Editor/PropertyEditor/Private/SDetailSingleItemRow.cpp
@@ -313,9 +313,12 @@ void SDetailSingleItemRow::Construct( const FArguments& InArgs, FDetailLayoutCus
 	PulseAnimation.AddCurve(0.0f, UE::PropertyEditor::Private::PulseAnimationLength, ECurveEaseFunction::CubicInOut);
 
 	// Play on construction if animation was started from a behavior the re-constructs this widget
-	if (DetailsView && DetailsView->IsNodeAnimating(GetPropertyNode()))
+	const TSharedPtr<FPropertyNode> AssociatedPropertyNode = GetPropertyNode();
+	if (DetailsView && DetailsView->IsNodeAnimating(AssociatedPropertyNode))
 	{
-		PulseAnimation.Play(SharedThis(this));
+		// Resume from the stored animation time if it's still animating
+		float AnimationTime = DetailsView->GetNodeAnimationTime(AssociatedPropertyNode);
+		PulseAnimation.Play(SharedThis(this), false, AnimationTime);
 	}
 
 	TSharedPtr<FDetailCategoryImpl> Category = InOwnerTreeNode->GetParentCategory();
diff --git a/Engine/Source/Editor/PropertyEditor/Private/SDetailsViewBase.cpp b/Engine/Source/Editor/PropertyEditor/Private/SDetailsViewBase.cpp
--- a/Engine/Source/Editor/PropertyEditor/Private/SDetailsViewBase.cpp
+++ b/Engine/Source/Editor/PropertyEditor/Private/SDetailsViewBase.cpp
@@ -1347,6 +1347,52 @@ void SDetailsViewBase::Tick( const FGeometry& AllottedGeometry, const double InC
 	}
 }
 
+SDetailsViewBase::FAnimatingNodeCollection::FAnimatingNodeCollection(
+	const TConstArrayView<TSharedPtr<FPropertyPath>>& InNodePaths,
+	const FGuid& InBatchId)
+	: NodePaths(InNodePaths)
+	, BatchId((InBatchId))
+	, CreationTime(FPlatformTime::Seconds())
+{
+	if (!InBatchId.IsValid())
+	{
+		FGuid CombinedBatchId;
+		for (const TSharedPtr<FPropertyPath>& NodePath : NodePaths)
+		{
+			if (NodePath.IsValid())
+			{
+				CombinedBatchId = FGuid::Combine(CombinedBatchId, GetIdForNodePath(NodePath));
+			}
+		}
+
+		BatchId = CombinedBatchId;
+	}
+}
+
+FGuid SDetailsViewBase::FAnimatingNodeCollection::GetIdForNodePath(const TSharedPtr<FPropertyPath>& InNodePath)
+{
+	return InNodePath.IsValid() ? FGuid::NewDeterministicGuid(InNodePath->ToString()) : FGuid();
+}
+
+bool SDetailsViewBase::FAnimatingNodeCollection::IsValid() const
+{
+	return Algo::AllOf(NodePaths,
+		[](const TSharedPtr<FPropertyPath>& InNodePath)
+		{
+			return InNodePath.IsValid();
+		});
+}
+
+void SDetailsViewBase::FAnimatingNodeCollection::ResetElapsedSeconds()
+{
+	CreationTime = FPlatformTime::Seconds();
+}
+
+float SDetailsViewBase::FAnimatingNodeCollection::GetElapsedSeconds() const
+{
+	return FPlatformTime::Seconds() - CreationTime;
+}
+
 /**
 * Recursively gets expanded items for a node
 *
@@ -1630,14 +1676,7 @@ void SDetailsViewBase::MarkNodeAnimating(TSharedPtr<FPropertyNode> InNode, float
 		TSharedPtr<FPropertyPath> NodePropertyPath = FPropertyNode::CreatePropertyPath(InNode.ToSharedRef());
 		
 		// Try finding existing, ignore Batch Id
-		FAnimatingNodeCollection* AnimatingNode = CurrentlyAnimatingNodeCollections.FindByPredicate(
-			[NodePropertyPath](const FAnimatingNodeCollection& AnimatingNode)
-			{
-				return Algo::AnyOf(AnimatingNode.NodePaths, [NodePropertyPath](const TSharedPtr<FPropertyPath>& NodePath)
-				{
-					return FPropertyPath::AreEqual(NodePath.ToSharedRef(), NodePropertyPath.ToSharedRef());
-				});
-			});
+		FAnimatingNodeCollection* AnimatingNode = GetAnimatingNodeCollectionForPropertyNode(InNode);
 
 		// If this node is part of a batch animation, let it play out (don't reset it)
 		if (AnimatingNode)
@@ -1648,7 +1687,7 @@ void SDetailsViewBase::MarkNodeAnimating(TSharedPtr<FPropertyNode> InNode, float
 				return;
 			}
 		}
-		
+
 		// If not found, but batch id found, add it to that batch
 		if (!AnimatingNode && InAnimationBatchId.IsSet())
 		{
@@ -1676,8 +1715,13 @@ void SDetailsViewBase::MarkNodeAnimating(TSharedPtr<FPropertyNode> InNode, float
 				AnimatingNode->BatchId = InAnimationBatchId.GetValue();
 			}
 		}
-		
-		GEditor->GetTimerManager()->ClearTimer(AnimatingNode->NodeTimer);
+		// Otherwise it already exists, so reset it's animated state
+		else
+		{
+			AnimatingNode->ResetElapsedSeconds();
+			GEditor->GetTimerManager()->ClearTimer(AnimatingNode->NodeTimer);
+		}
+
 		GEditor->GetTimerManager()->SetTimer(AnimatingNode->NodeTimer, FTimerDelegate::CreateSP(this, &SDetailsViewBase::HandleNodeAnimationComplete, AnimatingNode), InAnimationDuration, false);
 	}
 }
@@ -1689,24 +1733,17 @@ bool SDetailsViewBase::IsNodeAnimating(TSharedPtr<FPropertyNode> InNode)
 		return false;
 	}
 
-	TSharedRef<FPropertyPath> InNodePath = FPropertyNode::CreatePropertyPath(InNode.ToSharedRef());
-	return CurrentlyAnimatingNodeCollections.ContainsByPredicate(
-		[InNodePath](const FAnimatingNodeCollection& AnimatingNode)
-		{
-			return Algo::AnyOf(AnimatingNode.NodePaths, [InNodePath](const TSharedPtr<FPropertyPath>& NodePath)
-			{
-				return FPropertyPath::AreEqual(NodePath.ToSharedRef(), InNodePath);
-			});
-		});
+	return GetAnimatingNodeCollectionForPropertyNode(InNode) != nullptr;
 }
 
-bool SDetailsViewBase::FAnimatingNodeCollection::IsValid() const
+float SDetailsViewBase::GetNodeAnimationTime(TSharedPtr<FPropertyNode> InNode)
 {
-	return Algo::AllOf(NodePaths,
-			[](const TSharedPtr<FPropertyPath>& InNodePath)
-			{
-				return InNodePath.IsValid();
-			});
+	if (const FAnimatingNodeCollection* AnimatingNodeCollection = GetAnimatingNodeCollectionForPropertyNode(InNode))
+	{
+		return AnimatingNodeCollection->GetElapsedSeconds();
+	}
+
+	return 0.0f;
 }
 
 void SDetailsViewBase::HandleNodeAnimationComplete(FAnimatingNodeCollection* InAnimatedNode)
@@ -1721,6 +1758,26 @@ void SDetailsViewBase::HandleNodeAnimationComplete(FAnimatingNodeCollection* InA
 	}
 }
 
+SDetailsViewBase::FAnimatingNodeCollection* SDetailsViewBase::GetAnimatingNodeCollectionForPropertyNode(const TSharedPtr<FPropertyNode>& InNode)
+{
+	if (!InNode.IsValid())
+	{
+		return nullptr;
+	}
+
+	TSharedRef<FPropertyPath> InNodePath = FPropertyNode::CreatePropertyPath(InNode.ToSharedRef());
+	return CurrentlyAnimatingNodeCollections.FindByPredicate(
+		[InNodePath](const FAnimatingNodeCollection& AnimatingNode)
+		{
+			return Algo::AnyOf(
+				AnimatingNode.NodePaths,
+				[InNodePath](const TSharedPtr<FPropertyPath>& NodePath)
+				{
+					return FPropertyPath::AreEqual(NodePath.ToSharedRef(), InNodePath);
+				});
+		});
+}
+
 void SDetailsViewBase::FilterRootNode(const TSharedPtr<FComplexPropertyNode>& RootNode)
 {
 	if (RootNode.IsValid())
diff --git a/Engine/Source/Editor/PropertyEditor/Private/SDetailsViewBase.h b/Engine/Source/Editor/PropertyEditor/Private/SDetailsViewBase.h
--- a/Engine/Source/Editor/PropertyEditor/Private/SDetailsViewBase.h
+++ b/Engine/Source/Editor/PropertyEditor/Private/SDetailsViewBase.h
@@ -152,6 +152,7 @@ public:
 	virtual void RestoreExpandedItems(TSharedRef<FPropertyNode> StartNode, TMap<UStruct*, FStringPrefixTree>* OptionalExpansionStates);
 	virtual void MarkNodeAnimating(TSharedPtr<FPropertyNode> InNode, float InAnimationDuration, TOptional<FGuid> InAnimationBatchId) override;
 	virtual bool IsNodeAnimating(TSharedPtr<FPropertyNode> InNode) override;
+	virtual float GetNodeAnimationTime(TSharedPtr<FPropertyNode> InNode) override;
 	virtual FDetailColumnSizeData& GetColumnSizeData() override { return ColumnSizeData; }
 	virtual bool IsFavoritingEnabled() const override { return DetailsViewArgs.bAllowFavoriteSystem; }
 	virtual bool IsConnected() const = 0;
@@ -210,14 +211,12 @@ protected:
 	{
 		FAnimatingNodeCollection()
 			: BatchId(FGuid::NewGuid())
+			, CreationTime(FPlatformTime::Seconds())
 		{
 		}
 
-		explicit FAnimatingNodeCollection(const TConstArrayView<TSharedPtr<FPropertyPath>>& InNodePaths, const FGuid& InBatchId = FGuid::NewGuid())
-			: NodePaths(InNodePaths)
-			, BatchId(InBatchId)
-		{
-		}
+		/** Creates a collection for the given node(s). If a BatchId is not specified, it's calculated based on the provided nodes. */
+		explicit FAnimatingNodeCollection(const TConstArrayView<TSharedPtr<FPropertyPath>>& InNodePaths, const FGuid& InBatchId = FGuid());
 
 		TArray<TSharedPtr<FPropertyPath>> NodePaths;
 
@@ -226,6 +225,9 @@ protected:
 
 		FTimerHandle NodeTimer;
 
+		/** Utility function that provides a deterministic Id for the given NodePath. */
+		static FGuid GetIdForNodePath(const TSharedPtr<FPropertyPath>& InNodePath);
+
 		/** Get validity of this AnimatingNode, determined by it's node path. */
 		bool IsValid() const;
 
@@ -262,11 +264,23 @@ protected:
 		{
 			return !(*this == Other);
 		}
+
+		/** Effectively resets the stored animation time. */
+		void ResetElapsedSeconds();
+
+		float GetElapsedSeconds() const;
+
+	private:
+        /** The time at which this FAnimatingNodeCollection was created, useful for resuming animation on widget reconstruction. */
+        double CreationTime = 0.0f;
 	};
 
 	/** Called when a node collection animation completes */
 	void HandleNodeAnimationComplete(FAnimatingNodeCollection* InAnimatedNode);
 
+	/** Returns the owning FAnimatingNodeCollection for the provided PropertyNode, if any. */
+	FAnimatingNodeCollection* GetAnimatingNodeCollectionForPropertyNode(const TSharedPtr<FPropertyNode>& InNode);
+
 	/** Updates the property map for access when customizing the details view.  Generates default layout for properties */
 	void UpdatePropertyMaps();
 
