Fix static analysis warning - expression result is always true

diff --git a/Engine/Source/Editor/PropertyEditor/Private/PropertyHandleImpl.cpp b/Engine/Source/Editor/PropertyEditor/Private/PropertyHandleImpl.cpp
--- a/Engine/Source/Editor/PropertyEditor/Private/PropertyHandleImpl.cpp
+++ b/Engine/Source/Editor/PropertyEditor/Private/PropertyHandleImpl.cpp
@@ -6142,11 +6142,7 @@ FPropertyHandleItemAddResult FPropertyHandleMap::AddItem()
 		}
 	}
 
-	if (Result != FPropertyAccess::Success)
-	{
-		Implementation->ShowInvalidOperationError(LOCTEXT("DuplicateMapKey_Add", "Cannot add a new key to the map while a key with the same value exists"));
-	}
-
+	Implementation->ShowInvalidOperationError(LOCTEXT("DuplicateMapKey_Add", "Cannot add a new key to the map while a key with the same value exists"));
 	return FPropertyHandleItemAddResult();
 }
 
