Fix static analysis warning - expression result is always true

diff --git a/Engine/Source/Editor/PropertyEditor/Private/PropertyHandleImpl.cpp b/Engine/Source/Editor/PropertyEditor/Private/PropertyHandleImpl.cpp
--- a/Engine/Source/Editor/PropertyEditor/Private/PropertyHandleImpl.cpp
+++ b/Engine/Source/Editor/PropertyEditor/Private/PropertyHandleImpl.cpp
@@ -5890,11 +5890,7 @@ FPropertyHandleItemAddResult FPropertyHandleSet::AddItem()
 		}
 	}
 
-	if (Result != FPropertyAccess::Success)
-	{
-		Implementation->ShowInvalidOperationError(LOCTEXT("DuplicateSetElement_Add", "Cannot add a new element to the set while an element with the same value exists"));
-	}
-
+	Implementation->ShowInvalidOperationError(LOCTEXT("DuplicateSetElement_Add", "Cannot add a new element to the set while an element with the same value exists"));
 	return FPropertyHandleItemAddResult();
 }
 
