Fix assertion if customization intercepts child node creation and property node auto-expand is enabled

diff --git a/Engine/Source/Editor/PropertyEditor/Private/Presentation/PropertyEditor/PropertyEditor.cpp b/Engine/Source/Editor/PropertyEditor/Private/Presentation/PropertyEditor/PropertyEditor.cpp
--- a/Engine/Source/Editor/PropertyEditor/Private/Presentation/PropertyEditor/PropertyEditor.cpp
+++ b/Engine/Source/Editor/PropertyEditor/Private/Presentation/PropertyEditor/PropertyEditor.cpp
@@ -253,8 +253,9 @@ void FPropertyEditor::OnAddItem()
 
 	if (Result.GetAccessResult() == FPropertyAccess::Success && UE::Private::CVarExpandAddedItem->GetBool())
 	{
-		TSharedPtr<FPropertyNode>& ChildNode = PropertyNode->GetChildNode(Result.GetIndex());
-		if (ChildNode)
+		TSharedPtr<FPropertyNode> ChildNode;
+		bool bHasChildNode = PropertyNode->GetChildNode(Result.GetIndex(), ChildNode);
+		if (bHasChildNode)
 		{
 			ChildNode->SetNodeFlags(EPropertyNodeFlags::Expanded, true);
 		}
