Fix crash when evaluating an edit condition that is on a sparse class data member

[RN] Fixed crash when display objects using EditConditions on their SparseClassData

diff --git a/Engine/Source/Editor/PropertyEditor/Private/EditConditionContext.cpp b/Engine/Source/Editor/PropertyEditor/Private/EditConditionContext.cpp
--- a/Engine/Source/Editor/PropertyEditor/Private/EditConditionContext.cpp
+++ b/Engine/Source/Editor/PropertyEditor/Private/EditConditionContext.cpp
@@ -163,6 +163,14 @@ static const FPropertyNode* GetEditConditionParentNode(const TSharedPtr<FPropert
 		return ParentNode->GetParentNode();
 	}
 
+	if(PropertyNode->HasNodeFlags(EPropertyNodeFlags::IsSparseClassData))
+	{
+		// by using ourself as the parent node we allow GetValueAddressFromObject and friends to find the right c object,
+		// see GetValuePtrOfInstance, which calls InParentNode->GetValueAddressFromObject(Obj). If we return our 
+		// notional (UI) parent node, we will attempt to index the object with a sparse class data property:
+		return PropertyNode.Get(); 
+	}
+
 	return ParentNode;
 }
 
diff --git a/Engine/Source/Editor/PropertyEditor/Private/ObjectPropertyNode.h b/Engine/Source/Editor/PropertyEditor/Private/ObjectPropertyNode.h
--- a/Engine/Source/Editor/PropertyEditor/Private/ObjectPropertyNode.h
+++ b/Engine/Source/Editor/PropertyEditor/Private/ObjectPropertyNode.h
@@ -112,6 +112,8 @@ public:
 			return nullptr;
 		}
 
+		// GetValueAddressFromObject decayed to void so do quick safety check:
+		ensure(Obj->IsA(InProperty->GetOwnerClass()) || InParentNode->HasNodeFlags(EPropertyNodeFlags::IsSparseClassData));
 		return InProperty->ContainerPtrToValuePtr<uint8>(ParentPtr);
 	}
 	virtual TWeakObjectPtr<UObject> GetInstanceAsUObject(int32 Index) const override
