Remove overzealous ensure, nested structures like those found in WorldSettings (Window->World Settings) will index within the structure.

#rnx

diff --git a/Engine/Source/Editor/PropertyEditor/Private/ObjectPropertyNode.h b/Engine/Source/Editor/PropertyEditor/Private/ObjectPropertyNode.h
--- a/Engine/Source/Editor/PropertyEditor/Private/ObjectPropertyNode.h
+++ b/Engine/Source/Editor/PropertyEditor/Private/ObjectPropertyNode.h
@@ -112,8 +112,6 @@ public:
 			return nullptr;
 		}
 
-		// GetValueAddressFromObject decayed to void so do quick safety check:
-		ensure(Obj->IsA(InProperty->GetOwnerClass()) || InParentNode->HasNodeFlags(EPropertyNodeFlags::IsSparseClassData));
 		return InProperty->ContainerPtrToValuePtr<uint8>(ParentPtr);
 	}
 	virtual TWeakObjectPtr<UObject> GetInstanceAsUObject(int32 Index) const override
