[PropertyEditor] Added GetPrimaryPropertyHandle, which always returns a single, "primary" handle for the owning row (ie. FVector will return the FVector struct, not the X, Y, Z properties)

diff --git a/Engine/Source/Editor/PropertyEditor/Private/SDetailSingleItemRow.cpp b/Engine/Source/Editor/PropertyEditor/Private/SDetailSingleItemRow.cpp
--- a/Engine/Source/Editor/PropertyEditor/Private/SDetailSingleItemRow.cpp
+++ b/Engine/Source/Editor/PropertyEditor/Private/SDetailSingleItemRow.cpp
@@ -1331,6 +1331,23 @@ TArray<TSharedPtr<IPropertyHandle>> SDetailSingleItemRow::GetPropertyHandles(con
 	return WidgetRow.GetPropertyHandles();
 }
 
+TSharedPtr<IPropertyHandle> SDetailSingleItemRow::GetPrimaryPropertyHandle() const
+{
+	TSharedPtr<IPropertyHandle> PrimaryPropertyHandle = SDetailTableRowBase::GetPrimaryPropertyHandle();
+	if (PrimaryPropertyHandle.IsValid())
+	{
+		return PrimaryPropertyHandle;
+	}
+
+	if (const TArray<TSharedPtr<IPropertyHandle>>& WidgetPropertyHandles = WidgetRow.GetPropertyHandles();
+		!WidgetPropertyHandles.IsEmpty())
+	{
+		return WidgetPropertyHandles[0];
+	}
+
+	return nullptr;
+}
+
 void SDetailSingleItemRow::OnCopyProperty()
 {
 	if (OwnerTreeNode.IsValid())
@@ -1855,9 +1872,10 @@ void SDetailSingleItemRow::CreateGlobalExtensionWidgets(TArray<FPropertyRowExten
 	{
 		Args.PropertyHandle = PropertyEditorHelpers::GetPropertyHandle(Customization->GetPropertyNode().ToSharedRef(), nullptr, nullptr);
 	}
-	else if (GetPropertyHandles(true).Num() && GetPropertyHandles(true)[0] && GetPropertyHandles(true)[0]->IsValidHandle())
+	else if (const TSharedPtr<IPropertyHandle> PrimaryPropertyHandle = GetPrimaryPropertyHandle();
+		PrimaryPropertyHandle.IsValid() && PrimaryPropertyHandle->IsValidHandle())
 	{
-		Args.PropertyHandle = GetPropertyHandles(true)[0];
+		Args.PropertyHandle = PrimaryPropertyHandle;
 	}
 
 	PropertyEditorModule.GetGlobalRowExtensionDelegate().Broadcast(Args, OutExtensions);
diff --git a/Engine/Source/Editor/PropertyEditor/Private/SDetailSingleItemRow.h b/Engine/Source/Editor/PropertyEditor/Private/SDetailSingleItemRow.h
--- a/Engine/Source/Editor/PropertyEditor/Private/SDetailSingleItemRow.h
+++ b/Engine/Source/Editor/PropertyEditor/Private/SDetailSingleItemRow.h
@@ -71,6 +71,7 @@ public:
 protected:
 	virtual void PopulateContextMenu(UToolMenu* ToolMenu) override;
 	virtual TArray<TSharedPtr<IPropertyHandle>> GetPropertyHandles(const bool bRecursive = false) const override;
+	virtual TSharedPtr<IPropertyHandle> GetPrimaryPropertyHandle() const override;
 
 private:
 	void OnCopyProperty();
diff --git a/Engine/Source/Editor/PropertyEditor/Private/SDetailTableRowBase.cpp b/Engine/Source/Editor/PropertyEditor/Private/SDetailTableRowBase.cpp
--- a/Engine/Source/Editor/PropertyEditor/Private/SDetailTableRowBase.cpp
+++ b/Engine/Source/Editor/PropertyEditor/Private/SDetailTableRowBase.cpp
@@ -170,7 +170,7 @@ TArray<TSharedPtr<IPropertyHandle>> SDetailTableRowBase::GetPropertyHandles(cons
 		return true;
 	};
 
-	TArray<TSharedPtr<IPropertyHandle>> PropertyHandles;	
+	TArray<TSharedPtr<IPropertyHandle>> PropertyHandles;
 	if (const TSharedPtr<FDetailTreeNode> OwnerTreeNodePtr = OwnerTreeNode.Pin())
 	{
 		AppendPropertyHandles(OwnerTreeNodePtr, PropertyHandles);
@@ -181,6 +181,28 @@ TArray<TSharedPtr<IPropertyHandle>> SDetailTableRowBase::GetPropertyHandles(cons
 	return PropertyHandles;
 }
 
+TSharedPtr<IPropertyHandle> SDetailTableRowBase::GetPrimaryPropertyHandle() const
+{
+	if (const TSharedPtr<FDetailTreeNode> OwnerTreeNodePtr = OwnerTreeNode.Pin())
+	{
+		TSharedPtr<IPropertyHandle> PropertyHandle = OwnerTreeNodePtr->CreatePropertyHandle();
+		if (PropertyHandle.IsValid() && PropertyHandle->IsValidHandle())
+		{
+			return PropertyHandle;
+		}
+
+		// If the primary property handle can't be retrieved from the associated node, try to get it from the children
+		// This is less than ideal - the primary property handle should always be available from the node, except for category headers, etc.
+		const TArray<TSharedPtr<IPropertyHandle>> PropertyHandles = GetPropertyHandles(false);
+		if (!PropertyHandles.IsEmpty())
+		{
+			return PropertyHandles[0];
+		}
+	}
+
+	return nullptr;
+}
+
 void SDetailTableRowBase::ForceRefresh()
 {
 	if (const TSharedPtr<FDetailTreeNode> Owner = OwnerTreeNode.Pin();
diff --git a/Engine/Source/Editor/PropertyEditor/Private/SDetailTableRowBase.h b/Engine/Source/Editor/PropertyEditor/Private/SDetailTableRowBase.h
--- a/Engine/Source/Editor/PropertyEditor/Private/SDetailTableRowBase.h
+++ b/Engine/Source/Editor/PropertyEditor/Private/SDetailTableRowBase.h
@@ -46,9 +46,17 @@ protected:
 	/** Retrieve all property nodes represented by the given property handles. */
 	TArray<TSharedPtr<FPropertyNode>> GetPropertyNodesFromHandles(const TConstArrayView<TSharedPtr<IPropertyHandle>>& InPropertyHandles) const;
 
-	/** Retrieve all property handles represented by this row, and it's children if Recursive specified. */
+	/**
+	 * Retrieve all property handles represented by this row, and it's children if Recursive specified.
+	 * @note: This only returns the visible property handles, not the "primary" one (although often these are the same),
+	 * for example FVector will give you the X, Y and Z handles, not the FVector itself.
+	 * For the "primary" property handle, use GetPrimaryPropertyHandle.
+	 */
 	virtual TArray<TSharedPtr<IPropertyHandle>> GetPropertyHandles(const bool bRecursive = false) const;
 
+	/** Retrieve the primary property handle for this row, ie. an FVector should return the Struct itself, not it's X, Y, Z properties. */
+	virtual TSharedPtr<IPropertyHandle> GetPrimaryPropertyHandle() const;
+
 	/** Refreshes widget and associated list views. */
 	virtual void ForceRefresh();
 
