[PropertyEditor]
* Fixed compile error when compiling with forceinline->inline

diff --git a/Engine/Source/Editor/PropertyEditor/Private/DetailLayoutBuilderImpl.cpp b/Engine/Source/Editor/PropertyEditor/Private/DetailLayoutBuilderImpl.cpp
--- a/Engine/Source/Editor/PropertyEditor/Private/DetailLayoutBuilderImpl.cpp
+++ b/Engine/Source/Editor/PropertyEditor/Private/DetailLayoutBuilderImpl.cpp
@@ -280,9 +280,9 @@ TSharedRef<IPropertyHandle> FDetailLayoutBuilderImpl::GetProperty( const FName P
 
 FName FDetailLayoutBuilderImpl::GetTopLevelProperty()
 {
-	for (auto It = PropertyMap.CreateConstIterator(); It; ++It)
+	if (!PropertyMap.IsEmpty())
 	{
-		return It.Key();
+		return PropertyMap.CreateConstIterator().Key();
 	}
 	return NAME_None;
 }
