[Property Matrix Editor] Can now unselect individual cells from the property matrix editor

diff --git a/Engine/Source/Editor/PropertyEditor/Private/UserInterface/PropertyTable/SPropertyTableRow.h b/Engine/Source/Editor/PropertyEditor/Private/UserInterface/PropertyTable/SPropertyTableRow.h
--- a/Engine/Source/Editor/PropertyEditor/Private/UserInterface/PropertyTable/SPropertyTableRow.h
+++ b/Engine/Source/Editor/PropertyEditor/Private/UserInterface/PropertyTable/SPropertyTableRow.h
@@ -54,7 +54,7 @@ public:
 	FReply OnMouseButtonDown( const FGeometry& MyGeometry, const FPointerEvent& MouseEvent ) override
 	{
 		PreviousCurrentCell = Row->GetTable()->GetCurrentCell();
-		Row->GetTable()->SetCurrentCell( NULL );
+		Row->GetTable()->SetCurrentCell( nullptr );
 
 		return FReply::Unhandled();
 	}
@@ -81,13 +81,14 @@ public:
 				// The Row may not have any cells, so we'll detect an old LastClickedCell and clear
 				if ( !Row->HasCells() || LastClickedCell->GetRow() != Row || !LastClickedCell->GetColumn()->CanSelectCells() )
 				{
-					Table->SetLastClickedCell( NULL );
-					Table->SetCurrentCell( NULL );
+					Table->SetLastClickedCell( nullptr );
+					Table->SetCurrentCell( nullptr );
 				}
 				else
 				{
-					if ( CurrentCell.IsValid() && LastClickedCell == CurrentCell )
+					if ( CurrentCell.IsValid() && LastClickedCell == CurrentCell && !MouseEvent.IsControlDown())
 					{
+						// Don't let the table enter edit mode if the ctrl key is down
 						Table->GetLastClickedCell()->EnterEditMode();
 					}
 					else if ( MouseEvent.IsShiftDown() && CurrentCell.IsValid() && LastClickedCell != CurrentCell )
@@ -110,12 +111,20 @@ public:
 						Table->SelectCellRange( StartingCell, EndingCell );
 						Table->SetCurrentCell( CurrentCell );
 					}
-					else if ( MouseEvent.IsControlDown() )
+					else if (MouseEvent.IsControlDown())
 					{
-						SelectedCells.Add( LastClickedCell.ToSharedRef() );
-
-						Table->SetSelectedCells( SelectedCells );
-						Table->SetCurrentCell( LastClickedCell );
+						if (SelectedCells.Contains(LastClickedCell.ToSharedRef()))
+						{
+							SelectedCells.Remove(LastClickedCell.ToSharedRef());
+							Table->SetCurrentCell(nullptr);
+						}
+						else
+						{
+							SelectedCells.Add(LastClickedCell.ToSharedRef());
+							Table->SetCurrentCell(LastClickedCell);
+						}
+
+						Table->SetSelectedCells(SelectedCells);
 					}
 					else
 					{
@@ -132,38 +141,13 @@ public:
 		{
 			if ( !Table->GetSelectedRows().Contains( Row.ToSharedRef() ) )
 			{
-				Table->SetCurrentCell( NULL );
+				Table->SetCurrentCell( nullptr );
 			}
 		}
 		
 		return Reply;
 	}
 
-	FReply OnMouseButtonDoubleClick( const FGeometry& MyGeometry, const FPointerEvent& MouseEvent ) override
-	{
-		FReply Reply = FReply::Unhandled();
-		const TSharedRef< IPropertyTable > Table = Row->GetTable();
-
-		const EPropertyTableSelectionUnit::Type SelectionUnit = Table->GetSelectionUnit();
-		if ( SelectionUnit == EPropertyTableSelectionUnit::Row || SelectionUnit == EPropertyTableSelectionUnit::Cell )
-		{
-			Reply = SPropertyTableRowBase::OnMouseButtonDoubleClick( MyGeometry, MouseEvent );
-			Table->SetCurrentRow( Row );
-
-			if ( SelectionUnit == EPropertyTableSelectionUnit::Cell && Reply.IsEventHandled() )
-			{
-				const TSharedPtr< IPropertyTableCell > LastClickedCell = Table->GetLastClickedCell();
-				if ( !Reply.GetDetectDragRequest().IsValid() && LastClickedCell.IsValid() )
-				{
-					Table->SetCurrentCell( LastClickedCell );
-					Table->GetLastClickedCell()->EnterEditMode();
-				}
-			}
-		}
-
-		return Reply;
-	}
-
 	virtual void Tick( const FGeometry& AllottedGeometry, const double InCurrentTime, const float InDeltaTime ) override
 	{
 		Row->Tick();
