Implement SDetailsView tests for populating an array of soft object pointers.

#rnx

diff --git a/Engine/Source/Editor/PropertyEditor/Private/Tests/DetailsView/PropertyHandle/DetailsViewPropertyHandleArrayTest.cpp b/Engine/Source/Editor/PropertyEditor/Private/Tests/DetailsView/PropertyHandle/DetailsViewPropertyHandleArrayTest.cpp
new file mode 100644
--- /dev/null
+++ b/Engine/Source/Editor/PropertyEditor/Private/Tests/DetailsView/PropertyHandle/DetailsViewPropertyHandleArrayTest.cpp
@@ -0,0 +1,36 @@
+// Copyright Epic Games, Inc. All Rights Reserved.
+
+#if WITH_DEV_AUTOMATION_TESTS
+
+#include "DetailsViewPropertyHandleTestBase.h"
+
+DETAILS_VIEW_PROPERTY_HANDLE_TEST(FDetailsViewPropertyHandleArrayTest, "Editor.PropertyEditor.DetailsView.PropertyHandleArray")
+{
+	FDetailsViewPropertyHandleArrayTest() : FDetailsViewPropertyHandleTestBase("Properties", "TestValueSoftPtrArray") {}
+
+	TSharedPtr<IPropertyHandleArray> PropertyHandleArray;
+
+	BEFORE_EACH()
+	{
+		FDetailsViewPropertyHandleTestBase::Setup();
+		ASSERT_THAT(IsNotNull(PropertyHandle));
+
+		PropertyHandleArray = PropertyHandle->AsArray();
+		ASSERT_THAT(IsNotNull(PropertyHandleArray, TEXT("Property handle is an array check")));
+	}
+	
+	// The purpose of the "AddItem" test is to verify functionality that could not be covered in the test code implemented outside of this module
+	// due to a lack of access to the private headers. In particular, this test checks whether users can add an item to the array property
+	// through the SDetailsView widget.
+	// 
+	// Note: The test method may need to be updated soon due to changes in the PropertyHandle operation.
+	TEST_METHOD(AddItem)
+	{
+		const int32 InitialSize = TestObject->TestValueSoftPtrArray.Num();
+		const FPropertyAccess::Result Result = PropertyHandleArray->AddItem();
+		ASSERT_THAT(AreEqual(FPropertyAccess::Result::Success, Result, TEXT("Property handle added an item check")));
+		ASSERT_THAT(AreEqual(1, TestObject->TestValueSoftPtrArray.Num() - InitialSize, TEXT("Array size of the editing object's property increased by one check")));
+	}
+};
+
+#endif // WITH_DEV_AUTOMATION_TESTS
\ No newline at end of file
diff --git a/Engine/Source/Editor/PropertyEditor/Private/Tests/DetailsView/PropertyHandle/DetailsViewPropertyHandleObjectTest.cpp b/Engine/Source/Editor/PropertyEditor/Private/Tests/DetailsView/PropertyHandle/DetailsViewPropertyHandleObjectTest.cpp
new file mode 100644
--- /dev/null
+++ b/Engine/Source/Editor/PropertyEditor/Private/Tests/DetailsView/PropertyHandle/DetailsViewPropertyHandleObjectTest.cpp
@@ -0,0 +1,54 @@
+// Copyright Epic Games, Inc. All Rights Reserved.
+
+#if WITH_DEV_AUTOMATION_TESTS
+
+#include "DetailsViewPropertyHandleTestBase.h"
+
+DETAILS_VIEW_PROPERTY_HANDLE_TEST(FDetailsViewPropertyHandleObjectTest, "Editor.PropertyEditor.DetailsView.PropertyHandleObject")
+{
+	FDetailsViewPropertyHandleObjectTest() : FDetailsViewPropertyHandleTestBase("Properties", "TestValueSoftPtr") {}
+
+	// The purpose of the "SetValue" and "GetValue" tests is to verify functionality that could not be covered in the test code implemented outside of this module
+	// due to a lack of access to the private headers. In particular, these tests check whether users can set a property value through the SDetailsView widget.
+	// 
+	// Note: The test methods may need to be updated soon due to changes in the PropertyHandle operation.
+	TEST_METHOD(SetValue_FAssetData)
+	{
+		UDetailsViewPropertyHandleTestValueClass* const TestValueObject = NewObject<UDetailsViewPropertyHandleTestValueClass>();
+		const FPropertyAccess::Result Result = PropertyHandle->SetValue(FAssetData(TestValueObject));
+		ASSERT_THAT(AreEqual(FPropertyAccess::Result::Success, Result, TEXT("Property handle set value check")));
+		ASSERT_THAT(AreEqual(TestValueObject, TestObject->TestValueSoftPtr.Get(), TEXT("Property of the editing object is set correctly check")));
+	}
+
+	TEST_METHOD(SetValue_UObject)
+	{
+		UDetailsViewPropertyHandleTestValueClass* const TestValueObject = NewObject<UDetailsViewPropertyHandleTestValueClass>();
+		const FPropertyAccess::Result Result = PropertyHandle->SetValue(TestValueObject);
+		ASSERT_THAT(AreEqual(FPropertyAccess::Result::Success, Result, TEXT("Property handle set value check")));
+		ASSERT_THAT(AreEqual(TestValueObject, TestObject->TestValueSoftPtr.Get(), TEXT("Property of the editing object is set correctly check")));
+	}
+	
+	TEST_METHOD(GetValue_FAssetData)
+	{
+		UDetailsViewPropertyHandleTestValueClass* const TestValueObject = NewObject<UDetailsViewPropertyHandleTestValueClass>();
+		TestObject->TestValueSoftPtr = TestValueObject;
+		
+		FAssetData Value;
+		const FPropertyAccess::Result Result = PropertyHandle->GetValue(Value);
+		ASSERT_THAT(AreEqual(FPropertyAccess::Result::Success, Result, TEXT("Property handle retrieved value check")));
+		ASSERT_THAT(AreEqual(FAssetData(TestValueObject), Value, TEXT("The retrieved value is correct check")));
+	}
+	
+	TEST_METHOD(GetValue_UObject)
+	{
+		UDetailsViewPropertyHandleTestValueClass* const TestValueObject = NewObject<UDetailsViewPropertyHandleTestValueClass>();
+		TestObject->TestValueSoftPtr = TestValueObject;
+		
+		UObject* Value = nullptr;
+		const FPropertyAccess::Result Result = PropertyHandle->GetValue(Value);
+		ASSERT_THAT(AreEqual(FPropertyAccess::Result::Success, Result, TEXT("Property handle retrieved value check")));
+		ASSERT_THAT(AreEqual(Cast<UObject>(TestValueObject), Value, TEXT("The retrieved value is correct check")));
+	}
+};
+
+#endif // WITH_DEV_AUTOMATION_TESTS
\ No newline at end of file
diff --git a/Engine/Source/Editor/PropertyEditor/Private/Tests/DetailsView/PropertyHandle/DetailsViewPropertyHandleTestBase.h b/Engine/Source/Editor/PropertyEditor/Private/Tests/DetailsView/PropertyHandle/DetailsViewPropertyHandleTestBase.h
new file mode 100644
--- /dev/null
+++ b/Engine/Source/Editor/PropertyEditor/Private/Tests/DetailsView/PropertyHandle/DetailsViewPropertyHandleTestBase.h
@@ -0,0 +1,96 @@
+// Copyright Epic Games, Inc. All Rights Reserved.
+#pragma once
+
+#if WITH_DEV_AUTOMATION_TESTS
+
+#include "CoreMinimal.h"
+#include "CQTest.h"
+#include "DetailTreeNode.h"
+#include "DetailsViewPropertyHandleTestClass.h"
+#include "PropertyHandle.h"
+#include "SDetailsView.h"
+
+template<typename Derived, typename AsserterType>
+class FDetailsViewPropertyHandleTestBase : public TTest<Derived, AsserterType>
+{
+protected:
+	using TTest<Derived, AsserterType>::Assert;
+
+	FDetailsViewPropertyHandleTestBase(const FString& InCategoryName, const FString& InPropertyName) :
+		CategoryName(InCategoryName),
+		PropertyName(InPropertyName)
+	{
+	}
+
+	virtual void Setup() override
+	{
+		TestObject = NewObject<UDetailsViewPropertyHandleTestClass>();
+		DetailsView = CreateDetailsViewForObject(TestObject);
+		FindPropertyHandleInDetailsView();
+	}
+
+private:
+	static TSharedRef<SDetailsView> CreateDetailsViewForObject(UObject* InObject)
+	{
+		TSharedRef<SDetailsView> DetailsView = SNew(SDetailsView, FDetailsViewArgs());
+		DetailsView->SetObject(InObject);
+		return DetailsView;
+	}
+
+	static TSharedPtr<FDetailTreeNode> GetHeadNodeByName(SDetailsView& DetailsView, const FString& HeadNodeName)
+	{
+		TArray<TWeakPtr<FDetailTreeNode>> WeakNodes;
+		DetailsView.GetHeadNodes(WeakNodes);
+
+		for (const TWeakPtr<FDetailTreeNode>& WeakNode : WeakNodes)
+		{
+			if (TSharedPtr<FDetailTreeNode> Node = WeakNode.Pin())
+			{
+				if (HeadNodeName.Equals(Node->GetNodeName().ToString())) { return Node; }
+			}
+		}
+
+		return nullptr;
+	}
+
+	static TSharedPtr<FDetailTreeNode> GetChildNodeByName(FDetailTreeNode& ParentNode, const FString& ChildNodeName)
+	{
+		FDetailNodeList ChildNodes;
+		ParentNode.GetChildren(ChildNodes, true);
+
+		for (const TSharedRef<FDetailTreeNode>& Node : ChildNodes)
+		{
+			if (ChildNodeName.Equals(Node->GetNodeName().ToString())) { return Node; }
+		}
+
+		return nullptr;
+	}
+
+	void FindPropertyHandleInDetailsView()
+	{
+		const TSharedPtr<FDetailTreeNode> CategoryNode = GetHeadNodeByName(*DetailsView.Get(), CategoryName);
+		ASSERT_THAT(IsNotNull(CategoryNode, FString::Printf(TEXT("Head node with name '%s' retrieved check"), *CategoryName)));
+
+		const TSharedPtr<FDetailTreeNode> PropertyNode = GetChildNodeByName(*CategoryNode.Get(), PropertyName);
+		ASSERT_THAT(IsNotNull(PropertyNode, FString::Printf(TEXT("Child node with name '%s' retrieved check"), *PropertyName)));
+
+		const TSharedPtr<IDetailPropertyRow> PropertyRow = PropertyNode->GetRow();
+		ASSERT_THAT(IsNotNull(PropertyRow, TEXT("Property row retrieved check")));
+
+		PropertyHandle = PropertyRow->GetPropertyHandle();
+		ASSERT_THAT(IsNotNull(PropertyHandle, TEXT("Property handle retrieved check")));
+	}
+
+protected:
+	UDetailsViewPropertyHandleTestClass* TestObject = nullptr;
+	TSharedPtr<IPropertyHandle> PropertyHandle;
+
+private:
+	TSharedPtr<SDetailsView> DetailsView;
+	const FString CategoryName;
+	const FString PropertyName;
+};
+
+#define DETAILS_VIEW_PROPERTY_HANDLE_TEST(_ClassName, _TestDir) TEST_CLASS_WITH_BASE_AND_FLAGS(_ClassName, _TestDir, FDetailsViewPropertyHandleTestBase, EAutomationTestFlags::EditorContext | EAutomationTestFlags::EngineFilter)
+
+#endif // WITH_DEV_AUTOMATION_TESTS
\ No newline at end of file
diff --git a/Engine/Source/Editor/PropertyEditor/Private/Tests/DetailsView/PropertyHandle/DetailsViewPropertyHandleTestClass.h b/Engine/Source/Editor/PropertyEditor/Private/Tests/DetailsView/PropertyHandle/DetailsViewPropertyHandleTestClass.h
new file mode 100644
--- /dev/null
+++ b/Engine/Source/Editor/PropertyEditor/Private/Tests/DetailsView/PropertyHandle/DetailsViewPropertyHandleTestClass.h
@@ -0,0 +1,29 @@
+// Copyright Epic Games, Inc. All Rights Reserved.
+#pragma once
+
+#include "UObject/Object.h"
+
+#include "DetailsViewPropertyHandleTestClass.generated.h"
+
+#if WITH_EDITORONLY_DATA
+
+UCLASS()
+class UDetailsViewPropertyHandleTestValueClass : public UObject
+{
+	GENERATED_BODY()
+};
+
+UCLASS()
+class UDetailsViewPropertyHandleTestClass : public UObject
+{
+	GENERATED_BODY()
+
+public:	
+	UPROPERTY(EditAnywhere, Category = "Properties")
+	TSoftObjectPtr<UDetailsViewPropertyHandleTestValueClass> TestValueSoftPtr;
+	
+	UPROPERTY(EditAnywhere, Category = "Properties")
+	TArray<TSoftObjectPtr<UDetailsViewPropertyHandleTestValueClass>> TestValueSoftPtrArray;
+};
+
+#endif // WITH_EDITORONLY_DATA
\ No newline at end of file
diff --git a/Engine/Source/Editor/PropertyEditor/PropertyEditor.Build.cs b/Engine/Source/Editor/PropertyEditor/PropertyEditor.Build.cs
--- a/Engine/Source/Editor/PropertyEditor/PropertyEditor.Build.cs
+++ b/Engine/Source/Editor/PropertyEditor/PropertyEditor.Build.cs
@@ -52,6 +52,15 @@ public class PropertyEditor : ModuleRules
 			}
         );
 
+		if (Target.WithAutomationTests)
+		{
+			PrivateDependencyModuleNames.AddRange(
+				new string[] {
+					"CQTest",
+				}
+			);
+		}
+
 		DynamicallyLoadedModuleNames.AddRange(
 			new string[] {
                 "AssetRegistry",
