Fixed pasting paths into the asset property editor working for package paths if and only if the package was loaded (it now works for unloaded package paths too).

diff --git a/Engine/Source/Editor/PropertyEditor/Private/UserInterface/PropertyEditor/SPropertyEditorAsset.cpp b/Engine/Source/Editor/PropertyEditor/Private/UserInterface/PropertyEditor/SPropertyEditorAsset.cpp
--- a/Engine/Source/Editor/PropertyEditor/Private/UserInterface/PropertyEditor/SPropertyEditorAsset.cpp
+++ b/Engine/Source/Editor/PropertyEditor/Private/UserInterface/PropertyEditor/SPropertyEditorAsset.cpp
@@ -940,6 +940,47 @@ bool SPropertyEditorAsset::Supports(const FProperty* NodeProperty)
 	return false;
 }
 
+bool SPropertyEditorAsset::ParseAssetText(const FString& InText, FAssetData& OutAssetData)
+{
+	OutAssetData = {};
+
+	FString PossibleObjectPath = FPackageName::ExportTextPathToObjectPath(InText);
+
+	if (PossibleObjectPath.IsEmpty())
+	{
+		return false;
+	}
+
+	if (PossibleObjectPath == TEXT("None"))
+	{
+		return true;
+	}
+	
+	// All supported paths start with a /
+	if (PossibleObjectPath[0] == TEXT('/'))
+	{
+		const bool bTextIsExportTextPath = PossibleObjectPath != InText;
+
+		const FAssetRegistryModule& AssetRegistryModule = FModuleManager::Get().LoadModuleChecked<FAssetRegistryModule>(TEXT("AssetRegistry"));
+
+		// Check if we just have a package name.
+		int32 ObjectNameDelimiterIndex;
+		if (!bTextIsExportTextPath && !PossibleObjectPath.FindChar(TEXT('.'), ObjectNameDelimiterIndex))
+		{
+			// Assume that the object we're trying to load is the main asset inside of the package 
+			// which usually has the same name as the short package name.
+			PossibleObjectPath = PossibleObjectPath + TEXT('.') + FPackageName::GetShortName(PossibleObjectPath);
+		}
+
+		if (PossibleObjectPath.Len() < NAME_SIZE)
+		{
+			OutAssetData = AssetRegistryModule.Get().GetAssetByObjectPath(FSoftObjectPath(PossibleObjectPath));
+		}
+	}
+
+	return OutAssetData.IsValid();
+}
+
 TSharedRef<SWidget> SPropertyEditorAsset::OnGetMenuContent()
 {
 	FObjectOrAssetData Value;
@@ -1596,7 +1637,7 @@ void SPropertyEditorAsset::OnPaste()
 	FString DestPath;
 	FPlatformApplicationMisc::ClipboardPaste(DestPath);
 	
-	PasteFromText(TEXT(""), DestPath);
+	OnPasteFromText(TEXT(""), DestPath, {});
 }
 
 void SPropertyEditorAsset::OnPasteFromText(
@@ -1604,30 +1645,31 @@ void SPropertyEditorAsset::OnPasteFromText(
 	const FString& InText,
 	const TOptional<FGuid>& InOperationId)
 {
-	if (CanPasteFromText(InTag, InText))
+	FAssetData AssetData;
+	if (CanPasteFromText(InTag, InText, AssetData))
 	{
-		PasteFromText(InTag, InText);
+		PasteFromText(InTag, AssetData);
 	}
 }
 
 void SPropertyEditorAsset::PasteFromText(
 	const FString& InTag,
-	const FString& InText)
+	const FAssetData& InAssetData)
 {
-	if(InText == TEXT("None"))
+	if (!InAssetData.IsValid())
 	{
 		SetValue(nullptr);
 	}
 	else
 	{
-		UObject* Object = LoadObject<UObject>(nullptr, *InText);
+		UObject* Object = InAssetData.GetAsset();
 		if(Object && Object->IsA(ObjectClass))
 		{
 			// Check against custom asset filter
 			if (!OnShouldFilterAsset.IsBound()
-				|| !OnShouldFilterAsset.Execute(FAssetData(Object)))
+				|| !OnShouldFilterAsset.Execute(InAssetData))
 			{
-				SetValue(Object);
+				SetValue(InAssetData);
 			}
 		}
 	}
@@ -1638,36 +1680,26 @@ bool SPropertyEditorAsset::CanPaste()
 	FString ClipboardText;
 	FPlatformApplicationMisc::ClipboardPaste(ClipboardText);
 
-	return CanPasteFromText(TEXT(""), ClipboardText);
+	FAssetData AssetData;
+	return CanPasteFromText(TEXT(""), ClipboardText, AssetData);
 }
 
 bool SPropertyEditorAsset::CanPasteFromText(
 	const FString& InTag,
-	const FString& InText) const
+	const FString& InText,
+	FAssetData& OutAssetData) const
 {
 	if (!UE::PropertyEditor::TagMatchesProperty(InTag, PropertyHandle))
 	{
 		return false;
 	}
 
-	const FString PossibleObjectPath = FPackageName::ExportTextPathToObjectPath(InText);
-
-	bool bCanPaste = false;
-
-	if( CanEdit() )
+	if (!CanEdit())
 	{
-		if( PossibleObjectPath == TEXT("None") )
-		{
-			bCanPaste = true;
-		}
-		else
-		{
-			FAssetRegistryModule& AssetRegistryModule = FModuleManager::Get().LoadModuleChecked<FAssetRegistryModule>(TEXT("AssetRegistry"));
-			bCanPaste = PossibleObjectPath.Len() < NAME_SIZE && AssetRegistryModule.Get().GetAssetByObjectPath( FSoftObjectPath(PossibleObjectPath) ).IsValid();
-		}
+		return false;
 	}
 
-	return bCanPaste;
+	return ParseAssetText(InText, OutAssetData);
 }
 
 FReply SPropertyEditorAsset::OnAssetThumbnailDoubleClick( const FGeometry& InMyGeometry, const FPointerEvent& InMouseEvent )
diff --git a/Engine/Source/Editor/PropertyEditor/Private/UserInterface/PropertyEditor/SPropertyEditorAsset.h b/Engine/Source/Editor/PropertyEditor/Private/UserInterface/PropertyEditor/SPropertyEditorAsset.h
--- a/Engine/Source/Editor/PropertyEditor/Private/UserInterface/PropertyEditor/SPropertyEditorAsset.h
+++ b/Engine/Source/Editor/PropertyEditor/Private/UserInterface/PropertyEditor/SPropertyEditorAsset.h
@@ -92,6 +92,13 @@ public:
 	static bool Supports( const TSharedRef<FPropertyEditor>& InPropertyEditor );
 	static bool Supports( const FProperty* NodeProperty );
 
+	/**
+	 * Parses the input text and finds the corrisponding asset in the asset registry.
+	 * Supports various types of paths, such as export text, object and package paths.
+	 * @returns True if OutAssetData contains a valid asset, or the input text was None.
+	 */
+	static bool ParseAssetText(const FString& InText, FAssetData& OutAssetData);
+
 	/**
 	 * Construct the widget.
 	 * @param	InArgs				Arguments for widget construction
@@ -310,7 +317,7 @@ private:
 	/**
 	 * @return True if the (optionally tagged) input contents can be pasted
 	 */
-	bool CanPasteFromText(const FString& InTag, const FString& InText) const;
+	bool CanPasteFromText(const FString& InTag, const FString& InText, FAssetData& OutAssetData) const;
 
 	/**
 	 * Delegate handling pasting an optionally tagged text snippet
@@ -320,7 +327,7 @@ private:
 	/**
 	 * Handle pasting an optionally tagged text snippet
 	 */
-	void PasteFromText(const FString& InTag, const FString& InText);
+	void PasteFromText(const FString& InTag, const FAssetData& InAssetData);
 
 	/**
 	 * @return True of the current clipboard contents can be pasted
diff --git a/Engine/Source/Editor/PropertyEditor/Private/UserInterface/PropertyEditor/SPropertyMenuAssetPicker.cpp b/Engine/Source/Editor/PropertyEditor/Private/UserInterface/PropertyEditor/SPropertyMenuAssetPicker.cpp
--- a/Engine/Source/Editor/PropertyEditor/Private/UserInterface/PropertyEditor/SPropertyMenuAssetPicker.cpp
+++ b/Engine/Source/Editor/PropertyEditor/Private/UserInterface/PropertyEditor/SPropertyMenuAssetPicker.cpp
@@ -2,7 +2,6 @@
 
 #include "UserInterface/PropertyEditor/SPropertyMenuAssetPicker.h"
 
-#include "AssetRegistry/AssetRegistryModule.h"
 #include "AssetToolsModule.h"
 #include "ContentBrowserModule.h"
 #include "Editor.h"
@@ -16,6 +15,7 @@
 #include "PropertyEditorClipboard.h"
 #include "PropertyEditorCopyPastePrivate.h"
 #include "Styling/SlateIconFinder.h"
+#include "UserInterface/PropertyEditor/SPropertyEditorAsset.h"
 #include "UserInterface/PropertyEditor/PropertyEditorAssetConstants.h"
 #include "Widgets/Layout/SBox.h"
 
@@ -240,7 +240,7 @@ void SPropertyMenuAssetPicker::OnPaste()
 	FString DestPath;
 	FPropertyEditorClipboard::ClipboardPaste(DestPath);
 
-	PasteFromText(TEXT(""), DestPath);
+	OnPasteFromText(TEXT(""), DestPath, {});
 }
 
 void SPropertyMenuAssetPicker::OnPasteFromText(
@@ -249,21 +249,22 @@ void SPropertyMenuAssetPicker::OnPasteFromText(
 	const TOptional<FGuid>& InOperationId)
 {
 	// Naive check done elsewhere, guard with proper check here 
-	if (CanPasteFromText(InTag, InText))
+	FAssetData AssetData;
+	if (CanPasteFromText(InTag, InText, AssetData))
 	{
-		PasteFromText(InTag, InText);
+		PasteFromText(InTag, AssetData);
 	}
 }
 
-void SPropertyMenuAssetPicker::PasteFromText(const FString& InTag, const FString& InText)
+void SPropertyMenuAssetPicker::PasteFromText(const FString& InTag, const FAssetData& InAssetData)
 {
-	if(InText == TEXT("None"))
+	if (!InAssetData.IsValid())
 	{
 		SetValue(nullptr);
 	}
 	else
 	{
-		UObject* Object = LoadObject<UObject>(nullptr, *InText);
+		UObject* Object = InAssetData.GetAsset();
 		bool PassesAllowedClassesFilter = true;
 		if (Object && AllowedClasses.Num())
 		{
@@ -281,13 +282,11 @@ void SPropertyMenuAssetPicker::PasteFromText(const FString& InTag, const FString
 		}
 		if( Object && PassesAllowedClassesFilter )
 		{
-			FAssetData ObjectAssetData(Object);
-
 			// Check against custom asset filter
 			if (!OnShouldFilterAsset.IsBound()
-				|| !OnShouldFilterAsset.Execute(ObjectAssetData))
+				|| !OnShouldFilterAsset.Execute(InAssetData))
 			{
-				SetValue(ObjectAssetData);
+				SetValue(InAssetData);
 			}
 		}
 	}
@@ -299,10 +298,11 @@ bool SPropertyMenuAssetPicker::CanPaste()
 	FString ClipboardText;
 	FPropertyEditorClipboard::ClipboardPaste(ClipboardText);
 	
-	return CanPasteFromText(TEXT(""), ClipboardText);
+	FAssetData AssetData;
+	return CanPasteFromText(TEXT(""), ClipboardText, AssetData);
 }
 
-bool SPropertyMenuAssetPicker::CanPasteFromText(const FString& InTag, const FString& InText) const
+bool SPropertyMenuAssetPicker::CanPasteFromText(const FString& InTag, const FString& InText, FAssetData& OutAssetData) const
 {
 	if (!bAllowCopyPaste)
 	{
@@ -314,26 +314,7 @@ bool SPropertyMenuAssetPicker::CanPasteFromText(const FString& InTag, const FStr
 		return false;
 	}
 
-	FString Class;
-	FString PossibleObjectPath = InText;
-	if( InText.Split( TEXT("'"), &Class, &PossibleObjectPath, ESearchCase::CaseSensitive) )
-	{
-		// Remove the last item
-		PossibleObjectPath.LeftChopInline( 1, EAllowShrinking::No );
-	}
-
-	bool bCanPaste = false;
-	if( PossibleObjectPath == TEXT("None") )
-	{
-		bCanPaste = true;
-	}
-	else
-	{
-		const FAssetRegistryModule& AssetRegistryModule = FModuleManager::Get().LoadModuleChecked<FAssetRegistryModule>(TEXT("AssetRegistry"));
-		bCanPaste = PossibleObjectPath.Len() < NAME_SIZE && AssetRegistryModule.Get().GetAssetByObjectPath( FSoftObjectPath(PossibleObjectPath) ).IsValid();
-	}
-
-	return bCanPaste;
+	return SPropertyEditorAsset::ParseAssetText(InText, OutAssetData);
 }
 
 void SPropertyMenuAssetPicker::OnClear()
diff --git a/Engine/Source/Editor/PropertyEditor/Private/UserInterface/PropertyEditor/SPropertyMenuAssetPicker.h b/Engine/Source/Editor/PropertyEditor/Private/UserInterface/PropertyEditor/SPropertyMenuAssetPicker.h
--- a/Engine/Source/Editor/PropertyEditor/Private/UserInterface/PropertyEditor/SPropertyMenuAssetPicker.h
+++ b/Engine/Source/Editor/PropertyEditor/Private/UserInterface/PropertyEditor/SPropertyMenuAssetPicker.h
@@ -56,7 +56,7 @@ private:
 	/**
 	 * @return True if the (optionally tagged) input contents can be pasted
 	 */
-	bool CanPasteFromText(const FString& InTag, const FString& InText) const;
+	bool CanPasteFromText(const FString& InTag, const FString& InText, FAssetData& OutAssetData) const;
 
 	/**
 	 * Delegate handling pasting an optionally tagged text snippet
@@ -66,7 +66,7 @@ private:
 	/**
 	 * Handle pasting an optionally tagged text snippet
 	 */
-	void PasteFromText(const FString& InTag, const FString& InText);
+	void PasteFromText(const FString& InTag, const FAssetData& InAssetData);
 
 	/**
 	 * @return True of the current clipboard contents can be pasted
