Fixed false positive static analysis error

diff --git a/Engine/Source/Editor/PropertyEditor/Private/PropertyHandleImpl.cpp b/Engine/Source/Editor/PropertyEditor/Private/PropertyHandleImpl.cpp
--- a/Engine/Source/Editor/PropertyEditor/Private/PropertyHandleImpl.cpp
+++ b/Engine/Source/Editor/PropertyEditor/Private/PropertyHandleImpl.cpp
@@ -4541,7 +4541,7 @@ FPropertyAccess::Result FPropertyHandleObject::SetObjectValueFromSelection()
 			if (SelectedClass)
 			{
 				// MetaData class filter
-				if (NodeProperty && NodeProperty->GetOwnerProperty()->HasMetaData("GetClassFilter"))
+				if (NodeProperty->GetOwnerProperty()->HasMetaData("GetClassFilter"))
 				{
 					const FString& GetClassFilterFunctionName = NodeProperty->GetOwnerProperty()->GetMetaData("GetClassFilter");
 					if (!GetClassFilterFunctionName.IsEmpty())
