Fix Structure Details View did not restore the expansion state of customized properties

diff --git a/Engine/Source/Editor/PropertyEditor/Private/SDetailsView.cpp b/Engine/Source/Editor/PropertyEditor/Private/SDetailsView.cpp
--- a/Engine/Source/Editor/PropertyEditor/Private/SDetailsView.cpp
+++ b/Engine/Source/Editor/PropertyEditor/Private/SDetailsView.cpp
@@ -988,19 +988,6 @@ void SDetailsView::PreSetObject(int32 InNewNumObjects)
 	SelectedObjects.Empty();
 }
 
-
-void SDetailsView::GetExpansionState(UStruct* BaseStruct, FStringPrefixTree& OutPrefixTree) const
-{
-	//while a valid class, and we're either the same as the base class (for multiple actors being selected and base class is AActor) OR we're not down to AActor yet)
-	TArray<FString> DetailPropertyExpansionStrings;
-	for (UStruct* Struct = BaseStruct; Struct && ((BaseStruct == Struct) || (Struct != AActor::StaticClass())); Struct = Struct->GetSuperStruct())
-	{
-		GConfig->GetSingleLineArray(TEXT("DetailPropertyExpansion"), *Struct->GetName(), DetailPropertyExpansionStrings, GEditorPerProjectIni);
-	}
-
-	OutPrefixTree.InsertAll(DetailPropertyExpansionStrings);
-}
-
 /** Called at the end of SetObjectArray after we change the objects being observed */
 void SDetailsView::PostSetObject(const TArray<FDetailsViewObjectRoot>& Roots)
 {
@@ -1078,39 +1065,7 @@ void SDetailsView::PostSetObject(const TArray<FDetailsViewObjectRoot>& Roots)
 		RootPropertyNode->InitNode( InitParams );
 	}
 
-	// Restore expanded items is called twice. The first time handles the existing property nodes,
-	// and second time after handling property customization as it might add new items that need to be expanded too.
-
-	// Collect the expansion states here, since the first call to RestoreAllExpandedItems() might overwrite the values and that may prevent some customization expansions to be saved.
-	TMap<UStruct*, FStringPrefixTree> ExpansionStates;
-	for (TSharedPtr<FComplexPropertyNode>& RootPropertyNode : GetRootNodes())
-	{
-		check(RootPropertyNode.IsValid());
-		UStruct* BaseStruct = RootPropertyNode->FindComplexParent()->GetBaseStructure();
-		FStringPrefixTree& PrefixTree = ExpansionStates.Add(BaseStruct);
-		GetExpansionState(BaseStruct, PrefixTree);
-	}
-
-	for (FDetailLayoutData& LayoutData : DetailLayouts)
-	{
-		FRootPropertyNodeList& ExternalRootPropertyNodes = LayoutData.DetailLayout->GetExternalRootPropertyNodes();
-		for (TSharedPtr<FComplexPropertyNode>& ExternalRootPropertyNode : ExternalRootPropertyNodes)
-		{
-			check(ExternalRootPropertyNode.IsValid());
-			UStruct* BaseStruct = ExternalRootPropertyNode->FindComplexParent()->GetBaseStructure();
-			FStringPrefixTree& PrefixTree = ExpansionStates.Add(BaseStruct);
-			GetExpansionState(BaseStruct, PrefixTree);
-		}
-	}
-
-	// Restore existing nodes
-	RestoreAllExpandedItems(&ExpansionStates);
-
-	UpdatePropertyMaps();
-
-	// Restore nodes generated by property customization
-	RestoreAllExpandedItems(&ExpansionStates);
-
+	UpdatePropertyMapsWithItemExpansions();
 	UpdateFilteredDetails();
 }
 
diff --git a/Engine/Source/Editor/PropertyEditor/Private/SDetailsView.h b/Engine/Source/Editor/PropertyEditor/Private/SDetailsView.h
--- a/Engine/Source/Editor/PropertyEditor/Private/SDetailsView.h
+++ b/Engine/Source/Editor/PropertyEditor/Private/SDetailsView.h
@@ -238,8 +238,6 @@ private:
 	 */
 	void RefreshDisplayManager();
 
-	void GetExpansionState(UStruct* BaseStruct, FStringPrefixTree& OutPrefixTree) const;
-
 	/**
 	* The @code FDetailsViewStyle& @endcode shared pointer which provides style
 	* information for this details view.
diff --git a/Engine/Source/Editor/PropertyEditor/Private/SDetailsViewBase.cpp b/Engine/Source/Editor/PropertyEditor/Private/SDetailsViewBase.cpp
--- a/Engine/Source/Editor/PropertyEditor/Private/SDetailsViewBase.cpp
+++ b/Engine/Source/Editor/PropertyEditor/Private/SDetailsViewBase.cpp
@@ -640,6 +640,54 @@ void SDetailsViewBase::UpdatePropertyMaps()
 	}
 }
 
+void SDetailsViewBase::GetExpansionState(UStruct* BaseStruct, FStringPrefixTree& OutPrefixTree) const
+{
+	//while a valid class, and we're either the same as the base class (for multiple actors being selected and base class is AActor) OR we're not down to AActor yet)
+	TArray<FString> DetailPropertyExpansionStrings;
+	for (UStruct* Struct = BaseStruct; Struct && ((BaseStruct == Struct) || (Struct != AActor::StaticClass())); Struct = Struct->GetSuperStruct())
+	{
+		GConfig->GetSingleLineArray(TEXT("DetailPropertyExpansion"), *Struct->GetName(), DetailPropertyExpansionStrings, GEditorPerProjectIni);
+	}
+
+	OutPrefixTree.InsertAll(DetailPropertyExpansionStrings);
+}
+
+void SDetailsViewBase::UpdatePropertyMapsWithItemExpansions()
+{
+	// Restore expanded items is called twice. The first time handles the existing property nodes,
+	// and second time after handling property customization as it might add new items that need to be expanded too.
+
+	// Collect the expansion states here, since the first call to RestoreAllExpandedItems() might overwrite the values and that may prevent some customization expansions to be saved.
+	TMap<UStruct*, FStringPrefixTree> ExpansionStates;
+	for (TSharedPtr<FComplexPropertyNode>& RootPropertyNode : GetRootNodes())
+	{
+		check(RootPropertyNode.IsValid());
+		UStruct* BaseStruct = RootPropertyNode->FindComplexParent()->GetBaseStructure();
+		FStringPrefixTree& PrefixTree = ExpansionStates.Add(BaseStruct);
+		GetExpansionState(BaseStruct, PrefixTree);
+	}
+
+	for (const FDetailLayoutData& LayoutData : DetailLayouts)
+	{
+		FRootPropertyNodeList& ExternalRootPropertyNodes = LayoutData.DetailLayout->GetExternalRootPropertyNodes();
+		for (TSharedPtr<FComplexPropertyNode>& ExternalRootPropertyNode : ExternalRootPropertyNodes)
+		{
+			check(ExternalRootPropertyNode.IsValid());
+			UStruct* BaseStruct = ExternalRootPropertyNode->FindComplexParent()->GetBaseStructure();
+			FStringPrefixTree& PrefixTree = ExpansionStates.Add(BaseStruct);
+			GetExpansionState(BaseStruct, PrefixTree);
+		}
+	}
+
+	// Restore existing nodes
+	RestoreAllExpandedItems(&ExpansionStates);
+
+	UpdatePropertyMaps();
+
+	// Restore nodes generated by property customization
+	RestoreAllExpandedItems(&ExpansionStates);
+}
+
 void SDetailsViewBase::UpdateSinglePropertyMap(TSharedPtr<FComplexPropertyNode> InRootPropertyNode, FDetailLayoutData& LayoutData, bool bIsExternal)
 {
 	// Reset everything
diff --git a/Engine/Source/Editor/PropertyEditor/Private/SDetailsViewBase.h b/Engine/Source/Editor/PropertyEditor/Private/SDetailsViewBase.h
--- a/Engine/Source/Editor/PropertyEditor/Private/SDetailsViewBase.h
+++ b/Engine/Source/Editor/PropertyEditor/Private/SDetailsViewBase.h
@@ -286,6 +286,11 @@ protected:
 
 	virtual void CustomUpdatePropertyMap(TSharedPtr<FDetailLayoutBuilderImpl>& InDetailLayout) {}
 
+	void GetExpansionState(UStruct* BaseStruct, FStringPrefixTree& OutPrefixTree) const;
+
+	/** Updates the property map while keeping the expansion state of all items */
+	void UpdatePropertyMapsWithItemExpansions();
+
 	/** Called to get the visibility of the tree view */
 	EVisibility GetTreeVisibility() const;
 
diff --git a/Engine/Source/Editor/PropertyEditor/Private/SStructureDetailsView.cpp b/Engine/Source/Editor/PropertyEditor/Private/SStructureDetailsView.cpp
--- a/Engine/Source/Editor/PropertyEditor/Private/SStructureDetailsView.cpp
+++ b/Engine/Source/Editor/PropertyEditor/Private/SStructureDetailsView.cpp
@@ -282,10 +282,7 @@ void SStructureDetailsView::SetStructureProvider(TSharedPtr<IStructureDataProvid
 	RootNode->InitNode(InitParams);
 	RootNode->SetDisplayNameOverride(CustomName);
 
-	RestoreExpandedItems(RootNode.ToSharedRef());
-
-	UpdatePropertyMaps();
-
+	UpdatePropertyMapsWithItemExpansions();
 	UpdateFilteredDetails();
 }
 
