Actor Picker: Improve displayed name for actors that reside in Level Instances.

diff --git a/Engine/Source/Editor/PropertyEditor/Private/UserInterface/PropertyEditor/SPropertyEditorAsset.cpp b/Engine/Source/Editor/PropertyEditor/Private/UserInterface/PropertyEditor/SPropertyEditorAsset.cpp
--- a/Engine/Source/Editor/PropertyEditor/Private/UserInterface/PropertyEditor/SPropertyEditorAsset.cpp
+++ b/Engine/Source/Editor/PropertyEditor/Private/UserInterface/PropertyEditor/SPropertyEditorAsset.cpp
@@ -134,12 +134,11 @@ static bool GetEditorPathOwnerFromPropertyHandle(const TSharedPtr<IPropertyHandl
 	return true;
 }
 
-static FString GetActorEditorPathLabel(const AActor* InActor)
+static FString GetActorEditorPathTooltip(const AActor* InActor)
 {
 	check(InActor);
 	
 	TArray<FString> EditorPathOwners;
-	EditorPathOwners.Add(InActor->GetActorLabel());
 	const UObject* Context = InActor;
 
 	while (UObject* EditorPathOwner = FEditorPathHelper::GetEditorPathOwner(Context))
@@ -155,14 +154,20 @@ static FString GetActorEditorPathLabel(const AActor* InActor)
 		Context = EditorPathOwner;
 	}
 
+	// If there are no owners, we don't want the tooltip, as the actor label is already visible in the field
+	if (EditorPathOwners.IsEmpty())
+	{
+		return FString();
+	}
+
 	TStringBuilder<256> LabelBuilder;
-	LabelBuilder.Append(EditorPathOwners[EditorPathOwners.Num()-1]);
-	for (int32 i = EditorPathOwners.Num() - 2; i >= 0; --i)
+	for (int32 i = EditorPathOwners.Num() - 1; i >= 0; --i)
 	{
-		LabelBuilder.Append(TEXT("."));
 		LabelBuilder.Append(EditorPathOwners[i]);
+		LabelBuilder.Append(TEXT(" \u2192 "));
 	}
 
+	LabelBuilder.Append(InActor->GetActorLabel());
 	return LabelBuilder.ToString();
 }
 
@@ -1071,7 +1076,7 @@ FText SPropertyEditorAsset::OnGetAssetName() const
 
 				if (Actor)
 				{
-					Name = FText::AsCultureInvariant(GetActorEditorPathLabel(Actor));
+					Name = FText::AsCultureInvariant(Actor->GetActorLabel());
 				}
 				else
 				{
@@ -1148,6 +1153,24 @@ FText SPropertyEditorAsset::OnGetToolTip() const
 			else
 			{
 				ToolTipText = FText::Format(LOCTEXT("GoodActorReference", "Reference to Actor ID '{Actor}'"), Args);
+
+				if (const AActor* Actor = Cast<AActor>(Value.Object))
+				{
+					FString Path = GetActorEditorPathTooltip(Actor);
+					if (!Path.IsEmpty())
+					{
+						FText OwnerPath = FText::Format(
+							LOCTEXT("ReferenceOwnerPath", "Actor path: {0}"),
+							FText::AsCultureInvariant(MoveTemp(Path))
+						);
+
+						ToolTipText = FText::Format(
+							INVTEXT("{0}\n\n{1}"),
+							MoveTemp(OwnerPath),
+							MoveTemp(ToolTipText)
+						);
+					}
+				}
 			}
 		}
 		else if( Value.Object != nullptr )
