Actor Picker: Fix Browse button did not work when an LI actor was referenced by it.

diff --git a/Engine/Source/Editor/PropertyEditor/Private/Presentation/PropertyEditor/PropertyEditor.cpp b/Engine/Source/Editor/PropertyEditor/Private/Presentation/PropertyEditor/PropertyEditor.cpp
--- a/Engine/Source/Editor/PropertyEditor/Private/Presentation/PropertyEditor/PropertyEditor.cpp
+++ b/Engine/Source/Editor/PropertyEditor/Private/Presentation/PropertyEditor/PropertyEditor.cpp
@@ -8,11 +8,14 @@
 #include "ObjectPropertyNode.h"
 
 #include "SceneOutlinerFilters.h"
+#include "Selection.h"
 #include "IDetailPropertyRow.h"
 #include "PropertyCustomizationHelpers.h"
 #include "PropertyEditorHelpers.h"
 #include "Editor.h"
 #include "EditorClassUtils.h"
+#include "Editor/UnrealEdEngine.h"
+#include "Elements/Framework/EngineElementsLibrary.h"
 #include "Kismet2/KismetEditorUtilities.h"
 #include "IConfigEditorModule.h"
 #include "PropertyNode.h"
@@ -725,6 +728,25 @@ void FPropertyEditor::SyncToObjectsInNode( const TWeakPtr< FPropertyNode >& Weak
 		// If a list of addresses is provided, GetReadAddress may still return false but we can operate on the property addresses even if they have different values.
 		check( ReadAddresses.Num() > 0 );
 
+		TArray<UObject*> Objects;
+
+		// If this is a Soft Object Ptr, then let's resolve the path via the soft path,
+		// as it'll handle EditorPaths (e.g. LI actors), which Static(Find|Load)Object does not support.
+		if (const FSoftObjectProperty* SoftObjectProperty = CastField<FSoftObjectProperty>(NodeProperty))
+		{
+			for (int32 AddrIndex = 0; AddrIndex < ReadAddresses.Num(); ++AddrIndex)
+			{
+				const uint8* Address = ReadAddresses.GetAddress(AddrIndex);
+				if (UObject* LoadedObject = SoftObjectProperty->GetObjectPropertyValue(Address))
+				{
+					Objects.Add(LoadedObject);
+				}
+			}
+
+			// No need to check the FString version now.
+			ReadAddresses.Reset();
+		}
+
 		// Create a list of object names.
 		TArray<FString> ObjectNames;
 		ObjectNames.Empty( ReadAddresses.Num() );
@@ -742,7 +764,6 @@ void FPropertyEditor::SyncToObjectsInNode( const TWeakPtr< FPropertyNode >& Weak
 
 
 		// Create a list of objects to sync the generic browser to.
-		TArray<UObject*> Objects;
 		for ( int32 ObjectIndex = 0 ; ObjectIndex < ObjectNames.Num() ; ++ObjectIndex )
 		{
 			UObject* Object = nullptr;
@@ -782,14 +803,31 @@ void FPropertyEditor::SyncToObjectsInNode( const TWeakPtr< FPropertyNode >& Weak
 
 			if (Actor->GetLevel())
 			{
-				TArray<AActor*> Actors;
-				Actors.Add(Actor);
+				const USelection* ActorSelection = GEditor->GetSelectedActors();
+				if (UTypedElementSelectionSet* SelectionSet = ActorSelection->GetElementSelectionSet())
+				{
+					const FTypedElementSelectionOptions SelectionOptions = FTypedElementSelectionOptions()
+						.SetAllowHidden(true)
+						.SetWarnIfLocked(true)
+						.SetAllowLegacyNotifications(false)
+						.SetAllowSubRootSelection(true);
 
-				GEditor->SelectNone(/*bNoteSelectionChange=*/false, /*bDeselectBSPSurfs=*/true);
-				GEditor->SelectActor(Actor, /*InSelected=*/true, /*bNotify=*/true, /*bSelectEvenIfHidden=*/true);
+					FTypedElementHandle ActorHandle = UEngineElementsLibrary::AcquireEditorActorElementHandle(Actor);
+					const bool bSelectionChanged = SelectionSet->SetSelection(MakeArrayView(&ActorHandle, 1), SelectionOptions);
+
+					if (bSelectionChanged)
+					{
+						SelectionSet->NotifyPendingChanges();
+					}
+					else
+					{
+						// Reset the property windows, in case something has changed since previous selection
+						GUnrealEd->UpdateFloatingPropertyWindows();
+					}
+				}
 
 				// Jump to the location of the actor
-				GEditor->MoveViewportCamerasToActor(Actors, /*bActiveViewportOnly=*/false);
+				GEditor->MoveViewportCamerasToActor(*Actor, /*bActiveViewportOnly=*/false);
 			}
 		}
 		else if ( Objects.Num() > 0 )
diff --git a/Engine/Source/Editor/PropertyEditor/PropertyEditor.Build.cs b/Engine/Source/Editor/PropertyEditor/PropertyEditor.Build.cs
--- a/Engine/Source/Editor/PropertyEditor/PropertyEditor.Build.cs
+++ b/Engine/Source/Editor/PropertyEditor/PropertyEditor.Build.cs
@@ -48,7 +48,9 @@ public class PropertyEditor : ModuleRules
 				"ToolWidgets",
 				"WidgetRegistration",
 				"Json",
-				"ToolMenus"
+				"ToolMenus",
+				"TypedElementFramework",
+				"TypedElementRuntime",
 			}
         );
 
