IDetailsView implements IScrollableWidget, allowing it to be embedded in SScrollBorder etc.

diff --git a/Engine/Source/Editor/PropertyEditor/Private/SDetailsViewBase.cpp b/Engine/Source/Editor/PropertyEditor/Private/SDetailsViewBase.cpp
--- a/Engine/Source/Editor/PropertyEditor/Private/SDetailsViewBase.cpp
+++ b/Engine/Source/Editor/PropertyEditor/Private/SDetailsViewBase.cpp
@@ -1948,4 +1948,31 @@ TSharedPtr<FDetailsDisplayManager> SDetailsViewBase::GetDisplayManager()
 		DisplayManager = MakeShared<FDetailsDisplayManager>();
 	}
 	return DisplayManager;
-}
\ No newline at end of file
+}
+
+FVector2D SDetailsViewBase::GetScrollDistance()
+{
+	if (!DetailTree.IsValid())
+	{
+		return FVector2D::ZeroVector;
+	}
+
+	return DetailTree->GetScrollDistance();
+
+}
+
+FVector2D SDetailsViewBase::GetScrollDistanceRemaining()
+{
+	if (!DetailTree.IsValid())
+	{
+		return FVector2D::ZeroVector;
+	}
+
+	return DetailTree->GetScrollDistanceRemaining();
+
+}
+
+TSharedRef<SWidget> SDetailsViewBase::GetScrollWidget()
+{
+	return SharedThis(this);
+}
diff --git a/Engine/Source/Editor/PropertyEditor/Private/SDetailsViewBase.h b/Engine/Source/Editor/PropertyEditor/Private/SDetailsViewBase.h
--- a/Engine/Source/Editor/PropertyEditor/Private/SDetailsViewBase.h
+++ b/Engine/Source/Editor/PropertyEditor/Private/SDetailsViewBase.h
@@ -200,6 +200,12 @@ public:
 	virtual void Tick( const FGeometry& AllottedGeometry, const double InCurrentTime, const float InDeltaTime ) override;
 	// End of SWidget interface
 
+	// IScrollableWidget interface
+	virtual FVector2D GetScrollDistance() override;
+	virtual FVector2D GetScrollDistanceRemaining() override;
+	virtual TSharedRef<SWidget> GetScrollWidget() override;
+	// End of IScrollableWidget interface
+
 protected:
 	/**
 	 * Called when a color property is changed from a color picker
diff --git a/Engine/Source/Editor/PropertyEditor/Public/IDetailsView.h b/Engine/Source/Editor/PropertyEditor/Public/IDetailsView.h
--- a/Engine/Source/Editor/PropertyEditor/Public/IDetailsView.h
+++ b/Engine/Source/Editor/PropertyEditor/Public/IDetailsView.h
@@ -7,6 +7,7 @@
 #include "Widgets/SCompoundWidget.h"
 #include "PropertyEditorDelegates.h"
 #include "Framework/Commands/UICommandList.h"
+#include "Framework/Layout/IScrollableWidget.h"
 #include "DetailsViewArgs.h"
 
 class FDetailTreeNode;
@@ -31,7 +32,7 @@ DECLARE_DELEGATE_RetVal_OneParam(bool, FOnValidateDetailsViewPropertyNodes, cons
 /**
  * Interface class for all detail views
  */
-class IDetailsView : public SCompoundWidget
+class IDetailsView : public SCompoundWidget, public IScrollableWidget
 {
 public:
 	/** Sets the callback for when the property view changes */
@@ -319,4 +320,10 @@ public:
    	*/
    	virtual bool IsDefaultStyle() const = 0;
 
+
+	// IScrollableWidget interface
+	virtual FVector2D GetScrollDistance() = 0;
+	virtual FVector2D GetScrollDistanceRemaining() = 0;
+	virtual TSharedRef<SWidget> GetScrollWidget() = 0;
+	// End of IScrollableWidget interface
 };
