Implementing OnNotifyPreChangingProperty event to class IPropertyRowGenerator.

diff --git a/Engine/Source/Editor/PropertyEditor/Private/Presentation/PropertyTable/PropertyTable.h b/Engine/Source/Editor/PropertyEditor/Private/Presentation/PropertyTable/PropertyTable.h
--- a/Engine/Source/Editor/PropertyEditor/Private/Presentation/PropertyTable/PropertyTable.h
+++ b/Engine/Source/Editor/PropertyEditor/Private/Presentation/PropertyTable/PropertyTable.h
@@ -29,6 +29,7 @@ public:
 	virtual void CreateColorPickerWindow( const TSharedRef< class FPropertyEditor >& PropertyEditor, bool bUseAlpha ) const override;
 	virtual void EnqueueDeferredAction( FSimpleDelegate DeferredAction ) override;
 	virtual TSharedPtr<class FAssetThumbnailPool> GetThumbnailPool() const override;
+	virtual void NotifyStartedChangingProperties(const FPropertyChangedEvent& PropertyChangedEvent) override {}
 	virtual void NotifyFinishedChangingProperties(const FPropertyChangedEvent& PropertyChangedEvent) override {}
 	virtual const TArray<TSharedRef<class IClassViewerFilter>>& GetClassViewerFilters() const override;
 
diff --git a/Engine/Source/Editor/PropertyEditor/Private/PropertyHandleImpl.cpp b/Engine/Source/Editor/PropertyEditor/Private/PropertyHandleImpl.cpp
--- a/Engine/Source/Editor/PropertyEditor/Private/PropertyHandleImpl.cpp
+++ b/Engine/Source/Editor/PropertyEditor/Private/PropertyHandleImpl.cpp
@@ -311,6 +311,12 @@
 
 		const bool bTransactable = (Flags & EPropertyValueSetFlags::NotTransactable) == 0;
 		const bool bFinished = ( Flags & EPropertyValueSetFlags::InteractiveChange) == 0;
+
+		EPropertyChangeType::Type ChangeEventFlags = bFinished ? EPropertyChangeType::ValueSet : EPropertyChangeType::Interactive;
+		if (Flags & EPropertyValueSetFlags::ResetToDefault)
+		{
+			ChangeEventFlags = ChangeEventFlags | EPropertyChangeType::ResetToDefault;
+		}
 		
 		// List of top level objects sent to the PropertyChangedEvent
 		TArray<const UObject*> TopLevelObjects;
@@ -328,6 +334,11 @@
 
 			UObject* CurObject = Cur.Object;
 
+			if (CurObject)
+			{
+				TopLevelObjects.Add(CurObject);
+			}
+			
 			const FString& NewValue = InValues[ObjectIndex];
 
 			// Cache the value of the property before modifying it.
@@ -464,6 +475,10 @@
 				Property->ExportText_Direct(PreviousContainerValue, Addr, Addr, nullptr, 0);
 			}
 
+			//add on array index so we can tell which entry just changed
+			ArrayIndicesPerObject.Add(TMap<FString,int32>());
+			FPropertyValueImpl::GenerateArrayIndexMapToObjectNode(ArrayIndicesPerObject[ObjectIndex], InPropertyNode);
+
 			// Check if we need to call PreEditChange on all objects.
 			// Remove quotes from the original value because FName properties  
 			// are wrapped in quotes before getting here. This causes the 
@@ -488,6 +503,13 @@
 
 				InPropertyNode->NotifyPreChange(NodeProperty, NotifyHook);
 
+				if (PropertyUtilities.IsValid() && !bInteractiveChangeInProgress)
+				{
+					FPropertyChangedEvent StartedChangingPropertiesEvent(NodeProperty, ChangeEventFlags, MakeArrayView(TopLevelObjects));
+					StartedChangingPropertiesEvent.SetArrayIndexPerObject(ArrayIndicesPerObject);
+					PropertyUtilities.Pin()->NotifyStartedChangingProperties(StartedChangingPropertiesEvent);
+				}
+
 				bInteractiveChangeInProgress = (Flags & EPropertyValueSetFlags::InteractiveChange) != 0;
 			}
 
@@ -556,8 +578,6 @@
 						InPropertyNode->PropagatePropertyChange(CurObject, *NewValue, PreviousContainerValue.IsEmpty() ? PreviousValue : PreviousContainerValue);
 					}
 				}
-
-				TopLevelObjects.Add(CurObject);
 			}
 
 			// If the values before and after setting the property differ, mark the object dirty.
@@ -582,17 +602,8 @@
 					SetHelper.Rehash();
 				}
 			}
-
-			//add on array index so we can tell which entry just changed
-			ArrayIndicesPerObject.Add(TMap<FString,int32>());
-			FPropertyValueImpl::GenerateArrayIndexMapToObjectNode(ArrayIndicesPerObject[ObjectIndex], InPropertyNode);
 		}
-
-		EPropertyChangeType::Type ChangeEventFlags = bFinished ? EPropertyChangeType::ValueSet : EPropertyChangeType::Interactive;
-		if (Flags & EPropertyValueSetFlags::ResetToDefault)
-		{
-			ChangeEventFlags = ChangeEventFlags | EPropertyChangeType::ResetToDefault;
-		}
+		
 		FPropertyChangedEvent ChangeEvent(NodeProperty, ChangeEventFlags, MakeArrayView(TopLevelObjects));
 		ChangeEvent.SetArrayIndexPerObject(ArrayIndicesPerObject);
 
@@ -1130,6 +1141,7 @@
 			// If we added a new Map entry we need to rebuild children
 			bool bAddedMapEntry = false;
 
+			TArray< TMap<FString, int32> > PreAddArrayIndicesPerObject;
 			TArray< TMap<FString, int32> > ArrayIndicesPerObject;
 			TArray< TArray< UObject* > > AffectedInstancesPerObject;
 			AffectedInstancesPerObject.SetNum(ReadAddresses.Num());
@@ -1146,22 +1158,31 @@
 			FMapProperty* Map = CastField<FMapProperty>(NodeProperty);
 
 			check(Array || Set || Map);
-			
+
 			for ( int32 i = 0 ; i < ReadAddresses.Num() ; ++i )
 			{
 				void* Addr = ReadAddresses.GetAddress(i);
 				if ( Addr )
 				{
 					bHasChanges = true;
-
+					UObject* Obj = ObjectNode ? ObjectNode->GetUObject(i) : nullptr;
+					
 					if (!FApp::IsGame())
 					{
-						UObject* Obj = ObjectNode ? ObjectNode->GetUObject(i) : nullptr;
 						if (IsTemplate(Obj))
 						{
 							PropertyNodePin->GatherInstancesAffectedByContainerPropertyChange(Obj, Addr, EPropertyArrayChangeType::Add, AffectedInstancesPerObject[i]);
 						}
 					}
+
+					if (Obj)
+					{
+						TopLevelObjects.Add(Obj);
+					}
+					
+					//add on array index so we can tell which entry just changed
+					PreAddArrayIndicesPerObject.Add(TMap<FString, int32>());
+					FPropertyValueImpl::GenerateArrayIndexMapToObjectNode(PreAddArrayIndicesPerObject[i], PropertyNodePin.Get());
 				}
 			}
 
@@ -1175,6 +1196,13 @@
 
 				// send the PreEditChange notification to all selected objects
 				PropertyNodePin->NotifyPreChange(NodeProperty, NotifyHook, AllAffectedInstances);
+				
+				if (PropertyUtilities.IsValid() && !bInteractiveChangeInProgress)
+				{
+					FPropertyChangedEvent StartedChangingPropertiesEvent(NodeProperty, EPropertyChangeType::ArrayAdd, MakeArrayView(TopLevelObjects));
+					StartedChangingPropertiesEvent.SetArrayIndexPerObject(PreAddArrayIndicesPerObject);
+					PropertyUtilities.Pin()->NotifyStartedChangingProperties(StartedChangingPropertiesEvent);
+				}
 
 				for (int32 i = 0 ; i < ReadAddresses.Num() ; ++i )
 				{
@@ -1342,14 +1370,23 @@
 				{
 					bHasChanges = true;
 
+					UObject* Obj = ObjectNode ? ObjectNode->GetUObject(i) : nullptr;
 					if (!FApp::IsGame())
 					{
-						UObject* Obj = ObjectNode ? ObjectNode->GetUObject(i) : nullptr;
 						if (IsTemplate(Obj))
 						{
 							PropertyNodePin->GatherInstancesAffectedByContainerPropertyChange(Obj, Addr, EPropertyArrayChangeType::Clear, AffectedInstancesPerObject[i]);
 						}
 					}
+
+					if (Obj)
+					{
+						TopLevelObjects.Add(Obj);
+					}
+					
+					// Add on array index so we can tell which entry just changed
+					ArrayIndicesPerObject.Add(TMap<FString, int32>());
+					FPropertyValueImpl::GenerateArrayIndexMapToObjectNode(ArrayIndicesPerObject[i], PropertyNodePin.Get());
 				}
 			}
 
@@ -1360,10 +1397,17 @@
 				{
 					AllAffectedInstances.Append(AffectedInstances);
 				}
-
+				
 				// Send the PreEditChange notification to all selected objects
 				PropertyNodePin->NotifyPreChange(NodeProperty, NotifyHook, AllAffectedInstances);
 
+				if (PropertyUtilities.IsValid() && !bInteractiveChangeInProgress)
+				{
+					FPropertyChangedEvent StartedChangingPropertiesEvent(NodeProperty, EPropertyChangeType::ArrayClear, MakeArrayView(TopLevelObjects));
+					StartedChangingPropertiesEvent.SetArrayIndexPerObject(ArrayIndicesPerObject);
+					PropertyUtilities.Pin()->NotifyStartedChangingProperties(StartedChangingPropertiesEvent);
+				}
+
 				for ( int32 i = 0 ; i < ReadAddresses.Num() ; ++i )
 				{
 					void* Addr = ReadAddresses.GetAddress(i);
@@ -1376,14 +1420,8 @@
 							{
 								PropertyNodePin->PropagateContainerPropertyChange(Obj, Addr, AffectedInstancesPerObject[i], EPropertyArrayChangeType::Clear, -1);
 							}
-
-							TopLevelObjects.Add(Obj);
 						}
 
-						// Add on array index so we can tell which entry just changed
-						ArrayIndicesPerObject.Add(TMap<FString, int32>());
-						FPropertyValueImpl::GenerateArrayIndexMapToObjectNode(ArrayIndicesPerObject[i], PropertyNodePin.Get());
-
 						if (ArrayProperty)
 						{
 							FScriptArrayHelper ArrayHelper(ArrayProperty, Addr);
@@ -1504,6 +1542,9 @@
 		// Begin a property edit transaction.
 		FScopedTransaction Transaction( NSLOCTEXT("UnrealEd", "InsertChild", "Insert Child") );
 
+		// List of top level objects sent to the PropertyChangedEvent
+		TArray<const UObject*> TopLevelObjects;
+
 		if (!FApp::IsGame())
 		{
 			UObject* Obj = ObjectNode ? ObjectNode->GetUObject(0) : nullptr;
@@ -1511,19 +1552,34 @@
 			{
 				ChildNodePtr->GatherInstancesAffectedByContainerPropertyChange(Obj, Addr, EPropertyArrayChangeType::Insert, AffectedInstances);
 			}
+			
+			TopLevelObjects.Add(Obj);
 		}
 
 		TSet< UObject* > AllAffectedInstances;
 		AllAffectedInstances.Append(AffectedInstances);
+		
+		//set up indices for the coming events
+		TArray< TMap<FString,int32> > ArrayIndicesPerObject;
+		for (int32 ObjectIndex = 0; ObjectIndex < ReadAddresses.Num(); ++ObjectIndex)
+		{
+			//add on array index so we can tell which entry just changed
+			ArrayIndicesPerObject.Add(TMap<FString,int32>());
+			FPropertyValueImpl::GenerateArrayIndexMapToObjectNode(ArrayIndicesPerObject[ObjectIndex], ChildNodePtr );
+		}
 
 		ChildNodePtr->NotifyPreChange( ParentNode->GetProperty(), NotifyHook, AllAffectedInstances );
-
+		
+		if (PropertyUtilities.IsValid() && !bInteractiveChangeInProgress)
+		{
+			FPropertyChangedEvent StartedChangingPropertiesEvent(NodeProperty, EPropertyChangeType::ArrayAdd, MakeArrayView(TopLevelObjects));
+			StartedChangingPropertiesEvent.SetArrayIndexPerObject(ArrayIndicesPerObject);
+			PropertyUtilities.Pin()->NotifyStartedChangingProperties(StartedChangingPropertiesEvent);
+		}
+		
 		FScriptArrayHelper	ArrayHelper(ArrayProperty,Addr);
 		int32 Index = ChildNodePtr->GetArrayIndex();
-
-		// List of top level objects sent to the PropertyChangedEvent
-		TArray<const UObject*> TopLevelObjects;
-
+		
 		UObject* Obj = ObjectNode ? ObjectNode->GetUObject(0) : nullptr;
 		if (Obj)
 		{
@@ -1531,21 +1587,10 @@
 			{
 				ChildNodePtr->PropagateContainerPropertyChange(Obj, Addr, AffectedInstances, EPropertyArrayChangeType::Insert, Index);
 			}
-
-			TopLevelObjects.Add(Obj);
 		}
 
 		ArrayHelper.InsertValues(Index, 1 );
 
-		//set up indices for the coming events
-		TArray< TMap<FString,int32> > ArrayIndicesPerObject;
-		for (int32 ObjectIndex = 0; ObjectIndex < ReadAddresses.Num(); ++ObjectIndex)
-		{
-			//add on array index so we can tell which entry just changed
-			ArrayIndicesPerObject.Add(TMap<FString,int32>());
-			FPropertyValueImpl::GenerateArrayIndexMapToObjectNode(ArrayIndicesPerObject[ObjectIndex], ChildNodePtr );
-		}
-
 		FPropertyChangedEvent ChangeEvent(ParentNode->GetProperty(), EPropertyChangeType::ArrayAdd, MakeArrayView(TopLevelObjects));
 		ChangeEvent.SetArrayIndexPerObject(ArrayIndicesPerObject);
 		ChangeEvent.SetInstancesChanged(MoveTemp(AllAffectedInstances));
@@ -1613,14 +1658,27 @@
 			{
 				bHasChanges = true;
 
+				UObject* Obj = ObjectNode ? ObjectNode->GetUObject(i) : nullptr;
+				
 				if (!FApp::IsGame())
 				{
-					UObject* Obj = ObjectNode ? ObjectNode->GetUObject(i) : nullptr;
 					if (IsTemplate(Obj))
 					{
 						ChildNodePtr->GatherInstancesAffectedByContainerPropertyChange(Obj, Address, EPropertyArrayChangeType::Delete, AffectedInstancesPerObject[i]);
 					}
 				}
+
+				if (Obj)
+				{
+					TopLevelObjects.Add(Obj);
+				}
+				
+				int32 Index = ChildNodePtr->GetArrayIndex();
+				
+				// Add on array index so we can tell which entry just changed
+				int32 ArrayIndicesPerObjectIndex = ArrayIndicesPerObject.Add(TMap<FString, int32>());
+				FPropertyValueImpl::GenerateArrayIndexMapToObjectNode(ArrayIndicesPerObject[ArrayIndicesPerObjectIndex], ChildNodePtr);
+				ArrayIndicesPerObject[ArrayIndicesPerObjectIndex].Add(NodeProperty->GetName(), Index);
 			}
 		}
 
@@ -1634,6 +1692,13 @@
 
 			// Send the PreEditChange notification to all selected objects
 			ChildNodePtr->NotifyPreChange(NodeProperty, NotifyHook, AllAffectedInstances);
+			
+			if (PropertyUtilities.IsValid() && !bInteractiveChangeInProgress)
+			{
+				FPropertyChangedEvent StartedChangingPropertiesEvent(ParentNode->GetProperty(), EPropertyChangeType::ArrayRemove, MakeArrayView(TopLevelObjects));
+				StartedChangingPropertiesEvent.SetArrayIndexPerObject(ArrayIndicesPerObject);
+				PropertyUtilities.Pin()->NotifyStartedChangingProperties(StartedChangingPropertiesEvent);
+			}
 
 			// Perform the operation on the array for all selected objects
 			for (int32 i = 0; i < ReadAddresses.Num(); ++i)
@@ -1644,10 +1709,6 @@
 				{
 					int32 Index = ChildNodePtr->GetArrayIndex();
 
-					// Add on array index so we can tell which entry just changed
-					ArrayIndicesPerObject.Add(TMap<FString, int32>());
-					FPropertyValueImpl::GenerateArrayIndexMapToObjectNode(ArrayIndicesPerObject[i], ChildNodePtr);
-
 					UObject* Obj = ObjectNode ? ObjectNode->GetUObject(i) : nullptr;
 					if (Obj)
 					{
@@ -1655,8 +1716,6 @@
 						{
 							ChildNodePtr->PropagateContainerPropertyChange(Obj, Address, AffectedInstancesPerObject[i], EPropertyArrayChangeType::Delete, Index);
 						}
-
-						TopLevelObjects.Add(Obj);
 					}
 
 					if (ArrayProperty)
@@ -1717,8 +1776,6 @@
 						MapHelper.RemoveAt(InternalIndex);
 						MapHelper.Rehash();
 					}
-
-					ArrayIndicesPerObject[i].Add(NodeProperty->GetName(), Index);
 				}
 			}
 
@@ -1775,6 +1832,20 @@
 		TArray<const UObject*> TopLevelObjects;
 		TopLevelObjects.Reserve(ReadAddresses.Num());
 
+		for ( int32 i = 0 ; i < ReadAddresses.Num() ; ++i )
+		{
+			if (UObject* Obj = ObjectNode ? ObjectNode->GetUObject(i) : nullptr)
+			{
+				TopLevelObjects.Add(Obj);
+			}
+		}
+
+		if (PropertyUtilities.IsValid() && !bInteractiveChangeInProgress)
+		{
+			FPropertyChangedEvent StartedChangingPropertiesEvent(ParentNode->GetProperty(), EPropertyChangeType::ArrayMove, MakeArrayView(TopLevelObjects));
+			PropertyUtilities.Pin()->NotifyStartedChangingProperties(StartedChangingPropertiesEvent);
+		}
+
 		// perform the operation on the array for all selected objects
 		for ( int32 i = 0 ; i < ReadAddresses.Num() ; ++i )
 		{
@@ -1793,8 +1864,6 @@
 					{
 						FirstChildNodePtr->PropagateContainerPropertyChange(Obj, Address, EPropertyArrayChangeType::Swap, FirstIndex, SecondIndex);
 					}
-
-					TopLevelObjects.Add(Obj);
 				}
 
 				if (ArrayProperty)
@@ -2157,6 +2226,10 @@
 
 		FScopedTransaction Transaction( NSLOCTEXT("UnrealEd", "DuplicateChild", "Duplicate Child") );
 
+		TArray< TMap<FString, int32> > ArrayIndicesPerObject;
+		TSet< UObject* > AllAffectedInstances;
+		AllAffectedInstances.Append(AffectedInstances);
+		
 		if (!FApp::IsGame())
 		{
 			UObject* Obj = ObjectNode ? ObjectNode->GetUObject(0) : nullptr;
@@ -2164,34 +2237,39 @@
 			{
 				ChildNodePtr->GatherInstancesAffectedByContainerPropertyChange(Obj, Addr, EPropertyArrayChangeType::Duplicate, AffectedInstances);
 			}
+
+			if (Obj)
+			{
+				TopLevelObjects.Add(Obj);
+			}
+
+			ArrayIndicesPerObject.Add(TMap<FString, int32>());
+			FPropertyValueImpl::GenerateArrayIndexMapToObjectNode(ArrayIndicesPerObject[0], ChildNodePtr);
 		}
-
-		TSet< UObject* > AllAffectedInstances;
-		AllAffectedInstances.Append(AffectedInstances);
-
+		
 		ChildNodePtr->NotifyPreChange( ParentNode->GetProperty(), NotifyHook, AllAffectedInstances );
 
+		if (PropertyUtilities.IsValid() && !bInteractiveChangeInProgress)
+		{
+			FPropertyChangedEvent StartedChangingPropertiesEvent(ParentNode->GetProperty(), EPropertyChangeType::Duplicate, MakeArrayView(TopLevelObjects));
+			StartedChangingPropertiesEvent.SetArrayIndexPerObject(ArrayIndicesPerObject);
+			PropertyUtilities.Pin()->NotifyStartedChangingProperties(StartedChangingPropertiesEvent);
+		}
+
 		int32 Index = ChildNodePtr->GetArrayIndex();
 		UObject* Obj = ObjectNode ? ObjectNode->GetUObject(0) : nullptr;
 
-		TArray< TMap<FString, int32> > ArrayIndicesPerObject;
-
-		FScriptArrayHelper ArrayHelper(ArrayProperty, Addr);
-		FPropertyNode::DuplicateArrayEntry(NodeProperty, ArrayHelper, Index);
-
 		if (Obj)
 		{
 			if (IsTemplate(Obj) && !FApp::IsGame())
 			{
 				ChildNodePtr->PropagateContainerPropertyChange(Obj, Addr, AffectedInstances, EPropertyArrayChangeType::Duplicate, Index);
 			}
-
-			TopLevelObjects.Add(Obj);
-
-			ArrayIndicesPerObject.Add(TMap<FString, int32>());
-			FPropertyValueImpl::GenerateArrayIndexMapToObjectNode(ArrayIndicesPerObject[0], ChildNodePtr);
 		}
-
+		
+		FScriptArrayHelper ArrayHelper(ArrayProperty, Addr);
+		FPropertyNode::DuplicateArrayEntry(NodeProperty, ArrayHelper, Index);
+		
 		FPropertyChangedEvent ChangeEvent(ParentNode->GetProperty(), EPropertyChangeType::Duplicate, MakeArrayView(TopLevelObjects));
 		ChangeEvent.SetArrayIndexPerObject(ArrayIndicesPerObject);
 		ChangeEvent.SetInstancesChanged(MoveTemp(AllAffectedInstances));
@@ -2849,6 +2927,8 @@
 		if (Implementation->GetPropertyUtilities().IsValid() && Implementation->GetPropertyUtilities()->GetNotifyHook() != nullptr)
 		{
 			Implementation->GetPropertyNode()->NotifyPreChange(Implementation->GetPropertyNode()->GetProperty(), Implementation->GetPropertyUtilities()->GetNotifyHook());
+			FPropertyChangedEvent StartedChangingPropertiesEvent(Implementation->GetPropertyNode()->GetProperty());
+			Implementation->GetPropertyUtilities()->NotifyStartedChangingProperties(StartedChangingPropertiesEvent);
 		}
 
 		FResetToDefaultHandler Delegate = CustomResetToDefault.GetPropertyResetToDefaultDelegate();
@@ -5388,6 +5468,7 @@
 	}
 	
 	TArray< TMap<FString, int32> > ArrayIndicesPerObject;
+	ArrayIndicesPerObject.Reserve(ReadAddresses.Num());
 	TArray<TArray<UObject*>> AffectedInstancesPerObject;
 	AffectedInstancesPerObject.SetNum(ReadAddresses.Num());
 
@@ -5411,8 +5492,17 @@
 				{
 					PropertyNode->GatherInstancesAffectedByContainerPropertyChange(Obj, Addr, EPropertyArrayChangeType::Add, AffectedInstancesPerObject[i]);
 				}
+
+				if (Obj)
+				{
+					TopLevelObjects.Add(Obj);
+				}
 			}
 		}
+
+		// Add on array index so we can tell which entry just changed
+		ArrayIndicesPerObject.Add(TMap<FString, int32>());
+		FPropertyValueImpl::GenerateArrayIndexMapToObjectNode(ArrayIndicesPerObject[i], PropertyNode.Get());
 	}
 
 	TSet< UObject* > AllAffectedInstances;
@@ -5424,7 +5514,15 @@
 	// send the PreEditChange notification to all selected objects
 	FNotifyHook* NotifyHook = Implementation->GetNotifyHook();
 	PropertyNode->NotifyPreChange(OptionalProperty, NotifyHook, AllAffectedInstances);
-
+	
+	if (TSharedPtr<IPropertyUtilities> PropertyUtilities = Implementation->GetPropertyUtilities())
+	{
+		FPropertyChangedEvent StartedChangingPropertiesEvent(OptionalProperty, EPropertyChangeType::ValueSet, MakeArrayView(TopLevelObjects));
+		StartedChangingPropertiesEvent.SetArrayIndexPerObject(ArrayIndicesPerObject);
+		StartedChangingPropertiesEvent.SetInstancesChanged(MoveTemp(AllAffectedInstances));
+		PropertyUtilities->NotifyStartedChangingProperties(StartedChangingPropertiesEvent);
+	}
+	
 	for (int i = 0; i < ReadAddresses.Num(); i++)
 	{
 		void* Addr = ReadAddresses.GetAddress(i);
@@ -5489,10 +5587,6 @@
 				}
 			}
 		}
-
-		// Add on array index so we can tell which entry just changed
-		ArrayIndicesPerObject.Add(TMap<FString, int32>());
-		FPropertyValueImpl::GenerateArrayIndexMapToObjectNode(ArrayIndicesPerObject[i], PropertyNode.Get());
 	}
 
 	FPropertyChangedEvent ChangeEvent(OptionalProperty, EPropertyChangeType::ValueSet, MakeArrayView(TopLevelObjects));
@@ -5536,6 +5630,7 @@
 	}
 
 	TArray< TMap<FString, int32> > ArrayIndicesPerObject;
+	ArrayIndicesPerObject.Reserve(ReadAddresses.Num());
 	TArray<TArray<UObject*>> AffectedInstancesPerObject;
 	AffectedInstancesPerObject.SetNum(ReadAddresses.Num());
 
@@ -5561,6 +5656,9 @@
 				}
 			}
 		}
+		// Add on array index so we can tell which entry just changed
+		ArrayIndicesPerObject.Add(TMap<FString, int32>());
+		FPropertyValueImpl::GenerateArrayIndexMapToObjectNode(ArrayIndicesPerObject[i], PropertyNode.Get());
 	}
 
 	TSet< UObject* > AllAffectedInstances;
@@ -5573,6 +5671,13 @@
 	FNotifyHook* NotifyHook = Implementation->GetNotifyHook();
 	PropertyNode->NotifyPreChange(OptionalProperty, NotifyHook, AllAffectedInstances);
 	
+	if (TSharedPtr<IPropertyUtilities> PropertyUtilities = Implementation->GetPropertyUtilities())
+	{
+		FPropertyChangedEvent StartedChangingPropertiesEvent(OptionalProperty, EPropertyChangeType::ValueSet, MakeArrayView(TopLevelObjects));
+		StartedChangingPropertiesEvent.SetArrayIndexPerObject(ArrayIndicesPerObject);
+		PropertyUtilities->NotifyStartedChangingProperties(StartedChangingPropertiesEvent);
+	}
+	
 	for (int i = 0; i < ReadAddresses.Num(); i++)
 	{
 		UObject* Obj = ObjectNode ? ObjectNode->GetUObject(i) : nullptr;
@@ -5580,10 +5685,6 @@
 
 		void* Addr = ReadAddresses.GetAddress(i);
 		OptionalProperty->MarkUnset(Addr);
-
-		// Add on array index so we can tell which entry just changed
-		ArrayIndicesPerObject.Add(TMap<FString, int32>());
-		FPropertyValueImpl::GenerateArrayIndexMapToObjectNode(ArrayIndicesPerObject[i], PropertyNode.Get());
 	}
 
 	FPropertyChangedEvent ChangeEvent(OptionalProperty, EPropertyChangeType::ValueSet, MakeArrayView(TopLevelObjects));
diff --git a/Engine/Source/Editor/PropertyEditor/Private/PropertyRowGenerator.cpp b/Engine/Source/Editor/PropertyEditor/Private/PropertyRowGenerator.cpp
--- a/Engine/Source/Editor/PropertyEditor/Private/PropertyRowGenerator.cpp
+++ b/Engine/Source/Editor/PropertyEditor/Private/PropertyRowGenerator.cpp
@@ -78,6 +78,14 @@ public:
 			: TSharedPtr<class FAssetThumbnailPool>();
 	}
 
+	virtual void NotifyStartedChangingProperties(const FPropertyChangedEvent& PropertyChangedEvent)
+	{
+		if (Generator != nullptr)
+		{
+			Generator->OnNotifyStartedChangingProperties().Broadcast(PropertyChangedEvent);
+		}
+	}
+
 	virtual void NotifyFinishedChangingProperties(const FPropertyChangedEvent& PropertyChangedEvent) override 
 	{
 		if (Generator)
diff --git a/Engine/Source/Editor/PropertyEditor/Private/PropertyRowGenerator.h b/Engine/Source/Editor/PropertyEditor/Private/PropertyRowGenerator.h
--- a/Engine/Source/Editor/PropertyEditor/Private/PropertyRowGenerator.h
+++ b/Engine/Source/Editor/PropertyEditor/Private/PropertyRowGenerator.h
@@ -119,6 +119,7 @@ public:
 
 	const FCustomPropertyTypeLayoutMap& GetInstancedPropertyTypeLayoutMap() const;
 	void UpdateDetailRows();
+	virtual FOnNotifyStartedChangingProperties& OnNotifyStartedChangingProperties() override { return OnNotifyStartedChangingPropertiesDelegate; }
 	virtual FOnFinishedChangingProperties& OnFinishedChangingProperties() override { return OnFinishedChangingPropertiesDelegate; }
 
 private:
@@ -158,6 +159,8 @@ private:
 	TSharedRef<IPropertyGenerationUtilities> PropertyGenerationUtilities;
 	/** Delegate called when the details panel finishes editing a property (after post edit change is called) */
 	FOnFinishedChangingProperties OnFinishedChangingPropertiesDelegate;
+	/** Delegate called before the details panel finishes editing a property (after pre edit change is called) */
+	FOnNotifyStartedChangingProperties OnNotifyStartedChangingPropertiesDelegate;
 	/** The ValidatePropertyNodes function can be overridden with this member, if set. Useful if your implementation doesn't require this kind of validation each Tick. */
 	FOnValidatePropertyRowGeneratorNodes CustomValidatePropertyNodesFunction;
 	/** If specified only nodes for these properties will be generated */
diff --git a/Engine/Source/Editor/PropertyEditor/Private/SPropertyTreeViewImpl.cpp b/Engine/Source/Editor/PropertyEditor/Private/SPropertyTreeViewImpl.cpp
--- a/Engine/Source/Editor/PropertyEditor/Private/SPropertyTreeViewImpl.cpp
+++ b/Engine/Source/Editor/PropertyEditor/Private/SPropertyTreeViewImpl.cpp
@@ -94,6 +94,7 @@ public:
 		return NotImplemented;
 	}
 
+	virtual void NotifyStartedChangingProperties(const FPropertyChangedEvent& PropertyChangedEvent) override {}
 	virtual void NotifyFinishedChangingProperties(const FPropertyChangedEvent& PropertyChangedEvent) override {}
 
 	virtual bool DontUpdateValueWhileEditing() const override
diff --git a/Engine/Source/Editor/PropertyEditor/Private/SSingleProperty.cpp b/Engine/Source/Editor/PropertyEditor/Private/SSingleProperty.cpp
--- a/Engine/Source/Editor/PropertyEditor/Private/SSingleProperty.cpp
+++ b/Engine/Source/Editor/PropertyEditor/Private/SSingleProperty.cpp
@@ -82,6 +82,9 @@ public:
 	virtual void NotifyFinishedChangingProperties(const FPropertyChangedEvent& PropertyChangedEvent) override
 	{}
 
+	virtual void NotifyStartedChangingProperties(const FPropertyChangedEvent& PropertyChangedEvent) override
+	{}
+
 	virtual bool DontUpdateValueWhileEditing() const override
 	{
 		return false;
diff --git a/Engine/Source/Editor/PropertyEditor/Private/UserInterface/PropertyDetails/PropertyDetailsUtilities.h b/Engine/Source/Editor/PropertyEditor/Private/UserInterface/PropertyDetails/PropertyDetailsUtilities.h
--- a/Engine/Source/Editor/PropertyEditor/Private/UserInterface/PropertyDetails/PropertyDetailsUtilities.h
+++ b/Engine/Source/Editor/PropertyEditor/Private/UserInterface/PropertyDetails/PropertyDetailsUtilities.h
@@ -23,6 +23,7 @@ public:
 	virtual void RequestRefresh() override;
 	virtual void RequestForceRefresh() override;
 	virtual TSharedPtr<class FAssetThumbnailPool> GetThumbnailPool() const override;
+	virtual void NotifyStartedChangingProperties(const FPropertyChangedEvent& PropertyChangedEvent) override {};
 	virtual void NotifyFinishedChangingProperties(const FPropertyChangedEvent& PropertyChangedEvent) override;
 	virtual bool DontUpdateValueWhileEditing() const override;
 	const TArray<TWeakObjectPtr<UObject>>& GetSelectedObjects() const override;
diff --git a/Engine/Source/Editor/PropertyEditor/Public/IPropertyRowGenerator.h b/Engine/Source/Editor/PropertyEditor/Public/IPropertyRowGenerator.h
--- a/Engine/Source/Editor/PropertyEditor/Public/IPropertyRowGenerator.h
+++ b/Engine/Source/Editor/PropertyEditor/Public/IPropertyRowGenerator.h
@@ -108,6 +108,7 @@ public:
 	virtual void UnregisterInstancedCustomPropertyTypeLayout(FName PropertyTypeName, TSharedPtr<IPropertyTypeIdentifier> Identifier = nullptr) = 0;
 
 	virtual FOnFinishedChangingProperties& OnFinishedChangingProperties() = 0;
+	virtual FOnNotifyStartedChangingProperties& OnNotifyStartedChangingProperties() = 0;
 
 	/* Use this function to set a callback on FPropertyRowGenerator that will override the ValidatePropertyNodes function.
 	 * This is useful if your implementation doesn't need to validate nodes every tick or needs to perform some other form of validation. */
diff --git a/Engine/Source/Editor/PropertyEditor/Public/IPropertyUtilities.h b/Engine/Source/Editor/PropertyEditor/Public/IPropertyUtilities.h
--- a/Engine/Source/Editor/PropertyEditor/Public/IPropertyUtilities.h
+++ b/Engine/Source/Editor/PropertyEditor/Public/IPropertyUtilities.h
@@ -41,6 +41,7 @@ public:
 	virtual void RequestForceRefresh() = 0;
 	
 	virtual TSharedPtr<class FAssetThumbnailPool> GetThumbnailPool() const = 0;
+	virtual void NotifyStartedChangingProperties(const FPropertyChangedEvent& PropertyChangedEvent) = 0;
 	virtual void NotifyFinishedChangingProperties(const FPropertyChangedEvent& PropertyChangedEvent) = 0;
 	virtual bool HasClassDefaultObject() const = 0;
 	virtual const TArray<TWeakObjectPtr<UObject>>& GetSelectedObjects() const = 0;
diff --git a/Engine/Source/Editor/PropertyEditor/Public/PropertyEditorDelegates.h b/Engine/Source/Editor/PropertyEditor/Public/PropertyEditorDelegates.h
--- a/Engine/Source/Editor/PropertyEditor/Public/PropertyEditorDelegates.h
+++ b/Engine/Source/Editor/PropertyEditor/Public/PropertyEditorDelegates.h
@@ -111,6 +111,12 @@ DECLARE_DELEGATE_RetVal_OneParam(FName, FPropertyHandleLayoutOverride, const IPr
  */
 DECLARE_MULTICAST_DELEGATE_OneParam(FOnFinishedChangingProperties, const FPropertyChangedEvent&);
 
+/**
+ * A delegate which is called before properties have been edited and PreEditChange has been called on all objects.
+ * This can be used to identify what precisely changed.
+ */
+DECLARE_MULTICAST_DELEGATE_OneParam(FOnNotifyStartedChangingProperties, const FPropertyChangedEvent&);
+
 /**
  * A property row extension arguments is displayed at the end of a property row, either inline or as a button.
  */
