Add a GetNextKey property meta data with which a class can calculate the next valid key in a BP-exposed map

diff --git a/Engine/Source/Editor/PropertyEditor/Private/PropertyHandleImpl.cpp b/Engine/Source/Editor/PropertyEditor/Private/PropertyHandleImpl.cpp
--- a/Engine/Source/Editor/PropertyEditor/Private/PropertyHandleImpl.cpp
+++ b/Engine/Source/Editor/PropertyEditor/Private/PropertyHandleImpl.cpp
@@ -6174,7 +6174,35 @@ FPropertyHandleItemAddResult FPropertyHandleMap::AddItem()
 				FMapProperty* MapProperty = CastFieldChecked<FMapProperty>(PropNode->GetProperty());
 				FScriptMapHelper MapHelper(MapProperty, ValueBaseAddress);
 
-				if (const FEnumProperty* EnumProperty = CastField<const FEnumProperty>(MapProperty->GetKeyProperty()))
+				static const FName GetNextKeyMetaDataName = "GetNextKey";
+				const FString* GetNextKeyMetaData = MapProperty->FindMetaData(GetNextKeyMetaDataName);
+				if (GetNextKeyMetaData && !bIsStruct)
+				{
+					if (UFunction* Function = Addresses[0].Object->FindFunction(FName(*GetNextKeyMetaData)))
+					{
+						const FProperty* KeyProperty = MapProperty->GetKeyProperty();
+						void* Value = FMemory_Alloca_Aligned(KeyProperty->GetElementSize(), MapProperty->GetMinAlignment());
+						KeyProperty->InitializeValue(Value);
+
+						Addresses[0].Object->ProcessEvent(Function, Value);
+
+						FString StringValue;
+						KeyProperty->ExportText_Direct(
+							StringValue, Value, nullptr, nullptr, PPF_None
+						);
+
+						if (!HasKey(MapHelper, ValueBaseAddress, StringValue))
+						{
+							const int32 ChildNodeIndex = AddItemWithKey(StringValue);
+							if (ChildNodeIndex != INDEX_NONE)
+							{
+								Result = FPropertyAccess::Success;
+								return FPropertyHandleItemAddResult(ChildNodeIndex, Result);
+							}
+						}
+					}
+				}
+				else if (const FEnumProperty* EnumProperty = CastField<const FEnumProperty>(MapProperty->GetKeyProperty()))
 				{
 					if (UEnum* Enum = EnumProperty->GetEnum())
 					{
@@ -6196,23 +6224,11 @@ FPropertyHandleItemAddResult FPropertyHandleMap::AddItem()
 								continue;
 							}
 
-							// If we don't have this key then add an entry and set it to this key value.
-							const int32 ChildNodeIndex = Implementation->AddChild();
-							if (ChildNodeIndex >= 0)
+							const int32 ChildNodeIndex = AddItemWithKey(Enum->GetNameStringByIndex(EnumIndex));
+							if (ChildNodeIndex != INDEX_NONE)
 							{
-								TSharedPtr<FPropertyNode> ChildNode = Implementation->GetChildNode(ChildNodeIndex);
-								if (ChildNode.IsValid())
-								{
-									// We want to change the key, so get the key property.
-									TSharedPtr<FPropertyNode> ChildKeyNode = ChildNode->GetPropertyKeyNode();
-									if (ChildKeyNode.IsValid())
-									{
-										// Set the new entry to have the key value.
-										Implementation->ImportText(Enum->GetNameStringByIndex(EnumIndex), ChildKeyNode.Get(), EPropertyValueSetFlags::DefaultFlags);
-										Result = FPropertyAccess::Success;
-										return FPropertyHandleItemAddResult(ChildNodeIndex, Result);
-									}
-								}
+								Result = FPropertyAccess::Success;
+								return FPropertyHandleItemAddResult(ChildNodeIndex, Result);
 							}
 						}
 					}
@@ -6293,6 +6309,29 @@ bool FPropertyHandleMap::IsEditable() const
 	return Implementation->HasValidPropertyNode() && !Implementation->IsEditConst() && Implementation->IsPropertyTypeOf(FMapProperty::StaticClass());
 }
 
+int32 FPropertyHandleMap::AddItemWithKey(const FString& Value)
+{
+	// If we don't have this key then add an entry and set it to this key value.
+	const int32 ChildNodeIndex = Implementation->AddChild();
+	if (ChildNodeIndex >= 0)
+	{
+		TSharedPtr<FPropertyNode> ChildNode = Implementation->GetChildNode(ChildNodeIndex);
+		if (ChildNode.IsValid())
+		{
+			// We want to change the key, so get the key property.
+			TSharedPtr<FPropertyNode> ChildKeyNode = ChildNode->GetPropertyKeyNode();
+			if (ChildKeyNode.IsValid())
+			{
+				// Set the new entry to have the key value.
+				Implementation->ImportText(Value, ChildKeyNode.Get(), EPropertyValueSetFlags::DefaultFlags);
+				return ChildNodeIndex;
+			}
+		}
+	}
+
+	return INDEX_NONE;
+}
+
 // FieldPath
 
 bool FPropertyHandleFieldPath::Supports(TSharedRef<FPropertyNode> PropertyNode)
diff --git a/Engine/Source/Editor/PropertyEditor/Private/PropertyHandleImpl.h b/Engine/Source/Editor/PropertyEditor/Private/PropertyHandleImpl.h
--- a/Engine/Source/Editor/PropertyEditor/Private/PropertyHandleImpl.h
+++ b/Engine/Source/Editor/PropertyEditor/Private/PropertyHandleImpl.h
@@ -806,6 +806,9 @@ public:
 	virtual FString GetDocumentationLink() override { return FString("Engine/UI/LevelEditor/Details/Properties/Map/"); }
 	virtual FString GetDocumentationExcerptName() override { return FString("Maps"); }
 	virtual bool IsEditable() const override;
+
+private:
+	int32 AddItemWithKey(const FString& Value);
 };
 
 class FPropertyHandleFieldPath : public FPropertyHandleBase
