[CrashFix] Workaround for a design flaw of FPropertyChangeEvent::ArrayIndicesPerObject that causes issues when constructing change events for arrays of optionals.
For arrays of optionals, the array, the optional, and the optional value all have the same property name. The code in PropertyHandleImpl assumed the inner most property would be the one with an accurate array index but in this case it's actually the optional property because that's the one who's parent is an array. That assumption caused the index to get incorrectly stored as INDEX_NONE.

note: there is still a bug here that's unresolvable without changing the way FPropertyChangeEvent stores indices. In this example there was only one property who's index actually mattered (the other two were INDEX_NONE) but there are scenarios where there's multiple unique properties all of which have the same name but different indices.
for example, given the following pseudo object:
struct TestStruct
{
TArray<int32> MyArray;
};

class TestObject : UObject
{
TArray<TestStruct> MyArray;
};

if TestObject.MyArray[0].MyArray[2] is edited, it will be ambiguous whether the key "MyArray" in the ArrayIndexMap is referring to the array in TestObject or the array in TestStruct.
at the moment this code sets ArrayIndexMap["MyArray"] to 2 for this example losing the index of TestObject::MyArray.

This can't be resolved unless we store indices in such a way that's unambiguous between the two.

diff --git a/Engine/Source/Editor/PropertyEditor/Private/PropertyHandleImpl.cpp b/Engine/Source/Editor/PropertyEditor/Private/PropertyHandleImpl.cpp
--- a/Engine/Source/Editor/PropertyEditor/Private/PropertyHandleImpl.cpp
+++ b/Engine/Source/Editor/PropertyEditor/Private/PropertyHandleImpl.cpp
@@ -216,6 +216,21 @@ FString FPropertyValueImpl::GetPropertyValueArray() const
 	return String;
 }
 
+void LogAmbiguousArrayIndexWarning(const FPropertyNode* PropertyNode, const FString& CompetingName)
+{
+	FString LogMessage = TEXT("Found multiple properties in a property change chain with matching names but mismatched indices:\n");
+	for (const FPropertyNode* IterationNode = PropertyNode; (IterationNode != nullptr); IterationNode = IterationNode->GetParentNode())
+	{
+		const FProperty* Property = IterationNode->GetProperty();
+		if (Property->GetName() == CompetingName && IterationNode->GetArrayIndex() != INDEX_NONE)
+		{
+			LogMessage += FString::Format(TEXT(" - '{0}'\n"), { IterationNode->GetPropertyPath()});
+		}
+	}
+	LogMessage += TEXT("This causes ambiguity in the property change notifications!");
+	UE_LOG(LogPropertyNode, Warning, TEXT("%s"), *LogMessage);
+}
+
 
 void FPropertyValueImpl::GenerateArrayIndexMapToObjectNode( TMap<FString,int32>& OutArrayIndexMap, FPropertyNode* PropertyNode )
 {
@@ -227,8 +242,35 @@ void FPropertyValueImpl::GenerateArrayIndexMapToObjectNode( TMap<FString,int32>&
 			FProperty* Property = IterationNode->GetProperty();
 			if (Property)
 			{
-				//since we're starting from the lowest level, we have to take the first entry.  In the case of an array, the entries and the array itself have the same name, except the parent has an array index of -1
-				if (!OutArrayIndexMap.Contains(Property->GetName()))
+				int32 NodeIndex = IterationNode->GetArrayIndex();
+				if (int32* FoundIndex = OutArrayIndexMap.Find(Property->GetName()); FoundIndex)
+				{
+					if (*FoundIndex != INDEX_NONE && NodeIndex != INDEX_NONE && *FoundIndex != NodeIndex)
+					{
+						// TODO: due to the way array indices are stored, if this warning is printed, it's impossible to create a change notification that's resolvable :(
+						// for example, given the following pseudo object:
+						/*
+							struct TestStruct
+							{
+								TArray<int32> MyArray;
+							};
+
+							class TestObject : UObject
+							{
+								TArray<TestStruct> MyArray;
+							};
+						*/
+						// if TestObject.MyArray[0].MyArray[2] is edited, it will be ambiguous whether the key "MyArray" in the ArrayIndexMap is referring to the array in TestObject or the array in TestStruct.
+						// at the moment this code sets ArrayIndexMap["MyArray"] to 2 for this example losing the index of TestObject::MyArray
+						LogAmbiguousArrayIndexWarning(PropertyNode, Property->GetName());
+					}
+
+					if (*FoundIndex == INDEX_NONE)
+					{
+						*FoundIndex = NodeIndex;
+					}
+				}
+				else
 				{
 					OutArrayIndexMap.Add(Property->GetName(), IterationNode->GetArrayIndex());
 				}
