Blueprint property categories in the details panel are now sorted according to the user defined order.

diff --git a/Engine/Source/Editor/PropertyEditor/Private/ObjectPropertyNode.cpp b/Engine/Source/Editor/PropertyEditor/Private/ObjectPropertyNode.cpp
--- a/Engine/Source/Editor/PropertyEditor/Private/ObjectPropertyNode.cpp
+++ b/Engine/Source/Editor/PropertyEditor/Private/ObjectPropertyNode.cpp
@@ -522,7 +522,7 @@ static TSharedPtr<FPropertyNode> FindChildCategory( TSharedPtr<FPropertyNode> Pa
 }
 
 void FObjectPropertyNode::GetCategoryProperties(const TSet<UClass*>& ClassesToConsider, const FProperty* CurrentProperty, bool bShouldShowDisableEditOnInstance, bool bShouldShowHiddenProperties,
-	const TSet<FName>& CategoriesFromBlueprints, TSet<FName>& CategoriesFromProperties, TArray<FName>& SortedCategories)
+	TSet<FName>& CategoriesFromProperties, TArray<FName>& SortedCategories)
 {
 	static const FName Name_bShowOnlyWhenTrue("bShowOnlyWhenTrue");
 	static const FName Name_ShowOnlyInnerProperties("ShowOnlyInnerProperties");
@@ -550,11 +550,12 @@ void FObjectPropertyNode::GetCategoryProperties(const TSet<UClass*>& ClassesToCo
 	{
 		if (PropertyEditorHelpers::ShouldBeVisible(*this, CurrentProperty) && !HiddenCategories.Contains(CategoryName))
 		{
-			if (!CategoriesFromBlueprints.Contains(CategoryName) && !CategoriesFromProperties.Contains(CategoryName))
+			bool bIsAlreadyInSet;
+			CategoriesFromProperties.Add(CategoryName, &bIsAlreadyInSet);
+			if (!bIsAlreadyInSet)
 			{
-				SortedCategories.AddUnique(CategoryName);
+				SortedCategories.Add(CategoryName);
 			}
-			CategoriesFromProperties.Add(CategoryName);
 		}
 	}
 
@@ -565,12 +566,74 @@ void FObjectPropertyNode::GetCategoryProperties(const TSet<UClass*>& ClassesToCo
 		{
 			for (TFieldIterator<FProperty> It(StructProperty->Struct); It; ++It)
 			{
-				GetCategoryProperties(ClassesToConsider, *It, bShouldShowDisableEditOnInstance, bShouldShowHiddenProperties, CategoriesFromBlueprints, CategoriesFromProperties, SortedCategories);
+				GetCategoryProperties(ClassesToConsider, *It, bShouldShowDisableEditOnInstance, bShouldShowHiddenProperties, CategoriesFromProperties, SortedCategories);
 			}
 		}
 	}
 }
 
+#if WITH_EDITORONLY_DATA
+
+static void SortCategoriesByUserDefinedOrder(TArray<FName>& SortedCategories, const TSet<UClass*>& ClassesToConsider)
+{
+	// Acquire the sort indices from classes in display text format.
+	TMap<FString, int32> DisplayTextSortOrders;
+	for (UClass* Class : ClassesToConsider)
+	{
+		if (UBlueprint* BP = Cast<UBlueprint>(Class->ClassGeneratedBy))
+		{
+			for (const FName& DisplayTextCategory : BP->CategorySorting)
+			{ 
+				// These are already in display text format.
+				const int32 SortOrder = DisplayTextSortOrders.Num();
+				DisplayTextSortOrders.FindOrAdd(DisplayTextCategory.ToString(), SortOrder);
+			}
+		}
+
+		TArray<FString> ClassPrioritizeCategories;
+		Class->GetPrioritizeCategories(ClassPrioritizeCategories);
+		for (const FString& MetaCategory : ClassPrioritizeCategories)
+		{
+			// These are user declared meta specifiers, so we sanitize them.
+			const FString DisplayTextCategory = FEditorCategoryUtils::GetCategoryDisplayString(MetaCategory);
+			const int32 SortOrder = DisplayTextSortOrders.Num();
+			DisplayTextSortOrders.FindOrAdd(DisplayTextCategory, SortOrder);
+		}
+	}
+
+	Algo::Sort(SortedCategories, [&DisplayTextSortOrders](const FName& CategoryA, const FName& CategoryB)
+		{
+			const FString DisplayTextCategoryA = FEditorCategoryUtils::GetCategoryDisplayString(CategoryA.ToString());
+			const FString DisplayTextCategoryB = FEditorCategoryUtils::GetCategoryDisplayString(CategoryB.ToString());
+			const int32* SortOrderA = DisplayTextSortOrders.Find(DisplayTextCategoryA);
+			const int32* SortOrderB = DisplayTextSortOrders.Find(DisplayTextCategoryB);
+			if (SortOrderA)
+			{
+				if (SortOrderB)
+				{
+					return *SortOrderA < *SortOrderB; // Sort by user-defined order.
+				}
+				else
+				{
+					return true; // Move undefined items to the bottom.
+				}
+			}
+			else
+			{
+				if (SortOrderB)
+				{
+					return false; // Move undefined items to the bottom.
+				}
+				else
+				{
+					return DisplayTextCategoryA < DisplayTextCategoryB; // Sort everything else alphabetical.
+				}
+			}
+		});
+}
+
+#endif
+
 void FObjectPropertyNode::InternalInitChildNodes( FName SinglePropertyName )
 {
 	HiddenCategories.Reset();
@@ -589,35 +652,7 @@ void FObjectPropertyNode::InternalInitChildNodes( FName SinglePropertyName )
 	}
 
 	// Create a merged list of user-enforced sorted info, hidden category info, etc...
-	TSet<FName> CategoriesFromBlueprints, CategoriesFromProperties;
 	TArray<FName> SortedCategories;
-	TArray<FName> PrioritizeCategories;
-
-#if WITH_EDITORONLY_DATA
-	for (UClass* Class : ClassesToConsider)
-	{
-		if (UBlueprint* BP = Cast<UBlueprint>(Class->ClassGeneratedBy))
-		{
-			for (const FName& TestCategory : BP->CategorySorting)
-			{
-				if (!CategoriesFromBlueprints.Contains(TestCategory))
-				{
-					CategoriesFromBlueprints.Add(TestCategory);
-					SortedCategories.AddUnique(TestCategory);
-				}
-			}
-		}
-
-		TArray<FString> ClassPrioritizeCategories;
-		Class->GetPrioritizeCategories(ClassPrioritizeCategories);
-		for (const FString& ClassPrioritizeCategory : ClassPrioritizeCategories)
-		{
-			FName PrioritizeCategoryName = FName(ClassPrioritizeCategory);
-			SortedCategories.AddUnique(PrioritizeCategoryName);
-			PrioritizeCategories.AddUnique(PrioritizeCategoryName);
-		}
-	}
-#endif
 
 	const bool bShouldShowHiddenProperties = !!HasNodeFlags(EPropertyNodeFlags::ShouldShowHiddenProperties);
 	const bool bShouldShowDisableEditOnInstance = !!HasNodeFlags(EPropertyNodeFlags::ShouldShowDisableEditOnInstance);
@@ -625,9 +660,10 @@ void FObjectPropertyNode::InternalInitChildNodes( FName SinglePropertyName )
 
 	if (UClass* ResolvedBaseClass = BaseClass.Get())
 	{
+		TSet<FName> CategoriesFromProperties;
 		for (TFieldIterator<FProperty> It(ResolvedBaseClass); It; ++It)
 		{
-			GetCategoryProperties(ClassesToConsider, *It, bShouldShowDisableEditOnInstance, bShouldShowHiddenProperties, CategoriesFromBlueprints, CategoriesFromProperties, SortedCategories);
+			GetCategoryProperties(ClassesToConsider, *It, bShouldShowDisableEditOnInstance, bShouldShowHiddenProperties, CategoriesFromProperties, SortedCategories);
 		}
 
 		if (UScriptStruct* SparseClassDataStruct = ResolvedBaseClass->GetSparseClassDataStruct())
@@ -636,29 +672,16 @@ void FObjectPropertyNode::InternalInitChildNodes( FName SinglePropertyName )
 		
 			for (TFieldIterator<FProperty> It(SparseClassDataStruct); It; ++It)
 			{
-				GetCategoryProperties(ClassesToConsider, *It, bShouldShowDisableEditOnInstance, bShouldShowHiddenProperties, CategoriesFromBlueprints, CategoriesFromProperties, SortedCategories);
-			}
-		}
-
-		// Categories from the Blueprint class may include categories with no associated properties, so remove those ones
-		for (const FName& CategoryName : CategoriesFromBlueprints)
-		{
-			if (!CategoriesFromProperties.Contains(CategoryName))
-			{
-				SortedCategories.Remove(CategoryName);
+				GetCategoryProperties(ClassesToConsider, *It, bShouldShowDisableEditOnInstance, bShouldShowHiddenProperties, CategoriesFromProperties, SortedCategories);
 			}
 		}
 
-		// Remove any prioritized categories that don't actually exist
-		for (const FName& PrioritizeCategory : PrioritizeCategories)
-		{
-			if (!CategoriesFromProperties.Contains(PrioritizeCategory))
-			{
-				SortedCategories.Remove(PrioritizeCategory);
-			}
-		}
+#if WITH_EDITORONLY_DATA
+		SortCategoriesByUserDefinedOrder(SortedCategories, ClassesToConsider);
+#endif
 	}
 
+
 	//////////////////////////////////////////
 	// Add the category headers and the child items that belong inside of them.
 
diff --git a/Engine/Source/Editor/PropertyEditor/Private/ObjectPropertyNode.h b/Engine/Source/Editor/PropertyEditor/Private/ObjectPropertyNode.h
--- a/Engine/Source/Editor/PropertyEditor/Private/ObjectPropertyNode.h
+++ b/Engine/Source/Editor/PropertyEditor/Private/ObjectPropertyNode.h
@@ -169,7 +169,7 @@ private:
 
 	/** If CurrentProperty should show up in the ClassesToConsider make sure its category is in SortedCategories and CategoriesFromProperties. */
 	void GetCategoryProperties(const TSet<UClass*>& ClassesToConsider, const FProperty* CurrentProperty, bool bShouldShowDisableEditOnInstance, bool bShouldShowHiddenProperties,
-	const TSet<FName>& CategoriesFromBlueprints, TSet<FName>& CategoriesFromProperties, TArray<FName>& SortedCategories);
+	TSet<FName>& CategoriesFromProperties, TArray<FName>& SortedCategories);
 private:
 	/** The list of objects we are editing properties for. */
 	TArray< TWeakObjectPtr<UObject> >		Objects;
