FDetailCategoryImpl : Call RemoveTickableNode(*this) from the destructor to handle the case where this is deleted after the call to AddTickableNode(*this) but before the next FDetailCategoryImpl::Tick.

Fixes a crash in FDetailLayoutBuilderImpl::Tick

diff --git a/Engine/Source/Editor/PropertyEditor/Private/DetailCategoryBuilderImpl.cpp b/Engine/Source/Editor/PropertyEditor/Private/DetailCategoryBuilderImpl.cpp
--- a/Engine/Source/Editor/PropertyEditor/Private/DetailCategoryBuilderImpl.cpp
+++ b/Engine/Source/Editor/PropertyEditor/Private/DetailCategoryBuilderImpl.cpp
@@ -258,6 +258,7 @@ FDetailCategoryImpl::~FDetailCategoryImpl()
 	{
 		RemoveTickableNode(*AdvancedChildNode);
 	}
+	RemoveTickableNode(*this);
 }
 
 
