Fix for broken EditConditions related to members in sparse class data
* Existing implementation of GetPropertyValuePtr (FObjectPropertyNode::GetValuePtrOfInstance) fails to offset the StartAddress of the Sparse Class Data instance with the target FProperty's offset.
* FPropertyNode::GetValueBaseAddress with `bIsSparseData = true` is suspected as the StartAddress is expected to be the start of the object holding the property, but it is passed through unchanged with sparse data.
* Fixing the GetValueBaseAddress with `bIsSparseData = true` is preferred but has a very large testing area that affects public accessible functions.
* This is a stopgap fix to use the FProperty APIs only to get the address of the properties on a sparse class data object that EditConditions required to be tested.

diff --git a/Engine/Source/Editor/PropertyEditor/Private/EditConditionContext.cpp b/Engine/Source/Editor/PropertyEditor/Private/EditConditionContext.cpp
--- a/Engine/Source/Editor/PropertyEditor/Private/EditConditionContext.cpp
+++ b/Engine/Source/Editor/PropertyEditor/Private/EditConditionContext.cpp
@@ -163,23 +163,9 @@ static const FPropertyNode* GetEditConditionParentNode(const TSharedPtr<FPropert
 		return ParentNode->GetParentNode();
 	}
 
-	if(PropertyNode->HasNodeFlags(EPropertyNodeFlags::IsSparseClassData))
-	{
-		// by using ourself as the parent node we allow GetValueAddressFromObject and friends to find the right c object,
-		// see GetValuePtrOfInstance, which calls InParentNode->GetValueAddressFromObject(Obj). If we return our 
-		// notional (UI) parent node, we will attempt to index the object with a sparse class data property:
-		return PropertyNode.Get(); 
-	}
-
 	return ParentNode;
 }
 
-static const uint8* GetPropertyValuePtr(const FProperty* Property, const TSharedPtr<FPropertyNode>& PropertyNode, const FPropertyNode* ParentNode, const FComplexPropertyNode* ComplexParentNode, int32 Index)
-{
-	const uint8* ValuePtr = ComplexParentNode->GetValuePtrOfInstance(Index, Property, ParentNode);
-	return ValuePtr;
-}
-
 TOptional<bool> FEditConditionContext::GetBoolValue(const FString& PropertyName, TWeakObjectPtr<UFunction> CachedFunction) const
 {
 	TSharedPtr<FPropertyNode> PinnedNode = PropertyNode.Pin();
@@ -254,45 +240,81 @@ TOptional<bool> FEditConditionContext::GetBoolValue(const FString& PropertyName,
 		}
 	}
 
-	const FBoolProperty* BoolProperty = FindTypedField<FBoolProperty>(PinnedNode, PropertyName);
-	if (BoolProperty == nullptr)
+	const FBoolProperty* OperandProperty = FindTypedField<FBoolProperty>(PinnedNode, PropertyName);
+	if (OperandProperty == nullptr)
 	{
 		return TOptional<bool>();
 	}
 
-	const FPropertyNode* ParentNode = GetEditConditionParentNode(PinnedNode);
-	if (ParentNode == nullptr)
-	{
-		return TOptional<bool>();
-	}
-
-	const FComplexPropertyNode* ComplexParentNode = PinnedNode->FindComplexParent();
+	FComplexPropertyNode* ComplexParentNode = PinnedNode->FindComplexParent();
 	if (ComplexParentNode == nullptr)
 	{
 		return TOptional<bool>();
 	}
 
 	TOptional<bool> Result;
-	for (int32 Index = 0; Index < ComplexParentNode->GetInstancesNum(); ++Index)
+	if (!PinnedNode->HasNodeFlags(EPropertyNodeFlags::IsSparseClassData))
 	{
-		const uint8* ValuePtr = GetPropertyValuePtr(BoolProperty, PinnedNode, ParentNode, ComplexParentNode, Index);
-		if (ValuePtr == nullptr)
+		const FPropertyNode* ParentNode = GetEditConditionParentNode(PinnedNode);
+		if (ParentNode == nullptr)
 		{
 			return TOptional<bool>();
 		}
-
-		bool bValue = BoolProperty->GetPropertyValue(ValuePtr);
-		if (!Result.IsSet())
+		
+		for (int32 Index = 0; Index < ComplexParentNode->GetInstancesNum(); ++Index)
 		{
-			Result = bValue;
+			const uint8* ValuePtr = ComplexParentNode->GetValuePtrOfInstance(Index, OperandProperty, ParentNode);
+			if (ValuePtr == nullptr)
+			{
+				return TOptional<bool>();
+			}
+
+			bool bValue = OperandProperty->GetPropertyValue(ValuePtr);
+			if (!Result.IsSet())
+			{
+				Result = bValue;
+			}
+			else if (Result.GetValue() != bValue)
+			{
+				// all values aren't the same...
+				return TOptional<bool>();
+			}
 		}
-		else if (Result.GetValue() != bValue)
+	}
+	else
+	{
+		// The ObjectPropertyNode contains the address of the sparse class data that is needed to resolve the value
+		FObjectPropertyNode* ParentObjectNode = ComplexParentNode->AsObjectNode();
+		if (ParentObjectNode == nullptr) { return TOptional<bool>(); }
+		
+		for (int32 Index = 0; Index < ParentObjectNode->GetInstancesNum(); ++Index)
 		{
-			// all values aren't the same...
-			return TOptional<bool>();
+			UObject* Object = ParentObjectNode->GetUObject(Index);
+			if (Object == nullptr) { return TOptional<bool>(); }
+			
+			UClass* ObjectClass = Object->GetClass();
+			if (ObjectClass == nullptr) { return TOptional<bool>(); }
+			
+			const void* SparseClassDataInstance = ObjectClass->GetSparseClassData(EGetSparseClassDataMethod::ReturnIfNull);
+			// SparseClassData instance hasn't been initialized for this object's class
+			if (SparseClassDataInstance == nullptr) { return TOptional<bool>(); }
+			
+			// Get the address of the value where the property is stored within the instance of the Sparse Class Data
+			const int32 ArrayIndex = 0; // EditConditions does not allow indexing into arrays
+			const uint8* ValuePtr = OperandProperty->ContainerPtrToValuePtr<uint8>(SparseClassDataInstance, ArrayIndex);
+			
+			bool Value = OperandProperty->GetPropertyValue(ValuePtr);
+			
+			if (!Result.IsSet())
+			{
+				Result = Value;
+			}
+			else if (Result.GetValue() != Value)
+			{
+				return TOptional<bool>();
+			}
 		}
 	}
-
 	return Result;
 }
 
@@ -310,56 +332,92 @@ TOptional<int64> FEditConditionContext::GetIntegerValue(const FString& PropertyN
 		return TOptional<int64>();
 	}
 
-	const FProperty* Property = FindTypedField<FProperty>(PinnedNode, PropertyName);
-	const FNumericProperty* NumericProperty = CastField<FNumericProperty>(Property);
+	const FProperty* Property = FindTypedField<FNumericProperty>(PinnedNode, PropertyName);
+	const FNumericProperty* OperandProperty = CastField<FNumericProperty>(Property);
 
-	if (NumericProperty == nullptr)
+	if (OperandProperty == nullptr)
 	{
 		// Retry with an enum and its underlying property
 		if (const FEnumProperty* EnumProperty = CastField<FEnumProperty>(Property))
 		{
-			NumericProperty = EnumProperty->GetUnderlyingProperty();
+			OperandProperty = EnumProperty->GetUnderlyingProperty();
 		}
 	}
 
-	if (NumericProperty == nullptr || !NumericProperty->IsInteger())
-	{
-		return TOptional<int64>();
-	}
-
-	const FPropertyNode* ParentNode = GetEditConditionParentNode(PinnedNode);
-	if (ParentNode == nullptr)
+	if (OperandProperty == nullptr || !OperandProperty->IsInteger())
 	{
 		return TOptional<int64>();
 	}
 
-	const FComplexPropertyNode* ComplexParentNode = PinnedNode->FindComplexParent();
+	FComplexPropertyNode* ComplexParentNode = PinnedNode->FindComplexParent();
 	if (ComplexParentNode == nullptr)
 	{
 		return TOptional<int64>();
 	}
-
+	
 	TOptional<int64> Result;
-	for (int32 Index = 0; Index < ComplexParentNode->GetInstancesNum(); ++Index)
+	if (!PinnedNode->HasNodeFlags(EPropertyNodeFlags::IsSparseClassData))
 	{
-		const uint8* ValuePtr = GetPropertyValuePtr(Property, PinnedNode, ParentNode, ComplexParentNode, Index);
-		if (ValuePtr == nullptr)
+		const FPropertyNode* ParentNode = GetEditConditionParentNode(PinnedNode);
+		if (ParentNode == nullptr)
 		{
 			return TOptional<int64>();
 		}
 
-		int64 Value = NumericProperty->GetSignedIntPropertyValue(ValuePtr);
-		if (!Result.IsSet())
+		for (int32 Index = 0; Index < ComplexParentNode->GetInstancesNum(); ++Index)
 		{
-			Result = Value;
+			const uint8* ValuePtr = ComplexParentNode->GetValuePtrOfInstance(Index, Property, ParentNode);
+			if (ValuePtr == nullptr)
+			{
+				return TOptional<int64>();
+			}
+
+			int64 Value = OperandProperty->GetSignedIntPropertyValue(ValuePtr);
+			if (!Result.IsSet())
+			{
+				Result = Value;
+			}
+			else if (Result.GetValue() != Value)
+			{
+				// all values aren't the same...
+				return TOptional<int64>();
+			}
 		}
-		else if (Result.GetValue() != Value)
+	}
+	else
+	{
+		// The ObjectPropertyNode contains the address of the sparse class data that is needed to resolve the value
+		FObjectPropertyNode* ParentObjectNode = ComplexParentNode->AsObjectNode();
+		if (ParentObjectNode == nullptr) { return TOptional<int64>(); }
+
+		for (int32 Index = 0; Index < ParentObjectNode->GetInstancesNum(); ++Index)
 		{
-			// all values aren't the same...
-			return TOptional<int64>();
+			UObject* Object = ParentObjectNode->GetUObject(Index);
+			if (Object == nullptr) { return TOptional<int64>(); }
+			
+			UClass* ObjectClass = Object->GetClass();
+			if (ObjectClass == nullptr) { return TOptional<int64>(); }
+			
+			const void* SparseClassDataInstance = ObjectClass->GetSparseClassData(EGetSparseClassDataMethod::ReturnIfNull);
+			// SparseClassData instance hasn't been initialized for this object's class
+			if (SparseClassDataInstance == nullptr) { return TOptional<int64>(); }
+			
+			// Get the address of the value where the property is stored within the instance of the Sparse Class Data
+			const int32 ArrayIndex = 0; // EditConditions does not allow indexing into arrays
+			const uint8* ValuePtr = OperandProperty->ContainerPtrToValuePtr<uint8>(SparseClassDataInstance, ArrayIndex);
+			
+			int64 Value = OperandProperty->GetSignedIntPropertyValue(ValuePtr);
+			
+			if (!Result.IsSet())
+			{
+				Result = Value;
+			}
+			else if (Result.GetValue() != Value)
+			{
+				return TOptional<int64>();
+			}
 		}
 	}
-
 	return Result;
 }
 
@@ -377,52 +435,95 @@ TOptional<double> FEditConditionContext::GetNumericValue(const FString& Property
 		return TOptional<double>();
 	}
 
-	const FNumericProperty* NumericProperty = FindTypedField<FNumericProperty>(PinnedNode, PropertyName);
-	if (NumericProperty == nullptr)
+	const FNumericProperty* OperandProperty = FindTypedField<FNumericProperty>(PinnedNode, PropertyName);
+	if (OperandProperty == nullptr)
 	{
 		return TOptional<double>();
 	}
 
-	const FPropertyNode* ParentNode = GetEditConditionParentNode(PinnedNode);
-	if (ParentNode == nullptr)
-	{
-		return TOptional<double>();
-	}
-
-	const FComplexPropertyNode* ComplexParentNode = PinnedNode->FindComplexParent();
+	FComplexPropertyNode* ComplexParentNode = PinnedNode->FindComplexParent();
 	if (ComplexParentNode == nullptr)
 	{
 		return TOptional<double>();
 	}
 
 	TOptional<double> Result;
-	for (int32 Index = 0; Index < ComplexParentNode->GetInstancesNum(); ++Index)
+	if (!PinnedNode->HasNodeFlags(EPropertyNodeFlags::IsSparseClassData))
 	{
-		const uint8* ValuePtr = GetPropertyValuePtr(NumericProperty, PinnedNode, ParentNode, ComplexParentNode, Index);
-		if (ValuePtr == nullptr)
+		const FPropertyNode* ParentNode = GetEditConditionParentNode(PinnedNode);
+		if (ParentNode == nullptr)
 		{
 			return TOptional<double>();
 		}
+	
+		for (int32 Index = 0; Index < ComplexParentNode->GetInstancesNum(); ++Index)
+		{
+			const uint8* ValuePtr = ComplexParentNode->GetValuePtrOfInstance(Index, OperandProperty, ParentNode);
+			if (ValuePtr == nullptr)
+			{
+				return TOptional<double>();
+			}
 
-		double Value = 0;
+			double Value = 0;
 
-		if (NumericProperty->IsInteger())
-		{
-			Value = (double) NumericProperty->GetSignedIntPropertyValue(ValuePtr);
-		}
-		else if (NumericProperty->IsFloatingPoint())
-		{
-			Value = NumericProperty->GetFloatingPointPropertyValue(ValuePtr);
-		}
+			if (OperandProperty->IsInteger())
+			{
+				Value = (double) OperandProperty->GetSignedIntPropertyValue(ValuePtr);
+			}
+			else if (OperandProperty->IsFloatingPoint())
+			{
+				Value = OperandProperty->GetFloatingPointPropertyValue(ValuePtr);
+			}
 
-		if (!Result.IsSet())
-		{
-			Result = Value;
+			if (!Result.IsSet())
+			{
+				Result = Value;
+			}
+			else if (!FMath::IsNearlyEqual(Result.GetValue(), Value))
+			{
+				// all values aren't the same...
+				return TOptional<double>();
+			}
 		}
-		else if (!FMath::IsNearlyEqual(Result.GetValue(), Value))
+	}
+	else
+	{
+		// The ObjectPropertyNode contains the address of the sparse class data that is needed to resolve the value
+		FObjectPropertyNode* ParentObjectNode = ComplexParentNode->AsObjectNode();
+		if (ParentObjectNode == nullptr) { return TOptional<double>(); }
+		
+		for (int32 Index = 0; Index < ParentObjectNode->GetInstancesNum(); ++Index)
 		{
-			// all values aren't the same...
-			return TOptional<double>();
+			UObject* Object = ParentObjectNode->GetUObject(Index);
+			if (Object == nullptr) { return TOptional<double>(); }
+			
+			UClass* ObjectClass = Object->GetClass();
+			if (ObjectClass == nullptr) { return TOptional<double>(); }
+			
+			const void* SparseClassDataInstance = ObjectClass->GetSparseClassData(EGetSparseClassDataMethod::ReturnIfNull);
+			// SparseClassData instance hasn't been initialized for this object's class
+			if (SparseClassDataInstance == nullptr) { return TOptional<double>(); }
+			
+			// Get the address of the value where the property is stored within the instance of the Sparse Class Data
+			const int32 ArrayIndex = 0; // EditConditions does not allow indexing into arrays
+			const uint8* ValuePtr = OperandProperty->ContainerPtrToValuePtr<uint8>(SparseClassDataInstance, ArrayIndex);
+			double Value = 0;
+			if (OperandProperty->IsInteger())
+			{
+				Value = (double) OperandProperty->GetSignedIntPropertyValue(ValuePtr);
+			}
+			else if (OperandProperty->IsFloatingPoint())
+			{
+				Value = OperandProperty->GetFloatingPointPropertyValue(ValuePtr);
+			}
+			if (!Result.IsSet())
+			{
+				Result = Value;
+			}
+			else if (Result.GetValue() != Value)
+			{
+				return TOptional<double>();
+			}
 		}
 	}
 
@@ -450,55 +551,91 @@ TOptional<FString> FEditConditionContext::GetEnumValue(const FString& PropertyNa
 	}
 
 	const UEnum* EnumType = nullptr;
-	const FNumericProperty* NumericProperty = nullptr;
+	const FNumericProperty* OperandProperty = nullptr;
 	if (const FEnumProperty* EnumProperty = CastField<FEnumProperty>(Property))
 	{
-		NumericProperty = EnumProperty->GetUnderlyingProperty();
+		OperandProperty = EnumProperty->GetUnderlyingProperty();
 		EnumType = EnumProperty->GetEnum();
 	}
 	else if (const FByteProperty* ByteProperty = CastField<FByteProperty>(Property))
 	{
-		NumericProperty = ByteProperty;
+		OperandProperty = ByteProperty;
 		EnumType = ByteProperty->GetIntPropertyEnum();
 	}
 
-	if (EnumType == nullptr || NumericProperty == nullptr || !NumericProperty->IsInteger())
+	if (EnumType == nullptr || OperandProperty == nullptr || !OperandProperty->IsInteger())
 	{
 		return TOptional<FString>();
 	}
 
-	const FPropertyNode* ParentNode = GetEditConditionParentNode(PinnedNode);
-	if (ParentNode == nullptr)
-	{
-		return TOptional<FString>();
-	}
-
-	const FComplexPropertyNode* ComplexParentNode = PinnedNode->FindComplexParent();
+	FComplexPropertyNode* ComplexParentNode = PinnedNode->FindComplexParent();
 	if (ComplexParentNode == nullptr)
 	{
 		return TOptional<FString>();
 	}
 
 	TOptional<int64> Result;
-	for (int32 Index = 0; Index < ComplexParentNode->GetInstancesNum(); ++Index)
+	if (!PinnedNode->HasNodeFlags(EPropertyNodeFlags::IsSparseClassData))
 	{
-		// NOTE: this very intentionally fetches the value from Property, not NumericProperty, 
-		// because the underlying property of an enum does not return a valid value
-		const uint8* ValuePtr = GetPropertyValuePtr(Property, PinnedNode, ParentNode, ComplexParentNode, Index);
-		if (ValuePtr == nullptr)
+		const FPropertyNode* ParentNode = GetEditConditionParentNode(PinnedNode);
+		if (ParentNode == nullptr)
 		{
 			return TOptional<FString>();
 		}
-
-		const int64 Value = NumericProperty->GetSignedIntPropertyValue(ValuePtr);
-		if (!Result.IsSet())
+	
+		for (int32 Index = 0; Index < ComplexParentNode->GetInstancesNum(); ++Index)
 		{
-			Result = Value;
+			// NOTE: this very intentionally fetches the value from Property, not NumericProperty, 
+			// because the underlying property of an enum does not return a valid value
+			const uint8* ValuePtr = ComplexParentNode->GetValuePtrOfInstance(Index, Property, ParentNode);
+		
+			if (ValuePtr == nullptr)
+			{
+				return TOptional<FString>();
+			}
+
+			const int64 Value = OperandProperty->GetSignedIntPropertyValue(ValuePtr);
+			if (!Result.IsSet())
+			{
+				Result = Value;
+			}
+			else if (Result.GetValue() != Value)
+			{
+				// all values aren't the same...
+				return TOptional<FString>();
+			}
 		}
-		else if (Result.GetValue() != Value)
+	}
+	else
+	{
+		// The ObjectPropertyNode contains the address of the sparse class data that is needed to resolve the value
+		FObjectPropertyNode* ParentObjectNode = ComplexParentNode->AsObjectNode();
+		if (ParentObjectNode == nullptr) { return TOptional<FString>(); }
+		
+		for (int32 Index = 0; Index < ParentObjectNode->GetInstancesNum(); ++Index)
 		{
-			// all values aren't the same...
-			return TOptional<FString>();
+			UObject* Object = ParentObjectNode->GetUObject(Index);
+			if (Object == nullptr) { return TOptional<FString>(); }
+			
+			UClass* ObjectClass = Object->GetClass();
+			if (ObjectClass == nullptr) { return TOptional<FString>(); }
+			
+			const void* SparseClassDataInstance = ObjectClass->GetSparseClassData(EGetSparseClassDataMethod::ReturnIfNull);
+			// SparseClassData instance hasn't been initialized for this object's class
+			if (SparseClassDataInstance == nullptr) { return TOptional<FString>(); }
+			
+			// Get the address of the value where the property is stored within the instance of the Sparse Class Data
+			const int32 ArrayIndex = 0; // EditConditions does not allow indexing into arrays of UObject*
+			const uint8* ValuePtr = OperandProperty->ContainerPtrToValuePtr<uint8>(SparseClassDataInstance, ArrayIndex);
+			const int64 Value = OperandProperty->GetSignedIntPropertyValue(ValuePtr);
+			if (!Result.IsSet())
+			{
+				Result = Value;
+			}
+			else if (Result.GetValue() != Value)
+			{
+				return TOptional<FString>();
+			}
 		}
 	}
 
@@ -523,53 +660,85 @@ TOptional<UObject*> FEditConditionContext::GetPointerValue(const FString& Proper
 		// EditConditions Currently only support bool, see: UE-175891
 		return TOptional<UObject*>();
 	}
-
-	const FProperty* Property = FindTypedField<FProperty>(PinnedNode, PropertyName);
-	if (Property == nullptr)
-	{
-		return TOptional<UObject*>();
-	}
-
-	const FObjectPropertyBase* ObjectProperty = CastField<FObjectPropertyBase>(Property);
-	if (ObjectProperty == nullptr)
-	{
-		return TOptional<UObject*>();
-	}
-
-	const FPropertyNode* ParentNode = GetEditConditionParentNode(PinnedNode);
-	if (ParentNode == nullptr)
+	
+	FComplexPropertyNode* ComplexParentNode = PinnedNode->FindComplexParent();
+	if (ComplexParentNode == nullptr)
 	{
 		return TOptional<UObject*>();
 	}
 
-	const FComplexPropertyNode* ComplexParentNode = PinnedNode->FindComplexParent();
-	if (ComplexParentNode == nullptr)
+	// Get the property of the EditCondition operand
+	// EditCondition pointers can only be UObjects
+	const FObjectPropertyBase* OperandProperty = FindTypedField<FObjectPropertyBase>(PinnedNode, PropertyName);
+	if (OperandProperty == nullptr)
 	{
 		return TOptional<UObject*>();
 	}
 
-	TOptional<UObject*> Result;
-	for (int32 Index = 0; Index < ComplexParentNode->GetInstancesNum(); ++Index)
+	if (!PinnedNode->HasNodeFlags(EPropertyNodeFlags::IsSparseClassData))
 	{
-		const uint8* ValuePtr = GetPropertyValuePtr(Property, PinnedNode, ParentNode, ComplexParentNode, Index);
-		if (ValuePtr == nullptr)
+		const FPropertyNode* ParentNode = GetEditConditionParentNode(PinnedNode);
+		if (ParentNode == nullptr)
 		{
 			return TOptional<UObject*>();
 		}
 
-		UObject* Value = ObjectProperty->GetObjectPropertyValue(ValuePtr);
-		if (!Result.IsSet())
+		TOptional<UObject*> Result;
+		for (int32 Index = 0; Index < ComplexParentNode->GetInstancesNum(); ++Index)
 		{
-			Result = Value;
+			const uint8* ValuePtr = ComplexParentNode->GetValuePtrOfInstance(Index, OperandProperty, ParentNode);
+			if (ValuePtr == nullptr)
+			{
+				return TOptional<UObject*>();
+			}
+
+			UObject* Value = OperandProperty->GetObjectPropertyValue(ValuePtr);
+			if (!Result.IsSet())
+			{
+				Result = Value;
+			}
+			else if (Result.GetValue() != Value)
+			{
+				// all values aren't the same
+				return TOptional<UObject*>();
+			}
 		}
-		else if (Result.GetValue() != Value)
+		return Result;
+	}
+	else
+	{
+		// The ObjectPropertyNode contains the address of the sparse class data that is needed to resolve the value
+		FObjectPropertyNode* ParentObjectNode = ComplexParentNode->AsObjectNode();
+		if (ParentObjectNode == nullptr) { return TOptional<UObject*>(); }
+
+		TOptional<UObject*> Result;
+		for (int32 Index = 0; Index < ParentObjectNode->GetInstancesNum(); ++Index)
 		{
-			// all values aren't the same
-			return TOptional<UObject*>();
+			UObject* Object = ParentObjectNode->GetUObject(Index);
+			if (Object == nullptr) { return TOptional<UObject*>(); }
+			
+			UClass* ObjectClass = Object->GetClass();
+			if (ObjectClass == nullptr) { return TOptional<UObject*>(); }
+			
+			const void* SparseClassDataInstance = ObjectClass->GetSparseClassData(EGetSparseClassDataMethod::ReturnIfNull);
+			// SparseClassData instance hasn't been initialized for this object's class
+			if (SparseClassDataInstance == nullptr) { return TOptional<UObject*>(); }
+			
+			// Get the address of the value where the property is stored within the instance of the Sparse Class Data
+			const int32 ArrayIndex = 0; // EditConditions does not allow indexing into arrays of UObject*
+			const uint8* ValuePtr = OperandProperty->ContainerPtrToValuePtr<uint8>(SparseClassDataInstance, ArrayIndex);
+			UObject* Value = OperandProperty->GetObjectPropertyValue(ValuePtr);
+			if (!Result.IsSet())
+			{
+				Result = Value;
+			}
+			else if (Result.GetValue() != Value)
+			{
+				return TOptional<UObject*>();
+			}
 		}
+		return Result;
 	}
-
-	return Result;
 }
 
 TOptional<FString> FEditConditionContext::GetTypeName(const FString& PropertyName, TWeakObjectPtr<UFunction> CachedFunction) const
