Support ExternalObjectProperty in DetailGroup

#rnx

diff --git a/Engine/Source/Editor/PropertyEditor/Private/DetailGroup.cpp b/Engine/Source/Editor/PropertyEditor/Private/DetailGroup.cpp
--- a/Engine/Source/Editor/PropertyEditor/Private/DetailGroup.cpp
+++ b/Engine/Source/Editor/PropertyEditor/Private/DetailGroup.cpp
@@ -63,6 +63,19 @@ IDetailPropertyRow& FDetailGroup::AddPropertyRow( TSharedRef<IPropertyHandle> Pr
 	return *NewCustomization.PropertyRow;
 }
 
+IDetailPropertyRow& FDetailGroup::AddExternalObjectProperty(const TArray<UObject*>& Objects, FName PropertyName, EPropertyLocation::Type Location, const FAddPropertyParams& Params)
+{
+	FDetailLayoutCustomization NewCustomization;
+	NewCustomization.bCustom = true;
+	NewCustomization.bAdvanced = Location == EPropertyLocation::Advanced;
+
+	FDetailPropertyRow::MakeExternalPropertyRowCustomization(Objects, PropertyName, ParentCategory.Pin().ToSharedRef(), NewCustomization, Params);
+
+	GroupChildren.Add( NewCustomization );
+
+	return *NewCustomization.PropertyRow;
+}
+
 TSharedPtr<IDetailPropertyRow> FDetailGroup::FindPropertyRow(TSharedRef<IPropertyHandle> PropertyHandle) const
 {
 	for (const FDetailLayoutCustomization& Customization : GroupChildren)
diff --git a/Engine/Source/Editor/PropertyEditor/Private/DetailGroup.h b/Engine/Source/Editor/PropertyEditor/Private/DetailGroup.h
--- a/Engine/Source/Editor/PropertyEditor/Private/DetailGroup.h
+++ b/Engine/Source/Editor/PropertyEditor/Private/DetailGroup.h
@@ -25,6 +25,7 @@ public:
 	virtual IDetailPropertyRow& HeaderProperty( TSharedRef<IPropertyHandle> PropertyHandle ) override;
 	virtual FDetailWidgetRow& AddWidgetRow() override;
 	virtual IDetailPropertyRow& AddPropertyRow( TSharedRef<IPropertyHandle> PropertyHandle ) override;
+	virtual IDetailPropertyRow& AddExternalObjectProperty(const TArray<UObject*>& Objects, FName PropertyName, EPropertyLocation::Type Location, const FAddPropertyParams& Params) override;
 	virtual IDetailGroup& AddGroup(FName NewGroupName, const FText& InLocalizedDisplayName, bool bInStartExpanded = false) override;
 	virtual const TOptional<FText>& GetToolTip() const override;
 	virtual void SetToolTip(const FText& ToolTip) override;
diff --git a/Engine/Source/Editor/PropertyEditor/Public/IDetailGroup.h b/Engine/Source/Editor/PropertyEditor/Public/IDetailGroup.h
--- a/Engine/Source/Editor/PropertyEditor/Public/IDetailGroup.h
+++ b/Engine/Source/Editor/PropertyEditor/Public/IDetailGroup.h
@@ -53,6 +53,13 @@ public:
 	 */
 	virtual IDetailPropertyRow& AddPropertyRow(TSharedRef<IPropertyHandle> PropertyHandle) = 0;
 
+	/**
+	 * Adds a new row for an external property
+	 *
+	 * @return an interface for customizing the appearance of the property row
+	 */
+	virtual IDetailPropertyRow& AddExternalObjectProperty(const TArray<UObject*>& Objects, FName PropertyName, EPropertyLocation::Type Location, const FAddPropertyParams& Params) = 0;
+
 	/**
 	 * Adds a group to the category
 	 *
