PropertyNode::ToggleEditConditionState accessing invalid data offset when EditCondition is a Sparse Class Data member

Previously, MyParentNode->GetValueAddressFromObject would be called when the PropertyNode had a IsSparseClassData flag. The intention of this function usually is to return the address of the value pointed to by MyParentNode.  In the presence of IsSparseClassData, this will pass the start address of SparseClassData structure to FPropertyNode::GetValueAddress.  Problems arise when FItemPropertyNode::GetValueBaseAddress uses its own property offset to offset the address from the SparseClassData struct - which points to an unintended address.

diff --git a/Engine/Source/Editor/PropertyEditor/Private/PropertyNode.cpp b/Engine/Source/Editor/PropertyEditor/Private/PropertyNode.cpp
--- a/Engine/Source/Editor/PropertyEditor/Private/PropertyNode.cpp
+++ b/Engine/Source/Editor/PropertyEditor/Private/PropertyNode.cpp
@@ -1552,13 +1552,39 @@ void FPropertyNode::ToggleEditConditionState()
 
 	bool OldValue = true;
 
+	const bool bIsSparse = HasNodeFlags(EPropertyNodeFlags::IsSparseClassData);
+	
 	FComplexPropertyNode* ComplexParentNode = FindComplexParent();
-	for (int32 Index = 0; Index < ComplexParentNode->GetInstancesNum(); ++Index)
+
+	if (!bIsSparse)
 	{
-		uint8* ValuePtr = ComplexParentNode->GetValuePtrOfInstance(Index, EditConditionProperty, MyParentNode);
+		
+		for (int32 Index = 0; Index < ComplexParentNode->GetInstancesNum(); ++Index)
+		{
+			uint8* ValuePtr = ComplexParentNode->GetValuePtrOfInstance(Index, EditConditionProperty, MyParentNode);
+
+			OldValue &= EditConditionProperty->GetPropertyValue(ValuePtr);
+			EditConditionProperty->SetPropertyValue(ValuePtr, !OldValue);
+		}
+	}
+	else
+	{
+		FObjectPropertyNode* ObjectPropertyNode = ComplexParentNode->AsObjectNode();
+		if (ObjectPropertyNode)
+		{
+			for (int32 Index = 0; Index < ObjectPropertyNode->GetInstancesNum(); ++Index)
+			{
+				UObject* Object = ObjectPropertyNode->GetUObject(Index);
+				uint8* SparseClassData = GetStartAddressFromObject(Object);
 
-		OldValue &= EditConditionProperty->GetPropertyValue(ValuePtr);
-		EditConditionProperty->SetPropertyValue(ValuePtr, !OldValue);
+				// Compute offset from the object containing the EditConditionProperty to the property itself
+				uint8* ValuePtr = EditConditionProperty->ContainerPtrToValuePtr<uint8>(SparseClassData);
+
+				// Toggle the value
+				OldValue &= EditConditionProperty->GetPropertyValue(ValuePtr);
+				EditConditionProperty->SetPropertyValue(ValuePtr, !OldValue);
+			}
+		}
 	}
 
 	// Propagate the value change to any instances if we're editing a template object
@@ -1573,9 +1599,21 @@ void FPropertyNode::ToggleEditConditionState()
 			{
 				TArray<UObject*> ArchetypeInstances;
 				Object->GetArchetypeInstances(ArchetypeInstances);
+				
+				
 				for (int32 InstanceIndex = 0; InstanceIndex < ArchetypeInstances.Num(); ++InstanceIndex)
 				{
-					uint8* ArchetypeBaseOffset = MyParentNode->GetValueAddressFromObject(ArchetypeInstances[InstanceIndex]);
+					uint8* ArchetypeBaseOffset;
+					if (!bIsSparse)
+					{
+						ArchetypeBaseOffset = MyParentNode->GetValueAddressFromObject(ArchetypeInstances[InstanceIndex]);
+					}
+					else
+					{
+						uint8* SparseClassData = GetStartAddressFromObject(ArchetypeInstances[InstanceIndex]);
+						ArchetypeBaseOffset = SparseClassData;
+					}
+					
 					uint8* ArchetypeValueAddr = EditConditionProperty->ContainerPtrToValuePtr<uint8>(ArchetypeBaseOffset);
 
 					// Only propagate if the current value on the instance matches the previous value on the template.
