PR #13025: Prevent Add Item for TMap and TSet from adding enum values excluded by ValidEnumValues and InvalidEnumValues metadata on UProperties

#13025

diff --git a/Engine/Source/Editor/PropertyEditor/Private/PropertyHandleImpl.cpp b/Engine/Source/Editor/PropertyEditor/Private/PropertyHandleImpl.cpp
--- a/Engine/Source/Editor/PropertyEditor/Private/PropertyHandleImpl.cpp
+++ b/Engine/Source/Editor/PropertyEditor/Private/PropertyHandleImpl.cpp
@@ -5944,7 +5944,6 @@ bool FPropertyHandleSet::HasDefaultElement()
 
 FPropertyHandleItemAddResult FPropertyHandleSet::AddItem()
 {
-	FPropertyAccess::Result Result = FPropertyAccess::Fail;
 	if (IsEditable())
 	{
 		/**
@@ -6001,15 +6000,34 @@ FPropertyHandleItemAddResult FPropertyHandleSet::AddItem()
 
 				if (const FEnumProperty* EnumProperty = CastField<const FEnumProperty>(SetProperty->GetElementProperty()))
 				{
-					if (UEnum* Enum = EnumProperty->GetEnum())
+					if (const UEnum* Enum = EnumProperty->GetEnum())
 					{
-						static FString HiddenName(TEXT("Hidden"));
-						static FString SpacerName(TEXT("Spacer"));
+						const TArray<FName> ValidPropertyEnums = PropertyEditorHelpers::GetValidEnumsFromPropertyOverride(EnumProperty, Enum);
+						const TArray<FName> InvalidPropertyEnums = PropertyEditorHelpers::GetInvalidEnumsFromPropertyOverride(EnumProperty, Enum);
+
 						// NumEnums() - 1, because the last item in an enum is the _MAX item
 						for (int32 EnumIndex = 0; EnumIndex < Enum->NumEnums() - 1; ++EnumIndex)
 						{
-							// Skip hidden and spacers so we don't add them.
-							const bool bShouldBeHidden = Enum->HasMetaData(*HiddenName, EnumIndex) || Enum->HasMetaData(*SpacerName, EnumIndex);
+							static FString HiddenName(TEXT("Hidden"));
+							static FString SpacerName(TEXT("Spacer"));
+
+							// Ignore hidden enums
+							bool bShouldBeHidden = Enum->HasMetaData(*HiddenName, EnumIndex) || Enum->HasMetaData(*SpacerName, EnumIndex);
+
+							if (!bShouldBeHidden)
+							{
+								if (ValidPropertyEnums.Num() > 0)
+								{
+									bShouldBeHidden = ValidPropertyEnums.Find(Enum->GetNameByIndex(EnumIndex)) == INDEX_NONE;
+								}
+
+								// If both are specified, the metadata "InvalidEnumValues" takes precedence
+								if (InvalidPropertyEnums.Num() > 0)
+								{
+									bShouldBeHidden = InvalidPropertyEnums.Find(Enum->GetNameByIndex(EnumIndex)) != INDEX_NONE;
+								}
+							}
+
 							if (bShouldBeHidden)
 							{
 								continue;
@@ -6031,22 +6049,21 @@ FPropertyHandleItemAddResult FPropertyHandleSet::AddItem()
 								{
 									// Set the new entry to have the element value.
 									Implementation->ImportText(Enum->GetNameStringByIndex(EnumIndex), ChildNode.Get(), EPropertyValueSetFlags::DefaultFlags);
-									Result = FPropertyAccess::Success;
-									return FPropertyHandleItemAddResult(ChildNodeIndex, Result);
+									return FPropertyHandleItemAddResult(ChildNodeIndex, FPropertyAccess::Success);
 								}
 							}
 						}
 					}
+
+					Implementation->ShowInvalidOperationError(LOCTEXT("EnumSet_Add_AllValuesAdded", "Cannot add a new element to the set as all the enum values are already added"));
+					return FPropertyHandleItemAddResult();
 				}
-				else
+
+				if (!HasDefaultElement())
 				{
-					if (!HasDefaultElement())
-					{
-						int32 Index = Implementation->AddChild();
-						Implementation->GetPropertyNode()->RebuildChildren();
-						Result = FPropertyAccess::Success;
-						return FPropertyHandleItemAddResult(Index, Result);
-					}
+					const int32 Index = Implementation->AddChild();
+					Implementation->GetPropertyNode()->RebuildChildren();
+					return FPropertyHandleItemAddResult(Index, FPropertyAccess::Success);
 				}
 			}
 		}
@@ -6188,7 +6205,6 @@ bool FPropertyHandleMap::HasDefaultKey()
 
 FPropertyHandleItemAddResult FPropertyHandleMap::AddItem()
 {
-	FPropertyAccess::Result Result = FPropertyAccess::Fail;
 	if (IsEditable())
 	{
 		/**
@@ -6266,23 +6282,41 @@ FPropertyHandleItemAddResult FPropertyHandleMap::AddItem()
 							const int32 ChildNodeIndex = AddItemWithKey(StringValue);
 							if (ChildNodeIndex != INDEX_NONE)
 							{
-								Result = FPropertyAccess::Success;
-								return FPropertyHandleItemAddResult(ChildNodeIndex, Result);
+								return FPropertyHandleItemAddResult(ChildNodeIndex, FPropertyAccess::Success);
 							}
 						}
 					}
 				}
 				else if (const FEnumProperty* EnumProperty = CastField<const FEnumProperty>(MapProperty->GetKeyProperty()))
 				{
-					if (UEnum* Enum = EnumProperty->GetEnum())
+					if (const UEnum* Enum = EnumProperty->GetEnum())
 					{
-						static FString HiddenName(TEXT("Hidden"));
-						static FString SpacerName(TEXT("Spacer"));
+						const TArray<FName> ValidPropertyEnums = PropertyEditorHelpers::GetValidEnumsFromPropertyOverride(EnumProperty, Enum);
+						const TArray<FName> InvalidPropertyEnums = PropertyEditorHelpers::GetInvalidEnumsFromPropertyOverride(EnumProperty, Enum);
+
 						// NumEnums() - 1, because the last item in an enum is the _MAX item
 						for (int32 EnumIndex = 0; EnumIndex < Enum->NumEnums() - 1; ++EnumIndex)
 						{
-							// Skip hidden and spacers so we don't add them.
-							const bool bShouldBeHidden = Enum->HasMetaData(*HiddenName, EnumIndex) || Enum->HasMetaData(*SpacerName, EnumIndex);
+							static FString HiddenName(TEXT("Hidden"));
+							static FString SpacerName(TEXT("Spacer"));
+
+							// Ignore hidden enums
+							bool bShouldBeHidden = Enum->HasMetaData(*HiddenName, EnumIndex) || Enum->HasMetaData(*SpacerName, EnumIndex);
+
+							if (!bShouldBeHidden)
+							{
+								if (ValidPropertyEnums.Num() > 0)
+								{
+									bShouldBeHidden = ValidPropertyEnums.Find(Enum->GetNameByIndex(EnumIndex)) == INDEX_NONE;
+								}
+
+								// If both are specified, the metadata "InvalidEnumValues" takes precedence
+								if (InvalidPropertyEnums.Num() > 0)
+								{
+									bShouldBeHidden = InvalidPropertyEnums.Find(Enum->GetNameByIndex(EnumIndex)) != INDEX_NONE;
+								}
+							}
+
 							if (bShouldBeHidden)
 							{
 								continue;
@@ -6297,10 +6331,12 @@ FPropertyHandleItemAddResult FPropertyHandleMap::AddItem()
 							const int32 ChildNodeIndex = AddItemWithKey(Enum->GetNameStringByIndex(EnumIndex));
 							if (ChildNodeIndex != INDEX_NONE)
 							{
-								Result = FPropertyAccess::Success;
-								return FPropertyHandleItemAddResult(ChildNodeIndex, Result);
+								return FPropertyHandleItemAddResult(ChildNodeIndex, FPropertyAccess::Success);
 							}
 						}
+
+						Implementation->ShowInvalidOperationError(LOCTEXT("EnumMap_Add_AllValuesAdded", "Cannot add a new element to the map as all the enum values are already added"));
+						return FPropertyHandleItemAddResult();
 					}
 				}
 				else
@@ -6308,8 +6344,7 @@ FPropertyHandleItemAddResult FPropertyHandleMap::AddItem()
 					if (!HasDefaultKey())
 					{
 						const int32 Index = Implementation->AddChild();
-						Result = FPropertyAccess::Success;
-						return FPropertyHandleItemAddResult(Index, Result);
+						return FPropertyHandleItemAddResult(Index, FPropertyAccess::Success);
 					}
 				}
 			}
