PR #11475: Fixed copy-paste in the property matrix not respecting instanced properties on paste

diff --git a/Engine/Source/Editor/PropertyEditor/Private/Presentation/PropertyTable/PropertyTableCell.cpp b/Engine/Source/Editor/PropertyEditor/Private/Presentation/PropertyTable/PropertyTableCell.cpp
--- a/Engine/Source/Editor/PropertyEditor/Private/Presentation/PropertyTable/PropertyTableCell.cpp
+++ b/Engine/Source/Editor/PropertyEditor/Private/Presentation/PropertyTable/PropertyTableCell.cpp
@@ -136,7 +136,7 @@ void FPropertyTableCell::SetValueFromString( const FString& InString )
 	}
 
 	TSharedRef< IPropertyHandle > Handle = PropertyEditor->GetPropertyHandle();
-	Handle->SetValueFromFormattedString( InString );
+	Handle->SetValueFromFormattedString(InString, EPropertyValueSetFlags::InstanceObjects);
 }
 
 TWeakObjectPtr< UObject > FPropertyTableCell::GetObject() const
