PR #12599: Find an appropriate widget for FLinearColor properties when customizing the display of properties.

diff --git a/Engine/Source/Editor/PropertyEditor/Private/PropertyEditorHelpers.cpp b/Engine/Source/Editor/PropertyEditor/Private/PropertyEditorHelpers.cpp
--- a/Engine/Source/Editor/PropertyEditor/Private/PropertyEditorHelpers.cpp
+++ b/Engine/Source/Editor/PropertyEditor/Private/PropertyEditorHelpers.cpp
@@ -11,6 +11,7 @@
 #include "UserInterface/PropertyEditor/SPropertyEditor.h"
 #include "UserInterface/PropertyEditor/SPropertyEditorNumeric.h"
 #include "UserInterface/PropertyEditor/SPropertyEditorArray.h"
+#include "UserInterface/PropertyEditor/SPropertyEditorColor.h"
 #include "UserInterface/PropertyEditor/SPropertyEditorCombo.h"
 #include "UserInterface/PropertyEditor/SPropertyEditorEditInline.h"
 #include "UserInterface/PropertyEditor/SPropertyEditorText.h"
@@ -310,6 +311,11 @@ TSharedRef<SWidget> SPropertyValueWidget::ConstructPropertyEditorWidget( TShared
 			BoolWidget->GetDesiredWidth( MinDesiredWidth, MaxDesiredWidth );
 
 		}
+		else if ( SPropertyEditorColor::Supports( PropertyEditorRef ) )
+		{
+			TSharedRef<SPropertyEditorColor> ColorWidget = 
+				SAssignNew( PropertyWidget, SPropertyEditorColor, PropertyEditorRef, InPropertyUtilities.ToSharedRef() );
+		}
 		else if ( SPropertyEditorArrayItem::Supports( PropertyEditorRef ) )
 		{
 			TSharedRef<SPropertyEditorArrayItem> ArrayItemWidget = 
