PR #11039: SObjectPropertyEntryBox considers property handle's metadata

diff --git a/Engine/Source/Editor/PropertyEditor/Private/PropertyCustomizationHelpers.cpp b/Engine/Source/Editor/PropertyEditor/Private/PropertyCustomizationHelpers.cpp
--- a/Engine/Source/Editor/PropertyEditor/Private/PropertyCustomizationHelpers.cpp
+++ b/Engine/Source/Editor/PropertyEditor/Private/PropertyCustomizationHelpers.cpp
@@ -892,14 +892,14 @@ void SObjectPropertyEntryBox::Construct( const FArguments& InArgs )
 		PropertyHandle = InArgs._PropertyHandle;
 
 		// check if the property metadata wants us to display a thumbnail
-		const FString& DisplayThumbnailString = PropertyHandle->GetProperty()->GetMetaData(TEXT("DisplayThumbnail"));
+		const FString& DisplayThumbnailString = PropertyHandle->GetMetaData(TEXT("DisplayThumbnail"));
 		if(DisplayThumbnailString.Len() > 0)
 		{
 			bDisplayThumbnail = DisplayThumbnailString == TEXT("true");
 		}
 
 		// check if the property metadata has an override to the thumbnail size
-		const FString& ThumbnailSizeString = PropertyHandle->GetProperty()->GetMetaData(TEXT("ThumbnailSize"));
+		const FString& ThumbnailSizeString = PropertyHandle->GetMetaData(TEXT("ThumbnailSize"));
 		if ( ThumbnailSizeString.Len() > 0 )
 		{
 			FVector2D ParsedVector;
