Fix SparseClassData EditCondition when the condition is within a substructure of the SparseClassData

diff --git a/Engine/Source/Editor/PropertyEditor/Private/EditConditionContext.cpp b/Engine/Source/Editor/PropertyEditor/Private/EditConditionContext.cpp
--- a/Engine/Source/Editor/PropertyEditor/Private/EditConditionContext.cpp
+++ b/Engine/Source/Editor/PropertyEditor/Private/EditConditionContext.cpp
@@ -166,6 +166,35 @@ static const FPropertyNode* GetEditConditionParentNode(const TSharedPtr<FPropert
 	return ParentNode;
 }
 
+static uint8* ComputeEditConditionValuePointer(FPropertyNode& ConditionedProperty, FReadAddressList& ConditionPropertyAddresses, const FProperty& EditConditionProperty, int32 AddressIndex)
+{
+	// The strategy for getting the pointer to the vaue of the edit condition is to use the value pointer of the property
+	// that is being conditioned and then walk back that property's offset to find the owning struct's base address. 
+	// This owning struct's base address is then offset forward by the edit condition's property to find the edit condition's value pointer.
+	// 
+	// The assumption is that the edit condition property is on the same struct/class as the property that is being "conditioned".
+	// 
+	// Since the edit condition is inline, there will not be a PropertyNode available, therefore it is necessary to use
+	// the FProperty API to find the value pointer for the edit condition.
+	//
+	// It is also not possible to use the conditioned property's parent node which is typically used in the non-sparse class data case
+	// since the parent may not point to the struct that either property is in.  
+	// In the case of property directly on the SparseClassData, the parent node is often a category or object node.
+	// In the case of a property within a sub-struct in an array property of the SparseClassData, the offset between the SparseClassData instance
+	// and the edit condition property depends on the property chain.
+	// 
+	// Therefore, it is easiest to find the direct owning struct address that is shared by both conditioned and edit condition property and compute
+	// the edit condition value pointer directly.
+		
+	uint8* ConditionedPropertyValueAddress = ConditionPropertyAddresses.GetAddress(AddressIndex);
+	const int32 ConditionedPropertyOffset = ConditionedProperty.GetProperty()->GetOffset_ForInternal();
+	uint8* OwningStructStartAddress = ConditionedPropertyValueAddress - ConditionedPropertyOffset;
+
+	const int32 ArrayIndex = 0; // EditConditions does not allow indexing into arrays
+	uint8* EditConditionValuePtr = EditConditionProperty.ContainerPtrToValuePtr<uint8>(OwningStructStartAddress, ArrayIndex);
+	return EditConditionValuePtr;
+}
+
 TOptional<bool> FEditConditionContext::GetBoolValue(const FString& PropertyName, TWeakObjectPtr<UFunction> CachedFunction) const
 {
 	TSharedPtr<FPropertyNode> PinnedNode = PropertyNode.Pin();
@@ -283,28 +312,14 @@ TOptional<bool> FEditConditionContext::GetBoolValue(const FString& PropertyName,
 	}
 	else
 	{
-		// The ObjectPropertyNode contains the address of the sparse class data that is needed to resolve the value
-		FObjectPropertyNode* ParentObjectNode = ComplexParentNode->AsObjectNode();
-		if (ParentObjectNode == nullptr) { return TOptional<bool>(); }
-		
-		for (int32 Index = 0; Index < ParentObjectNode->GetInstancesNum(); ++Index)
+		FReadAddressList ThisPropertyReadAddresses;
+		PinnedNode->GetReadAddress(ThisPropertyReadAddresses);
+
+		for (int32 Index = 0; Index < ThisPropertyReadAddresses.Num(); ++Index)
 		{
-			UObject* Object = ParentObjectNode->GetUObject(Index);
-			if (Object == nullptr) { return TOptional<bool>(); }
-			
-			UClass* ObjectClass = Object->GetClass();
-			if (ObjectClass == nullptr) { return TOptional<bool>(); }
-			
-			const void* SparseClassDataInstance = ObjectClass->GetSparseClassData(EGetSparseClassDataMethod::ReturnIfNull);
-			// SparseClassData instance hasn't been initialized for this object's class
-			if (SparseClassDataInstance == nullptr) { return TOptional<bool>(); }
-			
-			// Get the address of the value where the property is stored within the instance of the Sparse Class Data
-			const int32 ArrayIndex = 0; // EditConditions does not allow indexing into arrays
-			const uint8* ValuePtr = OperandProperty->ContainerPtrToValuePtr<uint8>(SparseClassDataInstance, ArrayIndex);
-			
-			bool Value = OperandProperty->GetPropertyValue(ValuePtr);
-			
+			uint8* EditConditionValuePtr = ComputeEditConditionValuePointer(*PinnedNode, ThisPropertyReadAddresses, *OperandProperty, Index);
+			bool Value = OperandProperty->GetPropertyValue(EditConditionValuePtr);
+
 			if (!Result.IsSet())
 			{
 				Result = Value;
@@ -386,27 +401,13 @@ TOptional<int64> FEditConditionContext::GetIntegerValue(const FString& PropertyN
 	}
 	else
 	{
-		// The ObjectPropertyNode contains the address of the sparse class data that is needed to resolve the value
-		FObjectPropertyNode* ParentObjectNode = ComplexParentNode->AsObjectNode();
-		if (ParentObjectNode == nullptr) { return TOptional<int64>(); }
+		FReadAddressList ThisPropertyReadAddresses;
+		PinnedNode->GetReadAddress(ThisPropertyReadAddresses);
 
-		for (int32 Index = 0; Index < ParentObjectNode->GetInstancesNum(); ++Index)
+		for (int32 Index = 0; Index < ThisPropertyReadAddresses.Num(); ++Index)
 		{
-			UObject* Object = ParentObjectNode->GetUObject(Index);
-			if (Object == nullptr) { return TOptional<int64>(); }
-			
-			UClass* ObjectClass = Object->GetClass();
-			if (ObjectClass == nullptr) { return TOptional<int64>(); }
-			
-			const void* SparseClassDataInstance = ObjectClass->GetSparseClassData(EGetSparseClassDataMethod::ReturnIfNull);
-			// SparseClassData instance hasn't been initialized for this object's class
-			if (SparseClassDataInstance == nullptr) { return TOptional<int64>(); }
-			
-			// Get the address of the value where the property is stored within the instance of the Sparse Class Data
-			const int32 ArrayIndex = 0; // EditConditions does not allow indexing into arrays
-			const uint8* ValuePtr = OperandProperty->ContainerPtrToValuePtr<uint8>(SparseClassDataInstance, ArrayIndex);
-			
-			int64 Value = OperandProperty->GetSignedIntPropertyValue(ValuePtr);
+			uint8* EditConditionValuePtr = ComputeEditConditionValuePointer(*PinnedNode, ThisPropertyReadAddresses, *OperandProperty, Index);
+			int64 Value = OperandProperty->GetSignedIntPropertyValue(EditConditionValuePtr);
 			
 			if (!Result.IsSet())
 			{
@@ -488,33 +489,21 @@ TOptional<double> FEditConditionContext::GetNumericValue(const FString& Property
 	}
 	else
 	{
-		// The ObjectPropertyNode contains the address of the sparse class data that is needed to resolve the value
-		FObjectPropertyNode* ParentObjectNode = ComplexParentNode->AsObjectNode();
-		if (ParentObjectNode == nullptr) { return TOptional<double>(); }
-		
-		for (int32 Index = 0; Index < ParentObjectNode->GetInstancesNum(); ++Index)
+		FReadAddressList ThisPropertyReadAddresses;
+		PinnedNode->GetReadAddress(ThisPropertyReadAddresses);
+
+		for (int32 Index = 0; Index < ThisPropertyReadAddresses.Num(); ++Index)
 		{
-			UObject* Object = ParentObjectNode->GetUObject(Index);
-			if (Object == nullptr) { return TOptional<double>(); }
-			
-			UClass* ObjectClass = Object->GetClass();
-			if (ObjectClass == nullptr) { return TOptional<double>(); }
+			uint8* EditConditionValuePtr = ComputeEditConditionValuePointer(*PinnedNode, ThisPropertyReadAddresses, *OperandProperty, Index);
 			
-			const void* SparseClassDataInstance = ObjectClass->GetSparseClassData(EGetSparseClassDataMethod::ReturnIfNull);
-			// SparseClassData instance hasn't been initialized for this object's class
-			if (SparseClassDataInstance == nullptr) { return TOptional<double>(); }
-			
-			// Get the address of the value where the property is stored within the instance of the Sparse Class Data
-			const int32 ArrayIndex = 0; // EditConditions does not allow indexing into arrays
-			const uint8* ValuePtr = OperandProperty->ContainerPtrToValuePtr<uint8>(SparseClassDataInstance, ArrayIndex);
 			double Value = 0;
 			if (OperandProperty->IsInteger())
 			{
-				Value = (double) OperandProperty->GetSignedIntPropertyValue(ValuePtr);
+				Value = (double) OperandProperty->GetSignedIntPropertyValue(EditConditionValuePtr);
 			}
 			else if (OperandProperty->IsFloatingPoint())
 			{
-				Value = OperandProperty->GetFloatingPointPropertyValue(ValuePtr);
+				Value = OperandProperty->GetFloatingPointPropertyValue(EditConditionValuePtr);
 			}
 			if (!Result.IsSet())
 			{
@@ -608,26 +597,13 @@ TOptional<FString> FEditConditionContext::GetEnumValue(const FString& PropertyNa
 	}
 	else
 	{
-		// The ObjectPropertyNode contains the address of the sparse class data that is needed to resolve the value
-		FObjectPropertyNode* ParentObjectNode = ComplexParentNode->AsObjectNode();
-		if (ParentObjectNode == nullptr) { return TOptional<FString>(); }
-		
-		for (int32 Index = 0; Index < ParentObjectNode->GetInstancesNum(); ++Index)
+		FReadAddressList ThisPropertyReadAddresses;
+		PinnedNode->GetReadAddress(ThisPropertyReadAddresses);
+
+		for (int32 Index = 0; Index < ThisPropertyReadAddresses.Num(); ++Index)
 		{
-			UObject* Object = ParentObjectNode->GetUObject(Index);
-			if (Object == nullptr) { return TOptional<FString>(); }
-			
-			UClass* ObjectClass = Object->GetClass();
-			if (ObjectClass == nullptr) { return TOptional<FString>(); }
-			
-			const void* SparseClassDataInstance = ObjectClass->GetSparseClassData(EGetSparseClassDataMethod::ReturnIfNull);
-			// SparseClassData instance hasn't been initialized for this object's class
-			if (SparseClassDataInstance == nullptr) { return TOptional<FString>(); }
-			
-			// Get the address of the value where the property is stored within the instance of the Sparse Class Data
-			const int32 ArrayIndex = 0; // EditConditions does not allow indexing into arrays of UObject*
-			const uint8* ValuePtr = OperandProperty->ContainerPtrToValuePtr<uint8>(SparseClassDataInstance, ArrayIndex);
-			const int64 Value = OperandProperty->GetSignedIntPropertyValue(ValuePtr);
+			uint8* EditConditionValuePtr = ComputeEditConditionValuePointer(*PinnedNode, ThisPropertyReadAddresses, *OperandProperty, Index);
+			const int64 Value = OperandProperty->GetSignedIntPropertyValue(EditConditionValuePtr);
 			if (!Result.IsSet())
 			{
 				Result = Value;
@@ -707,27 +683,14 @@ TOptional<UObject*> FEditConditionContext::GetPointerValue(const FString& Proper
 	}
 	else
 	{
-		// The ObjectPropertyNode contains the address of the sparse class data that is needed to resolve the value
-		FObjectPropertyNode* ParentObjectNode = ComplexParentNode->AsObjectNode();
-		if (ParentObjectNode == nullptr) { return TOptional<UObject*>(); }
-
 		TOptional<UObject*> Result;
-		for (int32 Index = 0; Index < ParentObjectNode->GetInstancesNum(); ++Index)
+		FReadAddressList ThisPropertyReadAddresses;
+		PinnedNode->GetReadAddress(ThisPropertyReadAddresses);
+
+		for (int32 Index = 0; Index < ThisPropertyReadAddresses.Num(); ++Index)
 		{
-			UObject* Object = ParentObjectNode->GetUObject(Index);
-			if (Object == nullptr) { return TOptional<UObject*>(); }
-			
-			UClass* ObjectClass = Object->GetClass();
-			if (ObjectClass == nullptr) { return TOptional<UObject*>(); }
-			
-			const void* SparseClassDataInstance = ObjectClass->GetSparseClassData(EGetSparseClassDataMethod::ReturnIfNull);
-			// SparseClassData instance hasn't been initialized for this object's class
-			if (SparseClassDataInstance == nullptr) { return TOptional<UObject*>(); }
-			
-			// Get the address of the value where the property is stored within the instance of the Sparse Class Data
-			const int32 ArrayIndex = 0; // EditConditions does not allow indexing into arrays of UObject*
-			const uint8* ValuePtr = OperandProperty->ContainerPtrToValuePtr<uint8>(SparseClassDataInstance, ArrayIndex);
-			UObject* Value = OperandProperty->GetObjectPropertyValue(ValuePtr);
+			uint8* EditConditionValuePtr = ComputeEditConditionValuePointer(*PinnedNode, ThisPropertyReadAddresses, *OperandProperty, Index);
+			UObject* Value = OperandProperty->GetObjectPropertyValue(EditConditionValuePtr);
 			if (!Result.IsSet())
 			{
 				Result = Value;
diff --git a/Engine/Source/Editor/PropertyEditor/Private/PropertyNode.cpp b/Engine/Source/Editor/PropertyEditor/Private/PropertyNode.cpp
--- a/Engine/Source/Editor/PropertyEditor/Private/PropertyNode.cpp
+++ b/Engine/Source/Editor/PropertyEditor/Private/PropertyNode.cpp
@@ -1569,18 +1569,31 @@ void FPropertyNode::ToggleEditConditionState()
 	}
 	else
 	{
+		// The strategy for getting the pointer to the vaue of the edit condition is to use the value pointer of the property
+		// that is being conditioned and then walk back that property's offset to find the owning struct's base address. 
+		// This owning struct's base address is then offset forward by the edit condition's property to find the edit condition's value pointer.
+		// 
+		// The assumption is that the edit condition property is on the same struct/class as the property that is being "conditioned".
+		// 
+		// Since the edit condition is inline, there will not be a PropertyNode available, therefore it is necessary to use
+		// the FProperty API to find the value pointer for the edit condition.
+		//
+		// It is also not possible to use the conditioned property's parent node as above since the parent may not point to the struct that either
+		// property is in.  In the case of property directly on the SparseClassData, the parent node is often a category or object node.
+		// 
+		// Care must also be taken for properties in sparse class data structs that are sub-structs containers within the sparse class data.
+		// Therefore, no assumptions can be made about the relationship between the sparse class data's pointer and the offset of the properties
+		// that are conditioned or used as edit conditions without walking up that relationship chain.
 		FObjectPropertyNode* ObjectPropertyNode = ComplexParentNode->AsObjectNode();
 		if (ObjectPropertyNode)
 		{
 			for (int32 Index = 0; Index < ObjectPropertyNode->GetInstancesNum(); ++Index)
 			{
 				UObject* Object = ObjectPropertyNode->GetUObject(Index);
-				uint8* SparseClassData = GetStartAddressFromObject(Object);
-
-				// Compute offset from the object containing the EditConditionProperty to the property itself
-				uint8* ValuePtr = EditConditionProperty->ContainerPtrToValuePtr<uint8>(SparseClassData);
-
-				// Toggle the value
+				uint8* ConditionedPropertyAddress = GetValueAddressFromObject(Object);
+				const int32 ConditionedPropertyOffset = GetProperty()->GetOffset_ForInternal();
+				uint8* OwningStructStartAddress = ConditionedPropertyAddress - ConditionedPropertyOffset;
+				uint8* ValuePtr = EditConditionProperty->ContainerPtrToValuePtr<uint8>(OwningStructStartAddress);
 				OldValue &= EditConditionProperty->GetPropertyValue(ValuePtr);
 				EditConditionProperty->SetPropertyValue(ValuePtr, !OldValue);
 			}
@@ -1600,7 +1613,6 @@ void FPropertyNode::ToggleEditConditionState()
 				TArray<UObject*> ArchetypeInstances;
 				Object->GetArchetypeInstances(ArchetypeInstances);
 				
-				
 				for (int32 InstanceIndex = 0; InstanceIndex < ArchetypeInstances.Num(); ++InstanceIndex)
 				{
 					uint8* ArchetypeBaseOffset;
@@ -1610,8 +1622,10 @@ void FPropertyNode::ToggleEditConditionState()
 					}
 					else
 					{
-						uint8* SparseClassData = GetStartAddressFromObject(ArchetypeInstances[InstanceIndex]);
-						ArchetypeBaseOffset = SparseClassData;
+						uint8* ConditionedPropertyAddress = GetValueAddressFromObject(ArchetypeInstances[InstanceIndex]);
+						const int32 ConditionedPropertyOffset = GetProperty()->GetOffset_ForInternal();
+						uint8* OwningStructStartAddress = ConditionedPropertyAddress - ConditionedPropertyOffset;
+						ArchetypeBaseOffset = EditConditionProperty->ContainerPtrToValuePtr<uint8>(OwningStructStartAddress);
 					}
 					
 					uint8* ArchetypeValueAddr = EditConditionProperty->ContainerPtrToValuePtr<uint8>(ArchetypeBaseOffset);
