PR #11022: Add support to open assets directly instead of browsing.

diff --git a/Engine/Source/Editor/PropertyEditor/Private/Presentation/PropertyEditor/PropertyEditor.cpp b/Engine/Source/Editor/PropertyEditor/Private/Presentation/PropertyEditor/PropertyEditor.cpp
--- a/Engine/Source/Editor/PropertyEditor/Private/Presentation/PropertyEditor/PropertyEditor.cpp
+++ b/Engine/Source/Editor/PropertyEditor/Private/Presentation/PropertyEditor/PropertyEditor.cpp
@@ -832,7 +832,18 @@ void FPropertyEditor::SyncToObjectsInNode( const TWeakPtr< FPropertyNode >& Weak
 		}
 		else if ( Objects.Num() > 0 )
 		{
-			GEditor->SyncBrowserToObjects(Objects);
+			const FModifierKeysState KeyState = FSlateApplication::Get().GetModifierKeys();
+			if (KeyState.IsAltDown( ))
+			{
+				for (UObject *Obj : Objects)
+				{
+					GEditor->EditObject(Obj);
+				}
+			}
+			else
+			{
+				GEditor->SyncBrowserToObjects(Objects);
+			}
 		}
 	}
 
diff --git a/Engine/Source/Editor/PropertyEditor/Private/PropertyCustomizationHelpers.cpp b/Engine/Source/Editor/PropertyEditor/Private/PropertyCustomizationHelpers.cpp
--- a/Engine/Source/Editor/PropertyEditor/Private/PropertyCustomizationHelpers.cpp
+++ b/Engine/Source/Editor/PropertyEditor/Private/PropertyCustomizationHelpers.cpp
@@ -3,6 +3,7 @@
 #include "PropertyCustomizationHelpers.h"
 #include "IDetailChildrenBuilder.h"
 #include "Widgets/Text/STextBlock.h"
+#include "Framework/Application/SlateApplication.h"
 #include "Framework/MultiBox/MultiBoxBuilder.h"
 #include "Engine/Texture.h"
 #include "Factories/Factory.h"
@@ -59,7 +60,7 @@ namespace PropertyCustomizationHelpers
 			, _IsFocusable( true )
 		{}
 			SLATE_ATTRIBUTE( FText, Text )
-			SLATE_ARGUMENT( const FSlateBrush*, Image )
+			SLATE_ATTRIBUTE( const FSlateBrush*, Image )
 			SLATE_EVENT( FSimpleDelegate, OnClickAction )
 
 			/** Sometimes a button should only be mouse-clickable and never keyboard focusable. */
@@ -603,8 +604,8 @@ namespace PropertyCustomizationHelpers
 
 		return
 			SNew( SPropertyEditorButton )
-			.Text( OptionalToolTipText.Get().IsEmpty() ? LOCTEXT( "BrowseButtonToolTipText", "Browse to Asset in Content Browser") : OptionalToolTipText )
-			.Image(BrowseIcon)
+			.Text( OptionalToolTipText.Get().IsEmpty() ? LOCTEXT( "BrowseButtonToolTipText", "Browse to Asset in Content Browser (hold Alt to Open instead)") : OptionalToolTipText )
+			.Image_Lambda( [BrowseIcon]() { return FSlateApplication::Get().GetModifierKeys().IsAltDown() ? FAppStyle::Get().GetBrush("Icons.Edit") : BrowseIcon; } )
 			.OnClickAction( OnFindClicked )
 			.IsEnabled(IsEnabled)
 			.IsFocusable( false );
diff --git a/Engine/Source/Editor/PropertyEditor/Private/UserInterface/PropertyEditor/SPropertyEditorAsset.cpp b/Engine/Source/Editor/PropertyEditor/Private/UserInterface/PropertyEditor/SPropertyEditorAsset.cpp
--- a/Engine/Source/Editor/PropertyEditor/Private/UserInterface/PropertyEditor/SPropertyEditorAsset.cpp
+++ b/Engine/Source/Editor/PropertyEditor/Private/UserInterface/PropertyEditor/SPropertyEditorAsset.cpp
@@ -1514,9 +1514,17 @@ void SPropertyEditorAsset::OnBrowse()
 	}
 	else
 	{
-		TArray<FAssetData> AssetDataList;
-		AssetDataList.Add( Value.AssetData );
-		GEditor->SyncBrowserToObjects( AssetDataList );
+		const FModifierKeysState KeyState = FSlateApplication::Get().GetModifierKeys();
+		if (KeyState.IsAltDown( ))
+		{
+			GEditor->EditObject( Value.AssetData.GetAsset(  ) );
+		}
+		else
+		{
+			TArray<FAssetData> AssetDataList;
+			AssetDataList.Add( Value.AssetData );
+			GEditor->SyncBrowserToObjects( AssetDataList );
+		}
 	}
 }
 
