[Editor] Fixes issue where FEditConditionContext::GetIntegerValue could fails for enums

#rnx

diff --git a/Engine/Source/Editor/PropertyEditor/Private/EditConditionContext.cpp b/Engine/Source/Editor/PropertyEditor/Private/EditConditionContext.cpp
--- a/Engine/Source/Editor/PropertyEditor/Private/EditConditionContext.cpp
+++ b/Engine/Source/Editor/PropertyEditor/Private/EditConditionContext.cpp
@@ -347,7 +347,7 @@ TOptional<int64> FEditConditionContext::GetIntegerValue(const FString& PropertyN
 		return TOptional<int64>();
 	}
 
-	const FProperty* Property = FindTypedField<FNumericProperty>(PinnedNode, PropertyName);
+	const FProperty* Property = FindTypedField<FProperty>(PinnedNode, PropertyName);
 	const FNumericProperty* OperandProperty = CastField<FNumericProperty>(Property);
 
 	if (OperandProperty == nullptr)
