Fix double property offset applied to inline Edit Conditions on Sparse Class Data

diff --git a/Engine/Source/Editor/PropertyEditor/Private/PropertyNode.cpp b/Engine/Source/Editor/PropertyEditor/Private/PropertyNode.cpp
--- a/Engine/Source/Editor/PropertyEditor/Private/PropertyNode.cpp
+++ b/Engine/Source/Editor/PropertyEditor/Private/PropertyNode.cpp
@@ -1625,7 +1625,7 @@ void FPropertyNode::ToggleEditConditionState()
 						uint8* ConditionedPropertyAddress = GetValueAddressFromObject(ArchetypeInstances[InstanceIndex]);
 						const int32 ConditionedPropertyOffset = GetProperty()->GetOffset_ForInternal();
 						uint8* OwningStructStartAddress = ConditionedPropertyAddress - ConditionedPropertyOffset;
-						ArchetypeBaseOffset = EditConditionProperty->ContainerPtrToValuePtr<uint8>(OwningStructStartAddress);
+						ArchetypeBaseOffset = OwningStructStartAddress;
 					}
 					
 					uint8* ArchetypeValueAddr = EditConditionProperty->ContainerPtrToValuePtr<uint8>(ArchetypeBaseOffset);
