Export FDetailsDisplayManager::OverrideCreateCategoryNodes

diff --git a/Engine/Source/Editor/PropertyEditor/Public/DetailsDisplayManager.h b/Engine/Source/Editor/PropertyEditor/Public/DetailsDisplayManager.h
--- a/Engine/Source/Editor/PropertyEditor/Public/DetailsDisplayManager.h
+++ b/Engine/Source/Editor/PropertyEditor/Public/DetailsDisplayManager.h
@@ -29,7 +29,7 @@ public:
 
 	PROPERTYEDITOR_API virtual ~FDetailsDisplayManager();
 
-	virtual TOptional<bool> OverrideCreateCategoryNodes() const;
+	PROPERTYEDITOR_API virtual TOptional<bool> OverrideCreateCategoryNodes() const;
 
 	/**
 	 * Returns a boolean indicating if the Component Editor should be hidden 
