Fix issue where metadata provided on array properties was not being considered in PropertyEditorHelpers::GetRequiredPropertyButtons()

diff --git a/Engine/Source/Editor/PropertyEditor/Private/PropertyEditorHelpers.cpp b/Engine/Source/Editor/PropertyEditor/Private/PropertyEditorHelpers.cpp
--- a/Engine/Source/Editor/PropertyEditor/Private/PropertyEditorHelpers.cpp
+++ b/Engine/Source/Editor/PropertyEditor/Private/PropertyEditorHelpers.cpp
@@ -715,6 +715,9 @@ namespace PropertyEditorHelpers
 			return;
 		}
 
+		const FProperty* MetadataProperty = PropertyNode->GetMetaDataProperty();
+		check(MetadataProperty); // As NodeProperty is non-null, this should always be valid at this point.
+
 		// If the property is an item of a const container, don't create any buttons.
 		const FArrayProperty* OuterArrayProp = NodeProperty->GetOwner<FArrayProperty>();
 		const FSetProperty* OuterSetProp = NodeProperty->GetOwner<FSetProperty>();
@@ -759,7 +762,7 @@ namespace PropertyEditorHelpers
 					Class = NodeProperty->GetOwnerProperty()->GetClassMetaData(TEXT("MetaClass"));
 				}
 
-				if (Class && FKismetEditorUtilities::CanCreateBlueprintOfClass(Class) && !NodeProperty->HasMetaData("DisallowCreateNew"))
+				if (Class && FKismetEditorUtilities::CanCreateBlueprintOfClass(Class) && !MetadataProperty->HasMetaData("DisallowCreateNew"))
 				{
 					OutRequiredButtons.Add(EPropertyButton::NewBlueprint);
 				}
@@ -848,9 +851,9 @@ namespace PropertyEditorHelpers
 
 		if( OuterArrayProp )
 		{
-			if( PropertyNode->HasNodeFlags(EPropertyNodeFlags::SingleSelectOnly) && !(OuterArrayProp->PropertyFlags & CPF_EditFixedSize) )
+			if( PropertyNode->HasNodeFlags(EPropertyNodeFlags::SingleSelectOnly) && !(MetadataProperty->PropertyFlags & CPF_EditFixedSize) )
 			{
-				if (OuterArrayProp->HasMetaData(TEXT("NoElementDuplicate")))
+				if (MetadataProperty->HasMetaData(TEXT("NoElementDuplicate")))
 				{
 					OutRequiredButtons.Add( EPropertyButton::Insert_Delete );
 				}
diff --git a/Engine/Source/Editor/PropertyEditor/Private/PropertyHandleImpl.cpp b/Engine/Source/Editor/PropertyEditor/Private/PropertyHandleImpl.cpp
--- a/Engine/Source/Editor/PropertyEditor/Private/PropertyHandleImpl.cpp
+++ b/Engine/Source/Editor/PropertyEditor/Private/PropertyHandleImpl.cpp
@@ -3019,25 +3019,13 @@ FProperty* FPropertyHandleBase::GetProperty() const
 
 FProperty* FPropertyHandleBase::GetMetaDataProperty() const
 {
-	FProperty* MetaDataProperty = nullptr;
-
 	TSharedPtr<FPropertyNode> PropertyNode = Implementation->GetPropertyNode();
 	if( PropertyNode.IsValid() )
 	{
-		MetaDataProperty = PropertyNode->GetProperty();
-		
-		// If we are part of an array, we need to take our meta-data from the array property
-		if( PropertyNode->GetArrayIndex() != INDEX_NONE )
-		{
-			TSharedPtr<FPropertyNode> ParentNode = PropertyNode->GetParentNodeSharedPtr();
-			if (ParentNode)
-			{
-				MetaDataProperty = ParentNode->GetProperty();
-			}
-		}
+		return PropertyNode->GetMetaDataProperty();
 	}
 
-	return MetaDataProperty;
+	return nullptr;
 }
 
 bool FPropertyHandleBase::HasMetaData(const FName& Key) const
diff --git a/Engine/Source/Editor/PropertyEditor/Private/PropertyNode.cpp b/Engine/Source/Editor/PropertyEditor/Private/PropertyNode.cpp
--- a/Engine/Source/Editor/PropertyEditor/Private/PropertyNode.cpp
+++ b/Engine/Source/Editor/PropertyEditor/Private/PropertyNode.cpp
@@ -3698,6 +3698,23 @@ bool FPropertyNode::ParentOrSelfHasMetaData(const FName& MetaDataKey) const
 	return false;
 }
 
+FProperty* FPropertyNode::GetMetaDataProperty()
+{
+	FProperty* MetaDataProperty = GetProperty();
+
+	// If we are part of an array, we need to take our meta-data from the array property
+	if (GetArrayIndex() != INDEX_NONE)
+	{
+		FPropertyNode* ParentNode = GetParentNode();
+		if (ParentNode)
+		{
+			MetaDataProperty = ParentNode->GetProperty();
+		}
+	}
+	
+	return MetaDataProperty;
+}
+
 void FPropertyNode::InvalidateCachedState()
 {
 	bUpdateDiffersFromDefault = true;
diff --git a/Engine/Source/Editor/PropertyEditor/Private/PropertyNode.h b/Engine/Source/Editor/PropertyEditor/Private/PropertyNode.h
--- a/Engine/Source/Editor/PropertyEditor/Private/PropertyNode.h
+++ b/Engine/Source/Editor/PropertyEditor/Private/PropertyNode.h
@@ -1045,6 +1045,11 @@ public:
 
 	bool ParentOrSelfHasMetaData(const FName& MetaDataKey) const;
 
+	/**
+	 * Gets the property we should use to read meta-data
+	 */
+	FProperty* GetMetaDataProperty();
+
 	/**
 	 * Invalidates the cached state of this node in all children;
 	 */
