Details Panel: Modify the "Browse to Asset/Actor" button to only change to the edit icon when alt is held down and the actual property widget is focused

diff --git a/Engine/Source/Editor/PropertyEditor/Private/PropertyCustomizationHelpers.cpp b/Engine/Source/Editor/PropertyEditor/Private/PropertyCustomizationHelpers.cpp
--- a/Engine/Source/Editor/PropertyEditor/Private/PropertyCustomizationHelpers.cpp
+++ b/Engine/Source/Editor/PropertyEditor/Private/PropertyCustomizationHelpers.cpp
@@ -598,14 +598,19 @@ namespace PropertyCustomizationHelpers
 			];
 	}
 
-	TSharedRef<SWidget> MakeBrowseButton( FSimpleDelegate OnFindClicked, TAttribute<FText> OptionalToolTipText, TAttribute<bool> IsEnabled, const bool IsActor)
+	TSharedRef<SWidget> MakeBrowseButton( FSimpleDelegate OnFindClicked, TAttribute<FText> OptionalToolTipText, TAttribute<bool> IsEnabled, const bool IsActor, const TAttribute<const FSlateBrush*>& InIconOverride)
 	{
-		const FSlateBrush* BrowseIcon = IsActor ? FAppStyle::Get().GetBrush("Icons.SelectInViewport") : FAppStyle::Get().GetBrush("Icons.BrowseContent");
+		TAttribute<const FSlateBrush*> IconAttribute = InIconOverride;
+
+		if (!IconAttribute.IsSet())
+		{
+			IconAttribute = IsActor ? FAppStyle::Get().GetBrush("Icons.SelectInViewport") : FAppStyle::Get().GetBrush("Icons.BrowseContent");
+		}
 
 		return
 			SNew( SPropertyEditorButton )
-			.Text( OptionalToolTipText.Get().IsEmpty() ? LOCTEXT( "BrowseButtonToolTipText", "Browse to Asset in Content Browser (hold Alt to Open instead)") : OptionalToolTipText )
-			.Image_Lambda( [BrowseIcon]() { return FSlateApplication::Get().GetModifierKeys().IsAltDown() ? FAppStyle::Get().GetBrush("Icons.Edit") : BrowseIcon; } )
+			.Text( OptionalToolTipText.Get().IsEmpty() ? LOCTEXT( "BrowseButtonToolTipText", "Browse to Asset in Content Browser") : OptionalToolTipText )
+			.Image( IconAttribute )
 			.OnClickAction( OnFindClicked )
 			.IsEnabled(IsEnabled)
 			.IsFocusable( false );
diff --git a/Engine/Source/Editor/PropertyEditor/Private/UserInterface/PropertyEditor/SPropertyEditorAsset.cpp b/Engine/Source/Editor/PropertyEditor/Private/UserInterface/PropertyEditor/SPropertyEditorAsset.cpp
--- a/Engine/Source/Editor/PropertyEditor/Private/UserInterface/PropertyEditor/SPropertyEditorAsset.cpp
+++ b/Engine/Source/Editor/PropertyEditor/Private/UserInterface/PropertyEditor/SPropertyEditorAsset.cpp
@@ -766,6 +766,10 @@ void SPropertyEditorAsset::Construct(const FArguments& InArgs, const TSharedPtr<
 	{
 		FSimpleDelegate OnBrowseDelegate = InArgs._OnBrowseOverride.IsBound() ? InArgs._OnBrowseOverride : FSimpleDelegate::CreateSP(this, &SPropertyEditorAsset::OnBrowse);
 
+		// Only the default SPropertyEditorAsset::OnBrowse delegate supports alt+click to directly open the asset, so we only modify the icon if there is no browse override
+		TAttribute<const FSlateBrush*> IconDelegate = !InArgs._OnBrowseOverride.IsBound() ?
+			TAttribute<const FSlateBrush*>::CreateSP(this, &SPropertyEditorAsset::GetOnBrowseIcon) : TAttribute<const FSlateBrush*>();
+
 		ButtonBox->AddSlot()
 		.Padding( 2.0f, 0.0f )
 		.AutoWidth()
@@ -775,7 +779,8 @@ void SPropertyEditorAsset::Construct(const FArguments& InArgs, const TSharedPtr<
 				OnBrowseDelegate,
 				TAttribute<FText>( this, &SPropertyEditorAsset::GetOnBrowseToolTip ),
 				true,
-				bIsActor
+				bIsActor,
+				IconDelegate
 				)
 		];
 	}
@@ -1556,29 +1561,42 @@ FText SPropertyEditorAsset::GetOnBrowseToolTip() const
 {
 	FObjectOrAssetData Value;
 	GetValue( Value, FObjectOrAssetData::EAssetDataOptions::SkipAssetRegistryTagsGathering );
+	
+	FFormatNamedArguments Args;
+	Args.Add(TEXT("AltOpen"), LOCTEXT( "HoldAltToOpenText", "(hold Alt to Open instead)"));
 
 	if (Value.Object)
 	{
-		FFormatNamedArguments Args;
 		Args.Add(TEXT("Asset"), FText::AsCultureInvariant(Value.Object->GetName()));
 		if (bIsActor)
 		{
-			return FText::Format(LOCTEXT( "SelectSpecificActorInViewport", "Select '{Asset}' in the viewport"), Args);
+			return FText::Format(LOCTEXT( "SelectSpecificActorInViewport", "Select '{Asset}' in the viewport {AltOpen}"), Args);
 		}
 		else
 		{
-			return FText::Format(LOCTEXT( "BrowseToSpecificAssetInContentBrowser", "Browse to '{Asset}' in Content Browser"), Args);
+			return FText::Format(LOCTEXT( "BrowseToSpecificAssetInContentBrowser", "Browse to '{Asset}' in Content Browser {AltOpen}"), Args);
 		}
 	}
 	
 	if (bIsActor)
 	{
-		return LOCTEXT("SelectActorInViewport", "Select Actor in the viewport");
+		return FText::Format(LOCTEXT("SelectActorInViewport", "Select Actor in the viewport {AltOpen}"), Args);
 	}
 	else
 	{
-		return LOCTEXT("BrowseToAssetInContentBrowser", "Browse to Asset in Content Browser");
+		return FText::Format(LOCTEXT("BrowseToAssetInContentBrowser", "Browse to Asset in Content Browser {AltOpen}"), Args);
+	}
+}
+
+const FSlateBrush* SPropertyEditorAsset::GetOnBrowseIcon() const
+{
+	// If the widget is hovered and alt is held down, show the edit icon
+	if (IsHovered() && FSlateApplication::Get().GetModifierKeys().IsAltDown())
+	{
+		return FAppStyle::Get().GetBrush("Icons.Edit");
 	}
+	
+	return bIsActor ? FAppStyle::Get().GetBrush("Icons.SelectInViewport") : FAppStyle::Get().GetBrush("Icons.BrowseContent");
 }
 
 void SPropertyEditorAsset::OnUse()
diff --git a/Engine/Source/Editor/PropertyEditor/Private/UserInterface/PropertyEditor/SPropertyEditorAsset.h b/Engine/Source/Editor/PropertyEditor/Private/UserInterface/PropertyEditor/SPropertyEditorAsset.h
--- a/Engine/Source/Editor/PropertyEditor/Private/UserInterface/PropertyEditor/SPropertyEditorAsset.h
+++ b/Engine/Source/Editor/PropertyEditor/Private/UserInterface/PropertyEditor/SPropertyEditorAsset.h
@@ -273,6 +273,11 @@ private:
 	 * Get the tooltip text for the Browse button
 	 */
 	FText GetOnBrowseToolTip() const;
+	
+	/** 
+	 * Get the icon for the Browse button
+	 */
+	const FSlateBrush* GetOnBrowseIcon() const;
 
 	/** 
 	 * Use the selected object (replaces the referenced object if valid)
diff --git a/Engine/Source/Editor/PropertyEditor/Public/PropertyCustomizationHelpers.h b/Engine/Source/Editor/PropertyEditor/Public/PropertyCustomizationHelpers.h
--- a/Engine/Source/Editor/PropertyEditor/Public/PropertyCustomizationHelpers.h
+++ b/Engine/Source/Editor/PropertyEditor/Public/PropertyCustomizationHelpers.h
@@ -165,7 +165,7 @@ namespace PropertyCustomizationHelpers
 	PROPERTYEDITOR_API TSharedRef<SWidget> MakeVisibilityButton(FOnClicked OnVisibilityClicked, TAttribute<FText> OptionalToolTipText = FText(), TAttribute<bool> VisibilityDelegate = true);
 	PROPERTYEDITOR_API TSharedRef<SWidget> MakeNewBlueprintButton( FSimpleDelegate OnNewBlueprintClicked, TAttribute<FText> OptionalToolTipText = FText(), TAttribute<bool> IsEnabled = true );
 	PROPERTYEDITOR_API TSharedRef<SWidget> MakeUseSelectedButton( FSimpleDelegate OnUseSelectedClicked, TAttribute<FText> OptionalToolTipText = FText(), TAttribute<bool> IsEnabled = true, const bool IsActor = false );
-	PROPERTYEDITOR_API TSharedRef<SWidget> MakeBrowseButton( FSimpleDelegate OnFindClicked, TAttribute<FText> OptionalToolTipText = FText(), TAttribute<bool> IsEnabled = true, const bool IsActor = false);
+	PROPERTYEDITOR_API TSharedRef<SWidget> MakeBrowseButton( FSimpleDelegate OnFindClicked, TAttribute<FText> OptionalToolTipText = FText(), TAttribute<bool> IsEnabled = true, const bool IsActor = false, const TAttribute<const FSlateBrush*>& InIconOverride = TAttribute<const FSlateBrush*>());
 	PROPERTYEDITOR_API TSharedRef<SWidget> MakeAssetPickerAnchorButton( FOnGetAllowedClasses OnGetAllowedClasses, FOnAssetSelected OnAssetSelectedFromPicker, const TSharedPtr<IPropertyHandle>& PropertyHandle = TSharedPtr<IPropertyHandle>());
 	PROPERTYEDITOR_API TSharedRef<SWidget> MakeAssetPickerWithMenu( const FAssetData& InitialObject, const bool AllowClear, const TArray<const UClass*>& AllowedClasses, const TArray<UFactory*>& NewAssetFactories, FOnShouldFilterAsset OnShouldFilterAsset, FOnAssetSelected OnSet, FSimpleDelegate OnClose, const TSharedPtr<IPropertyHandle>& PropertyHandle = TSharedPtr<IPropertyHandle>(), const TArray<FAssetData>& OwnerAssetArray = TArray<FAssetData>());
 	PROPERTYEDITOR_API TSharedRef<SWidget> MakeAssetPickerWithMenu( const FAssetData& InitialObject, const bool AllowClear, const TArray<const UClass*>& AllowedClasses, const TArray<const UClass*>& DisallowedClasses, const TArray<UFactory*>& NewAssetFactories, FOnShouldFilterAsset OnShouldFilterAsset, FOnAssetSelected OnSet, FSimpleDelegate OnClose, const TSharedPtr<IPropertyHandle>& PropertyHandle = TSharedPtr<IPropertyHandle>(), const TArray<FAssetData>& OwnerAssetArray = TArray<FAssetData>());
