Property search in the details panel now will search the "Keywords" meta specifier.

- Previously only blueprint functions and variables would take into account the Keyword specifier
- This brings feature parity between the search functionality in BP node palette and the search in the details panel.
- Adds the "Displacement" keyword specificer to `bEnableTessellation` on materials.
- This is to avoid user confusion and to align the names across the material editor.
- The output pin is called "Displacement", the other properties (Magnitude and Scale) are called "Displacement".
- Now it follows that the checkbox to enable those features is called "Enable Displacement".

[RN] major | Rendering - Materials and Shaders

diff --git a/Engine/Source/Editor/PropertyEditor/Private/PropertyNode.cpp b/Engine/Source/Editor/PropertyEditor/Private/PropertyNode.cpp
--- a/Engine/Source/Editor/PropertyEditor/Private/PropertyNode.cpp
+++ b/Engine/Source/Editor/PropertyEditor/Private/PropertyNode.cpp
@@ -3015,6 +3015,17 @@ void FPropertyNode::FilterNodes( const TArray<FString>& InFilterStrings, const b
 			}
 		}
 
+		// Check keywords on the property.
+		static const FName KeywordPropertyName = FName(TEXT("Keywords"));
+		if (GetProperty())
+		{
+			const FString& KeywordProperty = GetProperty()->GetMetaData(KeywordPropertyName);
+			if (!KeywordProperty.IsEmpty())
+			{
+				AcceptableNames.Add(KeywordProperty);
+			}
+		}
+
 		// Get the basic name as well of the property
 		FProperty* TheProperty = GetProperty();
 		if (TheProperty && (TheProperty->GetName() != DisplayNameStr))
diff --git a/Engine/Source/Runtime/Engine/Public/Materials/Material.h b/Engine/Source/Runtime/Engine/Public/Materials/Material.h
--- a/Engine/Source/Runtime/Engine/Public/Materials/Material.h
+++ b/Engine/Source/Runtime/Engine/Public/Materials/Material.h
@@ -566,7 +566,7 @@ public:
 	uint8 bHasPixelAnimation : 1;
 	
 	/** Whether tessellation is enabled on the material. NOTE: Required for displacement to work. */
-	UPROPERTY(EditAnywhere, Category=Nanite)
+	UPROPERTY(EditAnywhere, Category=Nanite, meta = (Keywords = "Displacement"))
 	uint8 bEnableTessellation : 1;
 
 	/** Enables fading out and disabling of dynamic displacement in the distance, as displacement becomes unnoticeable */
