Fixed UAF Trait categories not getting generated in the details panel customization, after CL 33110958

diff --git a/Engine/Plugins/Experimental/AnimNext/Source/AnimNextEditor/Private/Common/AnimNextEdGraphNodeCustomization.cpp b/Engine/Plugins/Experimental/AnimNext/Source/AnimNextEditor/Private/Common/AnimNextEdGraphNodeCustomization.cpp
--- a/Engine/Plugins/Experimental/AnimNext/Source/AnimNextEditor/Private/Common/AnimNextEdGraphNodeCustomization.cpp
+++ b/Engine/Plugins/Experimental/AnimNext/Source/AnimNextEditor/Private/Common/AnimNextEdGraphNodeCustomization.cpp
@@ -76,7 +76,6 @@ void FAnimNextEdGraphNodeCustomization::CustomizeObjects(IDetailLayoutBuilder& D
 		if (ModelNode->GetTraitNames().Num() > 0)
 		{
 			RigVMGraphDetailCustomizationImpl.Reset();
-			GenerateRigVMData(EdGraphNode, CategoryDetailsData);
 			GenerateTraitData(EdGraphNode, CategoryDetailsData);
 		}
 		else 
@@ -393,7 +392,8 @@ void FAnimNextEdGraphNodeCustomization::PopulateCategory(IDetailLayoutBuilder& D
 				PinPath = PinPath.Replace(*RootPropertyPath, *PinnedTraitData->Name.ToString());
 				// Prefix node path
 				const FString NodePath = EdGraphNode->GetModelNodePath();
-				PinPath = NodePath + TEXT(".") + PinPath;
+				const bool bIsTrait = PinnedTraitData->Type == FCategoryDetailsData::EType::TraitStack;
+				PinPath = bIsTrait ? NodePath + TEXT(".") + PinnedTraitData->Name.ToString() + TEXT(".") + PinPath : NodePath + TEXT(".") + PinPath;
 				if(!Controller->SetPinDefaultValue(PinPath, ValueStr))
 				{
 					UE_LOG(LogAnimation, Warning, TEXT("Could not set default value for pin '%s'"), *PinPath);
diff --git a/Engine/Source/Editor/PropertyEditor/Private/DetailPropertyRow.cpp b/Engine/Source/Editor/PropertyEditor/Private/DetailPropertyRow.cpp
--- a/Engine/Source/Editor/PropertyEditor/Private/DetailPropertyRow.cpp
+++ b/Engine/Source/Editor/PropertyEditor/Private/DetailPropertyRow.cpp
@@ -688,30 +688,39 @@ void MakeExternalStructPropertyRowCustomization(const T& Struct, const UStruct*
 	}
 	else
 	{
-		static const FName PropertyEditorModuleName("PropertyEditor");
-		FPropertyEditorModule& PropertyEditorModule = FModuleManager::GetModuleChecked<FPropertyEditorModule>(PropertyEditorModuleName);
+		if (bAllowChildren && InitParams.bCreateCategoryNodes)
+		{
+			// Using the RootPropertyNode as the property node enables sub categories
+			OutCustomization.PropertyRow = MakeShared<FDetailPropertyRow>(RootPropertyNode, ParentCategory, RootPropertyNode);
+			OutCustomization.PropertyRow->SetCustomExpansionId(Parameters.GetUniqueId());
+		}
+		else
+		{
+			static const FName PropertyEditorModuleName("PropertyEditor");
+			FPropertyEditorModule& PropertyEditorModule = FModuleManager::GetModuleChecked<FPropertyEditorModule>(PropertyEditorModuleName);
 
-		// Make a "fake" struct property to represent the entire struct
-		FStructProperty* StructProperty = PropertyEditorModule.RegisterStructProperty(StructClass);
+			// Make a "fake" struct property to represent the entire struct
+			FStructProperty* StructProperty = PropertyEditorModule.RegisterStructProperty(StructClass);
 
-		// Generate a node for the struct
-		TSharedPtr<FItemPropertyNode> ItemNode = MakeShared<FItemPropertyNode>();
+			// Generate a node for the struct
+			TSharedPtr<FItemPropertyNode> ItemNode = MakeShared<FItemPropertyNode>();
 
-		FPropertyNodeInitParams ItemNodeInitParams;
-		ItemNodeInitParams.ParentNode = RootPropertyNode;
-		ItemNodeInitParams.Property = StructProperty;
-		ItemNodeInitParams.ArrayOffset = 0;
-		ItemNodeInitParams.ArrayIndex = INDEX_NONE;
-		ItemNodeInitParams.bAllowChildren = true;
-		ItemNodeInitParams.bForceHiddenPropertyVisibility = Parameters.ShouldForcePropertyVisible() || FPropertySettings::Get().ShowHiddenProperties();
-		ItemNodeInitParams.bCreateCategoryNodes = false;
+			FPropertyNodeInitParams ItemNodeInitParams;
+			ItemNodeInitParams.ParentNode = RootPropertyNode;
+			ItemNodeInitParams.Property = StructProperty;
+			ItemNodeInitParams.ArrayOffset = 0;
+			ItemNodeInitParams.ArrayIndex = INDEX_NONE;
+			ItemNodeInitParams.bAllowChildren = true;
+			ItemNodeInitParams.bForceHiddenPropertyVisibility = Parameters.ShouldForcePropertyVisible() || FPropertySettings::Get().ShowHiddenProperties();
+			ItemNodeInitParams.bCreateCategoryNodes = false;
 
-		ItemNode->InitNode(ItemNodeInitParams);
+			ItemNode->InitNode(ItemNodeInitParams);
 
-		RootPropertyNode->AddChildNode(ItemNode);
+			RootPropertyNode->AddChildNode(ItemNode);
 
-		OutCustomization.PropertyRow = MakeShared<FDetailPropertyRow>(ItemNode, ParentCategory, RootPropertyNode);
-		OutCustomization.PropertyRow->SetCustomExpansionId(Parameters.GetUniqueId());
+			OutCustomization.PropertyRow = MakeShared<FDetailPropertyRow>(ItemNode, ParentCategory, RootPropertyNode);
+			OutCustomization.PropertyRow->SetCustomExpansionId(Parameters.GetUniqueId());
+		}
 	}
 }
 
