Reparent properties with the None category to the category above it

None is a special category internally denoting no category.  C++ generally will not get into this situation since UHT requires categories specified in UPROPERTY.  However in verse, categories are an optional attribute so properties are placed in this None category.  It is undesireable to explicitly see this internal category so this change removes the display of it.  Note that the category property node will still exist.

diff --git a/Engine/Source/Editor/PropertyEditor/Private/DetailPropertyRow.cpp b/Engine/Source/Editor/PropertyEditor/Private/DetailPropertyRow.cpp
--- a/Engine/Source/Editor/PropertyEditor/Private/DetailPropertyRow.cpp
+++ b/Engine/Source/Editor/PropertyEditor/Private/DetailPropertyRow.cpp
@@ -500,59 +500,68 @@ void FDetailPropertyRow::GenerateChildrenForPropertyNode( TSharedPtr<FPropertyNo
 					TArray<TSharedRef<FDetailTreeNode>> PropNodes;
 					bool bHasKeyNode = false;
 
-					// Create and initialize the child first
-					FDetailLayoutCustomization Customization;
-					Customization.PropertyRow = MakeShared<FDetailPropertyRow>(ChildNode, ParentCategoryRef);
-
-					if (CustomResetToDefault.IsSet() && CustomResetToDefault->PropagatesToChildren())
+					FCategoryPropertyNode* CategoryNode = ChildNode->AsCategoryNode();
+					if (CategoryNode && CategoryNode->GetCategoryName() == NAME_None)
 					{
-						Customization.PropertyRow->OverrideResetToDefault(CustomResetToDefault.GetValue());
+						// Skip the category node and repoarent everything to the ParentCategory
+						GenerateChildrenForPropertyNode(ChildNode, OutChildren);
 					}
+					else
+					{
+						// Create and initialize the child first
+						FDetailLayoutCustomization Customization;
+						Customization.PropertyRow = MakeShared<FDetailPropertyRow>(ChildNode, ParentCategoryRef);
+
+						if (CustomResetToDefault.IsSet() && CustomResetToDefault->PropagatesToChildren())
+						{
+							Customization.PropertyRow->OverrideResetToDefault(CustomResetToDefault.GetValue());
+						}
 
-					TSharedRef<FDetailItemNode> ChildNodeItem = MakeShared<FDetailItemNode>(Customization, ParentCategoryRef, ParentEnabledState);
-					ChildNodeItem->Initialize();
+						TSharedRef<FDetailItemNode> ChildNodeItem = MakeShared<FDetailItemNode>(Customization, ParentCategoryRef, ParentEnabledState);
+						ChildNodeItem->Initialize();
 
-					if ( ChildNode->GetPropertyKeyNode().IsValid() )
-					{
-						// If the child has a key property, only create a second node for the key if the child did not already create a property
-						// editor for it
-						if ( !Customization.PropertyRow->PropertyKeyEditor.IsValid() )
+						if ( ChildNode->GetPropertyKeyNode().IsValid() )
 						{
-							FDetailLayoutCustomization KeyCustom;
-							KeyCustom.PropertyRow = MakeShared<FDetailPropertyRow>(ChildNode->GetPropertyKeyNode(), ParentCategoryRef);
-							TSharedRef<FDetailItemNode> KeyNodeItem = MakeShared<FDetailItemNode>(KeyCustom, ParentCategoryRef, ParentEnabledState);
-							KeyNodeItem->Initialize();
+							// If the child has a key property, only create a second node for the key if the child did not already create a property
+							// editor for it
+							if ( !Customization.PropertyRow->PropertyKeyEditor.IsValid() )
+							{
+								FDetailLayoutCustomization KeyCustom;
+								KeyCustom.PropertyRow = MakeShared<FDetailPropertyRow>(ChildNode->GetPropertyKeyNode(), ParentCategoryRef);
+								TSharedRef<FDetailItemNode> KeyNodeItem = MakeShared<FDetailItemNode>(KeyCustom, ParentCategoryRef, ParentEnabledState);
+								KeyNodeItem->Initialize();
 							
-							PropNodes.Add(KeyNodeItem);
-							bHasKeyNode = true;
+								PropNodes.Add(KeyNodeItem);
+								bHasKeyNode = true;
+							}
 						}
-					}
 
-					// Add the child node 
-					PropNodes.Add(ChildNodeItem);
+						// Add the child node 
+						PropNodes.Add(ChildNodeItem);
 					
-					// For set properties, set the name override to match the index
-					if (bSetProperty)
-					{
-						ChildNode->SetDisplayNameOverride(FText::AsNumber(ChildIndex));
-					}
+						// For set properties, set the name override to match the index
+						if (bSetProperty)
+						{
+							ChildNode->SetDisplayNameOverride(FText::AsNumber(ChildIndex));
+						}
 
-					if (bMapProperty && bHasKeyNode)
-					{
-						// Group the key/value nodes for map properties
-						static FText KeyValueGroupNameFormat = LOCTEXT("KeyValueGroupName", "Element {0}");
-						FText KeyValueGroupName = FText::Format(KeyValueGroupNameFormat, ChildIndex);
+						if (bMapProperty && bHasKeyNode)
+						{
+							// Group the key/value nodes for map properties
+							static FText KeyValueGroupNameFormat = LOCTEXT("KeyValueGroupName", "Element {0}");
+							FText KeyValueGroupName = FText::Format(KeyValueGroupNameFormat, ChildIndex);
 
-						TSharedRef<FDetailCategoryGroupNode> KeyValueGroupNode = MakeShared<FDetailCategoryGroupNode>(FName(*KeyValueGroupName.ToString()), ParentCategoryRef);
-						KeyValueGroupNode->SetChildren(PropNodes);
-						KeyValueGroupNode->SetShowBorder(false);
-						KeyValueGroupNode->SetHasSplitter(true);
+							TSharedRef<FDetailCategoryGroupNode> KeyValueGroupNode = MakeShared<FDetailCategoryGroupNode>(FName(*KeyValueGroupName.ToString()), ParentCategoryRef);
+							KeyValueGroupNode->SetChildren(PropNodes);
+							KeyValueGroupNode->SetShowBorder(false);
+							KeyValueGroupNode->SetHasSplitter(true);
 
-						OutChildren.Add(KeyValueGroupNode);
-					}
-					else
-					{
-						OutChildren.Append(PropNodes);
+							OutChildren.Add(KeyValueGroupNode);
+						}
+						else
+						{
+							OutChildren.Append(PropNodes);
+						}
 					}
 				}
 			}
