Fix GetAssetFilter on an actor component

#rnx

diff --git a/Engine/Source/Editor/PropertyEditor/Private/UserInterface/PropertyEditor/SPropertyEditorAsset.cpp b/Engine/Source/Editor/PropertyEditor/Private/UserInterface/PropertyEditor/SPropertyEditorAsset.cpp
--- a/Engine/Source/Editor/PropertyEditor/Private/UserInterface/PropertyEditor/SPropertyEditorAsset.cpp
+++ b/Engine/Source/Editor/PropertyEditor/Private/UserInterface/PropertyEditor/SPropertyEditorAsset.cpp
@@ -445,10 +445,37 @@ void SPropertyEditorAsset::Construct(const FArguments& InArgs, const TSharedPtr<
 			}
 			for (UObject* Object : ObjectList)
 			{
-				const UFunction* GetAssetFilterFunction = Object ? Object->FindFunction(*GetAssetFilterFunctionName) : nullptr;
-				if (GetAssetFilterFunction)
+				if (!Object)
 				{
-					AppendOnShouldFilterAssetCallback(FOnShouldFilterAsset::CreateUFunction(Object, GetAssetFilterFunction->GetFName()));
+					continue;
+				}
+				const UFunction* GetAssetFilterFunction = Object->FindFunction(*GetAssetFilterFunctionName);
+				if (ensureMsgf(GetAssetFilterFunction, TEXT("Could not find UFunction %s on %s"), *GetAssetFilterFunctionName, *GetNameSafe(Object)))
+				{
+					if (UActorComponent* ActorComponent = Cast<UActorComponent>(Object) )
+					{
+						// Create a soft reference on the component
+						FSoftComponentReference ComponentReference;
+						ComponentReference.OtherActor = ActorComponent->GetOwner();
+						ComponentReference.PathToComponent = ActorComponent->GetPathName(ActorComponent->GetOwner());
+						AppendOnShouldFilterAssetCallback(FOnShouldFilterAsset::CreateLambda([ComponentReference, FunctionName =  GetAssetFilterFunction->GetFName()](const FAssetData& AssetData) {
+
+							if (UActorComponent* Component = ComponentReference.GetComponent(nullptr))
+							{
+								FOnShouldFilterAsset UFunctionDelegate = FOnShouldFilterAsset::CreateUFunction(Component, FunctionName);
+								if (UFunctionDelegate.IsBound())
+								{
+									return UFunctionDelegate.Execute(AssetData);
+								}
+							}
+							return false;
+						}));
+						
+					}
+					else
+					{
+						AppendOnShouldFilterAssetCallback(FOnShouldFilterAsset::CreateUFunction(Object, GetAssetFilterFunction->GetFName()));
+					}
 				}
 			}
 		}
