[Core] Fix crash when editing values in containers in Property Bags

* Object == null was used as information if the object modified was a struct or not
* I changed that recently to still get the object pointer, because it was necessary for the ImportText to propagate the changes to other objects (like BP template -> BP instances)
* So now also keep track in the info struct if we are modifying a struct or not and changed everywhere we checked that.

#rnx

diff --git a/Engine/Source/Editor/PropertyEditor/Private/PropertyHandleImpl.cpp b/Engine/Source/Editor/PropertyEditor/Private/PropertyHandleImpl.cpp
--- a/Engine/Source/Editor/PropertyEditor/Private/PropertyHandleImpl.cpp
+++ b/Engine/Source/Editor/PropertyEditor/Private/PropertyHandleImpl.cpp
@@ -101,7 +101,7 @@ void FPropertyValueImpl::EnumerateObjectsToModify( FPropertyNode* InPropertyNode
 
 			uint8* BaseAddress = InPropertyNode->GetValueBaseAddress(StructAddress, InPropertyNode->HasNodeFlags(EPropertyNodeFlags::IsSparseClassData), bIsStruct);
 
-			if (!InObjectsToModifyCallback(FObjectBaseAddress(Object, StructAddress, BaseAddress), Index, NumInstances))
+			if (!InObjectsToModifyCallback(FObjectBaseAddress(Object, StructAddress, BaseAddress, bIsStruct), Index, NumInstances))
 			{
 				break;
 			}
@@ -404,7 +404,7 @@ FPropertyAccess::Result FPropertyValueImpl::ImportText( const TArray<FObjectBase
 			FProperty* Property = ParentNode ? ParentNode->GetProperty() : nullptr;
 
 			const bool bIsSparseClassData = InPropertyNode->HasNodeFlags(EPropertyNodeFlags::IsSparseClassData) != 0;
-			const bool bIsStruct = !Cur.Object;
+			const bool bIsStruct = Cur.bIsStruct;
 
 			bool bIsInContainer = false;
 
@@ -5928,7 +5928,7 @@ bool FPropertyHandleSet::HasDefaultElement()
 		if (Addresses.Num() > 0)
 		{
 			const bool bIsSparseClassData = PropNode->HasNodeFlags(EPropertyNodeFlags::IsSparseClassData) != 0;
-			const bool bIsStruct = !Addresses[0].Object;
+			const bool bIsStruct = Addresses[0].bIsStruct;
 			FSetProperty* SetProperty = CastFieldChecked<FSetProperty>(PropNode->GetProperty());
 			FScriptSetHelper SetHelper(SetProperty, PropNode->GetValueBaseAddress(Addresses[0].StructAddress, bIsSparseClassData, bIsStruct));
 
@@ -5990,7 +5990,7 @@ FPropertyHandleItemAddResult FPropertyHandleSet::AddItem()
 			if (Addresses.Num() > 0)
 			{
 				const bool bIsSparseClassData = PropNode->HasNodeFlags(EPropertyNodeFlags::IsSparseClassData) != 0;
-				const bool bIsStruct = !Addresses[0].Object;
+				const bool bIsStruct = Addresses[0].bIsStruct;
 				uint8* ValueBaseAddress = PropNode->GetValueBaseAddress(Addresses[0].StructAddress, bIsSparseClassData, bIsStruct);
 
 				FSetProperty* SetProperty = CastFieldChecked<FSetProperty>(PropNode->GetProperty());
@@ -6189,7 +6189,7 @@ bool FPropertyHandleMap::HasDefaultKey()
 		if (Addresses.Num() > 0)
 		{
 			const bool bIsSparseClassData = PropNode->HasNodeFlags(EPropertyNodeFlags::IsSparseClassData) != 0;
-			const bool bIsStruct = !Addresses[0].Object;
+			const bool bIsStruct = Addresses[0].bIsStruct;
 			FMapProperty* MapProperty = CastFieldChecked<FMapProperty>(PropNode->GetProperty());
 			FScriptMapHelper MapHelper(MapProperty, PropNode->GetValueBaseAddress(Addresses[0].StructAddress, bIsSparseClassData, bIsStruct));
 
@@ -6252,7 +6252,7 @@ FPropertyHandleItemAddResult FPropertyHandleMap::AddItem()
 			if (Addresses.Num() > 0)
 			{
 				const bool bIsSparseClassData = PropNode->HasNodeFlags(EPropertyNodeFlags::IsSparseClassData) != 0;
-				const bool bIsStruct = !Addresses[0].Object;
+				const bool bIsStruct = Addresses[0].bIsStruct;
 				uint8* ValueBaseAddress = PropNode->GetValueBaseAddress(Addresses[0].StructAddress, bIsSparseClassData, bIsStruct);
 			
 				FMapProperty* MapProperty = CastFieldChecked<FMapProperty>(PropNode->GetProperty());
diff --git a/Engine/Source/Editor/PropertyEditor/Private/PropertyHandleImpl.h b/Engine/Source/Editor/PropertyEditor/Private/PropertyHandleImpl.h
--- a/Engine/Source/Editor/PropertyEditor/Private/PropertyHandleImpl.h
+++ b/Engine/Source/Editor/PropertyEditor/Private/PropertyHandleImpl.h
@@ -11,23 +11,26 @@
 
 class FNotifyHook;
 class FPropertyRestriction;
-
+	
 class FObjectBaseAddress
 {
 public:
 	FObjectBaseAddress() = default;
-	FObjectBaseAddress(UObject* InObject, uint8* InStructAddress, uint8* InBaseAddress)
+	FObjectBaseAddress(UObject* InObject, uint8* InStructAddress, uint8* InBaseAddress, bool bInIsStruct)
 		: Object(InObject)
 		, StructAddress(InStructAddress)
 		, BaseAddress(InBaseAddress)
+		, bIsStruct(bInIsStruct)
 	{}
 
-	/** Object associated with the value, or null when editing a struct. */
+	/** Object associated with the value, or owner object when editing a struct. */
 	UObject* Object = nullptr;
 	/** Pointer to the object/struct instance that actually owns the value (may be pointing to sidecar data). */
 	uint8* StructAddress = nullptr;
 	/** Pointer to the base address of the value within StructAddress. */
 	uint8* BaseAddress = nullptr;
+	/** Editing a struct */
+	bool bIsStruct = false;
 };
 
 /**
