Fix SPropertyMenuComponentPicker crash when pasting //

diff --git a/Engine/Source/Editor/PropertyEditor/Private/UserInterface/PropertyEditor/SPropertyMenuComponentPicker.cpp b/Engine/Source/Editor/PropertyEditor/Private/UserInterface/PropertyEditor/SPropertyMenuComponentPicker.cpp
--- a/Engine/Source/Editor/PropertyEditor/Private/UserInterface/PropertyEditor/SPropertyMenuComponentPicker.cpp
+++ b/Engine/Source/Editor/PropertyEditor/Private/UserInterface/PropertyEditor/SPropertyMenuComponentPicker.cpp
@@ -213,7 +213,8 @@ bool SPropertyMenuComponentPicker::CanPaste()
 		FString Class = ClipboardText.Left(SpaceIndex);
 		FString PossibleObjectPath = ClipboardText.Mid(SpaceIndex);
 
-		bCanPaste = !Class.IsEmpty() && !PossibleObjectPath.IsEmpty();
+		bCanPaste = !Class.IsEmpty() && !PossibleObjectPath.IsEmpty()
+					&& !Class.Contains(TEXT("//")); // LoadClass does not tolerate '//'
 		if (bCanPaste)
 		{
 			bCanPaste = LoadClass<UActorComponent>(nullptr, *Class, nullptr, LOAD_Quiet | LOAD_NoWarn) != nullptr;
