PropertyEditor: Fix FPropertyHandleArray::AddItem() returning an incorrect index

diff --git a/Engine/Source/Editor/PropertyEditor/Private/PropertyHandleImpl.cpp b/Engine/Source/Editor/PropertyEditor/Private/PropertyHandleImpl.cpp
--- a/Engine/Source/Editor/PropertyEditor/Private/PropertyHandleImpl.cpp
+++ b/Engine/Source/Editor/PropertyEditor/Private/PropertyHandleImpl.cpp
@@ -5335,7 +5335,7 @@ FPropertyHandleItemAddResult FPropertyHandleArray::AddItem()
 	int32 Index = INDEX_NONE;
 	if( IsEditable() )
 	{
-		Implementation->AddChild();
+		Index = Implementation->AddChild();
 		Result = FPropertyAccess::Success;
 		return FPropertyHandleItemAddResult(Index, Result);
 	}
