Do not always use the component reference

#rnx

diff --git a/Engine/Source/Editor/PropertyEditor/Private/UserInterface/PropertyEditor/SPropertyEditorAsset.cpp b/Engine/Source/Editor/PropertyEditor/Private/UserInterface/PropertyEditor/SPropertyEditorAsset.cpp
--- a/Engine/Source/Editor/PropertyEditor/Private/UserInterface/PropertyEditor/SPropertyEditorAsset.cpp
+++ b/Engine/Source/Editor/PropertyEditor/Private/UserInterface/PropertyEditor/SPropertyEditorAsset.cpp
@@ -455,12 +455,16 @@ void SPropertyEditorAsset::Construct(const FArguments& InArgs, const TSharedPtr<
 					if (UActorComponent* ActorComponent = Cast<UActorComponent>(Object) )
 					{
 						// Create a soft reference on the component
+						TSoftObjectPtr<UActorComponent> ComponentSoftPtr(ActorComponent);
 						FSoftComponentReference ComponentReference;
 						ComponentReference.OtherActor = ActorComponent->GetOwner();
 						ComponentReference.PathToComponent = ActorComponent->GetPathName(ActorComponent->GetOwner());
-						AppendOnShouldFilterAssetCallback(FOnShouldFilterAsset::CreateLambda([ComponentReference, FunctionName =  GetAssetFilterFunction->GetFName()](const FAssetData& AssetData) {
-
-							if (UActorComponent* Component = ComponentReference.GetComponent(nullptr))
+						AppendOnShouldFilterAssetCallback(FOnShouldFilterAsset::CreateLambda([ComponentSoftPtr, ComponentReference, FunctionName =  GetAssetFilterFunction->GetFName()](const FAssetData& AssetData) mutable {
+							if (!ComponentSoftPtr.IsValid())
+							{
+								ComponentSoftPtr = ComponentReference.GetComponent(nullptr);
+							}
+							if (UActorComponent* Component = ComponentSoftPtr.Get())
 							{
 								FOnShouldFilterAsset UFunctionDelegate = FOnShouldFilterAsset::CreateUFunction(Component, FunctionName);
 								if (UFunctionDelegate.IsBound())
