You can now set a string's HintText via the UPROPERTY meta tag.  This affects the SEditableText text box when editing the string from a property grid.

diff --git a/Engine/Source/Editor/PropertyEditor/Private/UserInterface/PropertyEditor/SPropertyEditorText.cpp b/Engine/Source/Editor/PropertyEditor/Private/UserInterface/PropertyEditor/SPropertyEditorText.cpp
--- a/Engine/Source/Editor/PropertyEditor/Private/UserInterface/PropertyEditor/SPropertyEditorText.cpp
+++ b/Engine/Source/Editor/PropertyEditor/Private/UserInterface/PropertyEditor/SPropertyEditorText.cpp
@@ -37,6 +37,7 @@ void SPropertyEditorText::Construct( const FArguments& InArgs, const TSharedRef<
 	static const FName NAME_PasswordField = "PasswordField";
 	static const FName NAME_AllowedCharacters = "AllowedCharacters";
 	static const FName NAME_PropertyValidator = "PropertyValidator";
+	static const FName NAME_HintText = "HintText";
 
 	TSharedRef<IPropertyHandle> PropertyHandle(InPropertyEditor->GetPropertyHandle());
 	bIsMultiLine = PropertyHandle->GetBoolMetaData(NAME_MultiLine);
@@ -87,6 +88,8 @@ void SPropertyEditorText::Construct( const FArguments& InArgs, const TSharedRef<
 		}
 	}
 
+	FString HintText = PropertyHandle->GetMetaData(NAME_HintText);
+
 	TSharedPtr<SHorizontalBox> HorizontalBox;
 	if(bIsMultiLine)
 	{
@@ -133,6 +136,7 @@ void SPropertyEditorText::Construct( const FArguments& InArgs, const TSharedRef<
 				.SelectAllTextOnCommit( true )
 				.IsReadOnly(this, &SPropertyEditorText::IsReadOnly)
 				.IsPassword( bIsPassword )
+				.HintText(FText::FromString(HintText))
 			]
 		];
 
