Added support for showing Verse paths to asset property editor tooltips.
Don't show the string value for the asset in the empty space of the value widget.

diff --git a/Engine/Source/Editor/PropertyEditor/Private/UserInterface/PropertyEditor/SPropertyEditorAsset.cpp b/Engine/Source/Editor/PropertyEditor/Private/UserInterface/PropertyEditor/SPropertyEditorAsset.cpp
--- a/Engine/Source/Editor/PropertyEditor/Private/UserInterface/PropertyEditor/SPropertyEditorAsset.cpp
+++ b/Engine/Source/Editor/PropertyEditor/Private/UserInterface/PropertyEditor/SPropertyEditorAsset.cpp
@@ -838,6 +838,10 @@ void SPropertyEditorAsset::Construct(const FArguments& InArgs, const TSharedPtr<
 	{
 		ButtonBoxWrapper->SetVisibility(NumButtons > 0 ? EVisibility::Visible : EVisibility::Collapsed);
 	}
+
+	// Create a default empty tooltip to prevent SPropertyValueWidget::Construct from setting our tooltip to the default FPropertyEditor::GetValueAsText.
+	// We do this because we create a user friendly tooltip in SPropertyEditorAsset::OnGetToolTip, but only want it to appear on certain child widgets to avoid it covering up any buttons.
+	SetToolTipText(TAttribute<FText>::CreateLambda([]() { return FText::GetEmpty(); }));
 }
 
 void SPropertyEditorAsset::GetDesiredWidth( float& OutMinDesiredWidth, float &OutMaxDesiredWidth )
@@ -1217,12 +1221,40 @@ FText SPropertyEditorAsset::OnGetToolTip() const
 		}
 		else if( Value.Object != nullptr )
 		{
-			// Display the package name which is a valid path to the object without redundant information
-			ToolTipText = FText::AsCultureInvariant(Value.Object->GetOutermost()->GetName());
+			UE::Core::FVersePath VersePath;
+
+			if (FAssetToolsModule::GetModule().Get().ShowingContentVersePath())
+			{
+				VersePath = Value.Object->GetVersePath();
+			}
+
+			if (VersePath.IsValid())
+			{
+				ToolTipText = FText::AsCultureInvariant(MoveTemp(VersePath).ToString());
+			}
+			else
+			{
+				// Display the package name which is a valid path to the object without redundant information
+				ToolTipText = FText::AsCultureInvariant(Value.Object->GetOutermost()->GetName());
+			}
 		}
 		else if( Value.AssetData.IsValid() )
 		{
-			ToolTipText = FText::AsCultureInvariant(Value.AssetData.PackageName.ToString());
+			UE::Core::FVersePath VersePath;
+
+			if (FAssetToolsModule::GetModule().Get().ShowingContentVersePath())
+			{
+				VersePath = Value.AssetData.GetVersePath();
+			}
+
+			if (VersePath.IsValid())
+			{
+				ToolTipText = FText::AsCultureInvariant(MoveTemp(VersePath).ToString());
+			}
+			else
+			{
+				ToolTipText = FText::AsCultureInvariant(Value.AssetData.PackageName.ToString());
+			}
 		}
 	}
 	else if( Result == FPropertyAccess::MultipleValues )
@@ -1232,7 +1264,30 @@ FText SPropertyEditorAsset::OnGetToolTip() const
 
 	if( ToolTipText.IsEmpty() )
 	{
-		ToolTipText = FText::AsCultureInvariant(ObjectPath.Get());
+		UE::Core::FVersePath VersePath;
+
+		if (FAssetToolsModule::GetModule().Get().ShowingContentVersePath())
+		{
+			FSoftObjectPath SoftObjectPath(ObjectPath.Get());
+			if (SoftObjectPath.IsValid())
+			{
+				const FAssetRegistryModule& AssetRegistryModule = FModuleManager::LoadModuleChecked<FAssetRegistryModule>(TEXT("AssetRegistry"));
+				FAssetData AssetData;
+				if (AssetRegistryModule.Get().TryGetAssetByObjectPath(SoftObjectPath, AssetData) == UE::AssetRegistry::EExists::Exists)
+				{
+					VersePath = AssetData.GetVersePath();
+				}
+			}
+		}
+
+		if (VersePath.IsValid())
+		{
+			ToolTipText = FText::AsCultureInvariant(MoveTemp(VersePath).ToString());
+		}
+		else
+		{
+			ToolTipText = FText::AsCultureInvariant(ObjectPath.Get());
+		}
 	}
 
 	return ToolTipText;
