[Backout] - 40148190
Fix regression where browse to asset no longer works on unloaded assets

Original CL Desc
------------
Actor Picker: Fix Browse button did not work when an LI actor was referenced by it.Backed out 40148190

diff --git a/Engine/Source/Editor/PropertyEditor/Private/Presentation/PropertyEditor/PropertyEditor.cpp b/Engine/Source/Editor/PropertyEditor/Private/Presentation/PropertyEditor/PropertyEditor.cpp
--- a/Engine/Source/Editor/PropertyEditor/Private/Presentation/PropertyEditor/PropertyEditor.cpp
+++ b/Engine/Source/Editor/PropertyEditor/Private/Presentation/PropertyEditor/PropertyEditor.cpp
@@ -8,14 +8,11 @@
 #include "ObjectPropertyNode.h"
 
 #include "SceneOutlinerFilters.h"
-#include "Selection.h"
 #include "IDetailPropertyRow.h"
 #include "PropertyCustomizationHelpers.h"
 #include "PropertyEditorHelpers.h"
 #include "Editor.h"
 #include "EditorClassUtils.h"
-#include "Editor/UnrealEdEngine.h"
-#include "Elements/Framework/EngineElementsLibrary.h"
 #include "Kismet2/KismetEditorUtilities.h"
 #include "IConfigEditorModule.h"
 #include "PropertyNode.h"
@@ -728,25 +725,6 @@ void FPropertyEditor::SyncToObjectsInNode( const TWeakPtr< FPropertyNode >& Weak
 		// If a list of addresses is provided, GetReadAddress may still return false but we can operate on the property addresses even if they have different values.
 		check( ReadAddresses.Num() > 0 );
 
-		TArray<UObject*> Objects;
-
-		// If this is a Soft Object Ptr, then let's resolve the path via the soft path,
-		// as it'll handle EditorPaths (e.g. LI actors), which Static(Find|Load)Object does not support.
-		if (const FSoftObjectProperty* SoftObjectProperty = CastField<FSoftObjectProperty>(NodeProperty))
-		{
-			for (int32 AddrIndex = 0; AddrIndex < ReadAddresses.Num(); ++AddrIndex)
-			{
-				const uint8* Address = ReadAddresses.GetAddress(AddrIndex);
-				if (UObject* LoadedObject = SoftObjectProperty->GetObjectPropertyValue(Address))
-				{
-					Objects.Add(LoadedObject);
-				}
-			}
-
-			// No need to check the FString version now.
-			ReadAddresses.Reset();
-		}
-
 		// Create a list of object names.
 		TArray<FString> ObjectNames;
 		ObjectNames.Empty( ReadAddresses.Num() );
@@ -764,6 +742,7 @@ void FPropertyEditor::SyncToObjectsInNode( const TWeakPtr< FPropertyNode >& Weak
 
 
 		// Create a list of objects to sync the generic browser to.
+		TArray<UObject*> Objects;
 		for ( int32 ObjectIndex = 0 ; ObjectIndex < ObjectNames.Num() ; ++ObjectIndex )
 		{
 			UObject* Object = nullptr;
@@ -803,31 +782,14 @@ void FPropertyEditor::SyncToObjectsInNode( const TWeakPtr< FPropertyNode >& Weak
 
 			if (Actor->GetLevel())
 			{
-				const USelection* ActorSelection = GEditor->GetSelectedActors();
-				if (UTypedElementSelectionSet* SelectionSet = ActorSelection->GetElementSelectionSet())
-				{
-					const FTypedElementSelectionOptions SelectionOptions = FTypedElementSelectionOptions()
-						.SetAllowHidden(true)
-						.SetWarnIfLocked(true)
-						.SetAllowLegacyNotifications(false)
-						.SetAllowSubRootSelection(true);
+				TArray<AActor*> Actors;
+				Actors.Add(Actor);
 
-					FTypedElementHandle ActorHandle = UEngineElementsLibrary::AcquireEditorActorElementHandle(Actor);
-					const bool bSelectionChanged = SelectionSet->SetSelection(MakeArrayView(&ActorHandle, 1), SelectionOptions);
-
-					if (bSelectionChanged)
-					{
-						SelectionSet->NotifyPendingChanges();
-					}
-					else
-					{
-						// Reset the property windows, in case something has changed since previous selection
-						GUnrealEd->UpdateFloatingPropertyWindows();
-					}
-				}
+				GEditor->SelectNone(/*bNoteSelectionChange=*/false, /*bDeselectBSPSurfs=*/true);
+				GEditor->SelectActor(Actor, /*InSelected=*/true, /*bNotify=*/true, /*bSelectEvenIfHidden=*/true);
 
 				// Jump to the location of the actor
-				GEditor->MoveViewportCamerasToActor(*Actor, /*bActiveViewportOnly=*/false);
+				GEditor->MoveViewportCamerasToActor(Actors, /*bActiveViewportOnly=*/false);
 			}
 		}
 		else if ( Objects.Num() > 0 )
diff --git a/Engine/Source/Editor/PropertyEditor/PropertyEditor.Build.cs b/Engine/Source/Editor/PropertyEditor/PropertyEditor.Build.cs
--- a/Engine/Source/Editor/PropertyEditor/PropertyEditor.Build.cs
+++ b/Engine/Source/Editor/PropertyEditor/PropertyEditor.Build.cs
@@ -49,9 +49,7 @@ public class PropertyEditor : ModuleRules
 				"ToolWidgets",
 				"WidgetRegistration",
 				"Json",
-				"ToolMenus",
-				"TypedElementFramework",
-				"TypedElementRuntime",
+				"ToolMenus"
 			}
         );
 
