[UE] Include propagated subobject and archetype flags on new instances constructed for non-nullable editable reference container properties on editor "add" actions.

diff --git a/Engine/Source/Editor/PropertyEditor/Private/PropertyHandleImpl.cpp b/Engine/Source/Editor/PropertyEditor/Private/PropertyHandleImpl.cpp
--- a/Engine/Source/Editor/PropertyEditor/Private/PropertyHandleImpl.cpp
+++ b/Engine/Source/Editor/PropertyEditor/Private/PropertyHandleImpl.cpp
@@ -1286,9 +1286,20 @@ int32 FPropertyValueImpl::AddChild()
 							NewItemOuter = const_cast<UObject*>(InstanceFromIDO);
 						}
 
+						// Ensure that we propagate the appropriate flags to new instances below, and include the archetype flag if the outer is a template.
+						EObjectFlags MaskedOuterFlags = RF_NoFlags;
+						if (NewItemOuter)
+						{
+							MaskedOuterFlags = NewItemOuter->GetMaskedFlags(RF_PropagateToSubObjects);
+							if (NewItemOuter->HasAnyFlags(RF_ClassDefaultObject | RF_ArchetypeObject))
+							{
+								MaskedOuterFlags |= RF_ArchetypeObject;
+							}
+						}
+
 						if (Array)
 						{
-							Array->PerformOperationWithSetter(Obj, Addr, [NewItemOuter, Array, &Index, &ReturnLogicalIndex](void* DirectAddress)
+							Array->PerformOperationWithSetter(Obj, Addr, [NewItemOuter, Array, &Index, &ReturnLogicalIndex, MaskedOuterFlags](void* DirectAddress)
 							{								
 								FScriptArrayHelper	ArrayHelper(Array, DirectAddress);
 								Index = ArrayHelper.AddValue();
@@ -1298,14 +1309,14 @@ int32 FPropertyValueImpl::AddChild()
 								FObjectProperty* InnerObjectProperty = CastField<FObjectProperty>(Array->Inner);
 								if (InnerObjectProperty && InnerObjectProperty->HasAnyPropertyFlags(CPF_NonNullable) && !InnerObjectProperty->PropertyClass->HasAnyClassFlags(CLASS_Abstract))
 								{
-									UObject* NewItem = NewObject<UObject>(NewItemOuter, InnerObjectProperty->PropertyClass);
+									UObject* NewItem = NewObject<UObject>(NewItemOuter, InnerObjectProperty->PropertyClass, NAME_None, MaskedOuterFlags);
 									InnerObjectProperty->SetObjectPropertyValue(ArrayHelper.GetRawPtr(Index), NewItem);
 								}
 							});
 						}
 						else if (Set)
 						{
-							Set->PerformOperationWithSetter(Obj, Addr, [NewItemOuter, Set, &Index, &ReturnLogicalIndex](void* DirectAddress)
+							Set->PerformOperationWithSetter(Obj, Addr, [NewItemOuter, Set, &Index, &ReturnLogicalIndex, MaskedOuterFlags](void* DirectAddress)
 							{
 								FScriptSetHelper	SetHelper(Set, DirectAddress);
 								Index = SetHelper.AddDefaultValue_Invalid_NeedsRehash();
@@ -1315,7 +1326,7 @@ int32 FPropertyValueImpl::AddChild()
 								FObjectProperty* ElementObjectProperty = CastField<FObjectProperty>(Set->ElementProp);
 								if (ElementObjectProperty && ElementObjectProperty->HasAnyPropertyFlags(CPF_NonNullable) && !ElementObjectProperty->PropertyClass->HasAnyClassFlags(CLASS_Abstract))
 								{
-									UObject* NewItem = NewObject<UObject>(NewItemOuter, ElementObjectProperty->PropertyClass);
+									UObject* NewItem = NewObject<UObject>(NewItemOuter, ElementObjectProperty->PropertyClass, NAME_None, MaskedOuterFlags);
 									ElementObjectProperty->SetObjectPropertyValue(SetHelper.GetElementPtr(Index), NewItem);
 								}
 
@@ -1324,7 +1335,7 @@ int32 FPropertyValueImpl::AddChild()
 						}
 						else if (Map)
 						{
-							Map->PerformOperationWithSetter(Obj, Addr, [NewItemOuter, Map, &Index, &bAddedMapEntry, &ReturnLogicalIndex](void* DirectAddress)
+							Map->PerformOperationWithSetter(Obj, Addr, [NewItemOuter, Map, &Index, &bAddedMapEntry, &ReturnLogicalIndex, MaskedOuterFlags](void* DirectAddress)
 							{
 								FScriptMapHelper	MapHelper(Map, DirectAddress);
 								Index = MapHelper.AddDefaultValue_Invalid_NeedsRehash();
@@ -1335,7 +1346,7 @@ int32 FPropertyValueImpl::AddChild()
 									FObjectProperty* KeyObjectProperty = CastField<FObjectProperty>(Map->KeyProp);
 									if (KeyObjectProperty && KeyObjectProperty->HasAnyPropertyFlags(CPF_NonNullable) && !KeyObjectProperty->PropertyClass->HasAnyClassFlags(CLASS_Abstract))
 									{
-										UObject* NewItem = NewObject<UObject>(NewItemOuter, KeyObjectProperty->PropertyClass);
+										UObject* NewItem = NewObject<UObject>(NewItemOuter, KeyObjectProperty->PropertyClass, NAME_None, MaskedOuterFlags);
 										KeyObjectProperty->SetObjectPropertyValue(MapHelper.GetKeyPtr(Index), NewItem);
 									}
 								}
@@ -1344,7 +1355,7 @@ int32 FPropertyValueImpl::AddChild()
 									FObjectProperty* ValueObjectProperty = CastField<FObjectProperty>(Map->ValueProp);
 									if (ValueObjectProperty && ValueObjectProperty->HasAnyPropertyFlags(CPF_NonNullable) && !ValueObjectProperty->PropertyClass->HasAnyClassFlags(CLASS_Abstract))
 									{
-										UObject* NewItem = NewObject<UObject>(NewItemOuter, ValueObjectProperty->PropertyClass);
+										UObject* NewItem = NewObject<UObject>(NewItemOuter, ValueObjectProperty->PropertyClass, NAME_None, MaskedOuterFlags);
 										ValueObjectProperty->SetObjectPropertyValue(MapHelper.GetValuePtr(Index), NewItem);
 									}
 								}
