Refresh details view on language change so that the categories update correctly

#rnx

diff --git a/Engine/Source/Editor/PropertyEditor/Private/SDetailsViewBase.cpp b/Engine/Source/Editor/PropertyEditor/Private/SDetailsViewBase.cpp
--- a/Engine/Source/Editor/PropertyEditor/Private/SDetailsViewBase.cpp
+++ b/Engine/Source/Editor/PropertyEditor/Private/SDetailsViewBase.cpp
@@ -47,12 +47,19 @@ SDetailsViewBase::SDetailsViewBase() :
 	// RequestForceRefresh is deferred until next tick to avoid the editor locking up for several minutes in one frame when refreshing multiple times
 	PropertyPermissionListChangedDelegate = FPropertyEditorPermissionList::Get().PermissionListUpdatedDelegate.AddLambda([this](TSoftObjectPtr<const UStruct> Struct, FName Owner) { RequestForceRefresh(); });
 	PropertyPermissionListEnabledDelegate = FPropertyEditorPermissionList::Get().PermissionListEnabledDelegate.AddRaw(this, &SDetailsViewBase::RequestForceRefresh);
+
+	OnCultureChangedDelegate = FInternationalization::Get().OnCultureChanged().AddRaw(this, &SDetailsViewBase::RequestForceRefresh);
 }
 
 SDetailsViewBase::~SDetailsViewBase()
 {
 	FPropertyEditorPermissionList::Get().PermissionListUpdatedDelegate.Remove(PropertyPermissionListChangedDelegate);
 	FPropertyEditorPermissionList::Get().PermissionListEnabledDelegate.Remove(PropertyPermissionListEnabledDelegate);
+
+	if (FInternationalization::IsAvailable())
+	{
+		FInternationalization::Get().OnCultureChanged().Remove(OnCultureChangedDelegate);
+	}
 }
 
 void SDetailsViewBase::OnGetChildrenForDetailTree(TSharedRef<FDetailTreeNode> InTreeNode, TArray< TSharedRef<FDetailTreeNode> >& OutChildren)
diff --git a/Engine/Source/Editor/PropertyEditor/Private/SDetailsViewBase.h b/Engine/Source/Editor/PropertyEditor/Private/SDetailsViewBase.h
--- a/Engine/Source/Editor/PropertyEditor/Private/SDetailsViewBase.h
+++ b/Engine/Source/Editor/PropertyEditor/Private/SDetailsViewBase.h
@@ -552,6 +552,9 @@ protected:
 	FDelegateHandle PropertyPermissionListChangedDelegate;
 	FDelegateHandle PropertyPermissionListEnabledDelegate;
 
+	/** Used to refresh the tree when the language changes */
+	FDelegateHandle OnCultureChangedDelegate;
+
 	/** Delegate for overriding the show modified filter */
 	FSimpleDelegate CustomFilterDelegate;
 	FText CustomFilterLabel;
diff --git a/Engine/Source/Editor/PropertyEditor/Private/SPropertyTreeViewImpl.cpp b/Engine/Source/Editor/PropertyEditor/Private/SPropertyTreeViewImpl.cpp
--- a/Engine/Source/Editor/PropertyEditor/Private/SPropertyTreeViewImpl.cpp
+++ b/Engine/Source/Editor/PropertyEditor/Private/SPropertyTreeViewImpl.cpp
@@ -169,6 +169,8 @@ void SPropertyTreeViewImpl::Construct(const FArguments& InArgs)
 	
 	FPropertyEditorPermissionList::Get().PermissionListUpdatedDelegate.AddSP(this, &SPropertyTreeViewImpl::OnPermissionListUpdated);
 	FPropertyEditorPermissionList::Get().PermissionListEnabledDelegate.AddSP(this, &SPropertyTreeViewImpl::RequestRefresh);
+
+	FInternationalization::Get().OnCultureChanged().AddSP(this, &SPropertyTreeViewImpl::RequestRefresh);
 }
 
 void SPropertyTreeViewImpl::OnPermissionListUpdated(TSoftObjectPtr<const UStruct>, FName)
