PropertyPermissionList - Fix buttons not appearing for functions marked CallInEditor details panel unless added to AllowList explicitly

diff --git a/Engine/Source/Editor/PropertyEditor/Private/PropertyPermissionList.cpp b/Engine/Source/Editor/PropertyEditor/Private/PropertyPermissionList.cpp
--- a/Engine/Source/Editor/PropertyEditor/Private/PropertyPermissionList.cpp
+++ b/Engine/Source/Editor/PropertyEditor/Private/PropertyPermissionList.cpp
@@ -3,15 +3,23 @@
 #include "PropertyPermissionList.h"
 
 #include "Editor.h"
+#include "HAL/IConsoleManager.h"
 #include "Misc/CoreDelegates.h"
 #include "UObject/UnrealType.h"
 #include "DetailTreeNode.h"
 
 DEFINE_LOG_CATEGORY(LogPropertyEditorPermissionList);
 
-namespace
+namespace PropertyEditorPermissionList
 {
 	const FName PropertyPermissionListOwner = "PropertyPermissionList";
+	static const FName NAME_CallInEditor("CallInEditor");
+
+	static bool bSupportAllowAllFunctions = true;
+	static FAutoConsoleVariableRef CVarSupportAllowAllFunctions(
+		TEXT("PropertyEditor.PropertyEditorPermissionList.SupportAllowAllFunctions"),
+		bSupportAllowAllFunctions,
+		TEXT("Enable allowing all zero parameter call in editor functions when allowing all properties"));
 }
 
 bool operator==(const FPropertyPermissionList::FPermissionListUpdate& A, const FPropertyPermissionList::FPermissionListUpdate& B)
@@ -340,7 +348,19 @@ const FNamePermissionList& FPropertyPermissionList::GetCachedPermissionListForSt
 				{
 					for (TFieldIterator<FProperty> Property(SuperStruct, EFieldIteratorFlags::IncludeSuper, EFieldIteratorFlags::ExcludeDeprecated); Property; ++Property)
 					{
-						NewPermissionList.AddAllowListItem(PropertyPermissionListOwner, Property->GetFName());
+						NewPermissionList.AddAllowListItem(PropertyEditorPermissionList::PropertyPermissionListOwner, Property->GetFName());
+					}
+
+					if (PropertyEditorPermissionList::bSupportAllowAllFunctions)
+					{
+						for (TFieldIterator<UFunction> FunctionIter(SuperStruct, EFieldIteratorFlags::IncludeSuper, EFieldIteratorFlags::ExcludeDeprecated); FunctionIter; ++FunctionIter)
+						{
+							UFunction* TestFunction = *FunctionIter;
+							if ((TestFunction->ParmsSize == 0) && TestFunction->GetBoolMetaData(PropertyEditorPermissionList::NAME_CallInEditor))
+							{
+								NewPermissionList.AddAllowListItem(PropertyEditorPermissionList::PropertyPermissionListOwner, TestFunction->GetFName());
+							}
+						}
 					}
 				}
 				NewPermissionList.Append(Entry->PermissionList);
@@ -355,7 +375,19 @@ const FNamePermissionList& FPropertyPermissionList::GetCachedPermissionListForSt
 		{
 			for (TFieldIterator<FProperty> Property(Struct, EFieldIteratorFlags::ExcludeSuper, EFieldIteratorFlags::ExcludeDeprecated); Property; ++Property)
 			{
-				NewPermissionList.AddAllowListItem(PropertyPermissionListOwner, Property->GetFName());
+				NewPermissionList.AddAllowListItem(PropertyEditorPermissionList::PropertyPermissionListOwner, Property->GetFName());
+			}
+
+			if (PropertyEditorPermissionList::bSupportAllowAllFunctions)
+			{
+				for (TFieldIterator<UFunction> FunctionIter(Struct, EFieldIteratorFlags::ExcludeSuper, EFieldIteratorFlags::ExcludeDeprecated); FunctionIter; ++FunctionIter)
+				{
+					UFunction* TestFunction = *FunctionIter;
+					if ((TestFunction->ParmsSize == 0) && TestFunction->GetBoolMetaData(PropertyEditorPermissionList::NAME_CallInEditor))
+					{
+						NewPermissionList.AddAllowListItem(PropertyEditorPermissionList::PropertyPermissionListOwner, TestFunction->GetFName());
+					}
+				}
 			}
 		}
 
