Updated headers using UnrealCodeFixup to make sure dllstorage is on methods/staticvars instead of types. Part 5/n

diff --git a/Engine/Source/Editor/Persona/Public/BlendSpaceAnalysis.h b/Engine/Source/Editor/Persona/Public/BlendSpaceAnalysis.h
--- a/Engine/Source/Editor/Persona/Public/BlendSpaceAnalysis.h
+++ b/Engine/Source/Editor/Persona/Public/BlendSpaceAnalysis.h
@@ -10,6 +10,8 @@
 
 #include "BlendSpaceAnalysis.generated.h"
 
+#define UE_API PERSONA_API
+
 #define LOCTEXT_NAMESPACE "BlendSpaceAnalysis"
 //#define ANALYSIS_VERBOSE_LOG
 
@@ -83,12 +85,12 @@ enum class EAnalysisEulerAxis : uint8
  * the parameters used by the engine functions. User defined can inherit from this and add their own - then the
  * user-defined MakeCache function should replace any base class cache that is passed in with their own.
 */
-UCLASS()
-class PERSONA_API UCachedAnalysisProperties : public UObject
+UCLASS(MinimalAPI)
+class UCachedAnalysisProperties : public UObject
 {
 	GENERATED_BODY()
 public:
-	void CopyFrom(const UCachedAnalysisProperties& Other);
+	UE_API void CopyFrom(const UCachedAnalysisProperties& Other);
 	EAnalysisLinearAxis     LinearFunctionAxis = EAnalysisLinearAxis::PlusX;
 	EAnalysisEulerAxis      EulerFunctionAxis = EAnalysisEulerAxis::Pitch;
 	FBoneSocketTarget       BoneSocket1;
@@ -103,8 +105,8 @@ public:
 	float                   EndTimeFraction = 1.0f;
 };
 
-UCLASS(Abstract)
-class PERSONA_API ULinearAnalysisPropertiesBase : public UAnalysisProperties
+UCLASS(MinimalAPI, Abstract)
+class ULinearAnalysisPropertiesBase : public UAnalysisProperties
 {
 	GENERATED_BODY()
 public:
@@ -135,13 +137,13 @@ public:
 	float EndTimeFraction = 1.0f;
 };
 
-UCLASS()
-class PERSONA_API ULinearAnalysisProperties : public ULinearAnalysisPropertiesBase
+UCLASS(MinimalAPI)
+class ULinearAnalysisProperties : public ULinearAnalysisPropertiesBase
 {
 	GENERATED_BODY()
 public:
-	void InitializeFromCache(TObjectPtr<UCachedAnalysisProperties> Cache) override;
-	void MakeCache(TObjectPtr<UCachedAnalysisProperties>& Cache, UBlendSpace* BlendSpace) override;
+	UE_API void InitializeFromCache(TObjectPtr<UCachedAnalysisProperties> Cache) override;
+	UE_API void MakeCache(TObjectPtr<UCachedAnalysisProperties>& Cache, UBlendSpace* BlendSpace) override;
 
 	/** Axis for the analysis function */
 	UPROPERTY(EditAnywhere, DisplayName = "Axis", Category = AnalysisProperties)
@@ -149,13 +151,13 @@ public:
 };
 
 
-UCLASS()
-class PERSONA_API UEulerAnalysisProperties : public UAnalysisProperties
+UCLASS(MinimalAPI)
+class UEulerAnalysisProperties : public UAnalysisProperties
 {
 	GENERATED_BODY()
 public:
-	void InitializeFromCache(TObjectPtr<UCachedAnalysisProperties> Cache) override;
-	void MakeCache(TObjectPtr<UCachedAnalysisProperties>& Cache, UBlendSpace* BlendSpace) override;
+	UE_API void InitializeFromCache(TObjectPtr<UCachedAnalysisProperties> Cache) override;
+	UE_API void MakeCache(TObjectPtr<UCachedAnalysisProperties>& Cache, UBlendSpace* BlendSpace) override;
 
 	/** Axis for the analysis function */
 	UPROPERTY(EditAnywhere, DisplayName = "Axis", Category = AnalysisProperties)
@@ -961,3 +963,4 @@ static bool CalculateOrientationRate(
 
 #undef LOCTEXT_NAMESPACE
 
+#undef UE_API
diff --git a/Engine/Source/Editor/Persona/Public/BoneSelectionWidget.h b/Engine/Source/Editor/Persona/Public/BoneSelectionWidget.h
--- a/Engine/Source/Editor/Persona/Public/BoneSelectionWidget.h
+++ b/Engine/Source/Editor/Persona/Public/BoneSelectionWidget.h
@@ -11,6 +11,8 @@
 #include "Widgets/Views/STableRow.h"
 #include "Widgets/Views/STreeView.h"
 
+#define UE_API PERSONA_API
+
 class IEditableSkeleton;
 class SComboButton;
 
@@ -19,7 +21,7 @@ DECLARE_DELEGATE_RetVal_OneParam(FName, FGetSelectedBone, bool& /*bMultipleValue
 DECLARE_DELEGATE_RetVal(const struct FReferenceSkeleton&, FGetReferenceSkeleton);
 DECLARE_DELEGATE_RetVal(const TArray<class USkeletalMeshSocket*>&, FGetSocketList);
 
-class PERSONA_API SBoneTreeMenu : public SCompoundWidget
+class SBoneTreeMenu : public SCompoundWidget
 {
 public:
 	// Storage object for bone hierarchy
@@ -56,32 +58,32 @@ public:
 	*
 	* @param	InArgs	The declaration data for this widget
 	*/
-	void Construct(const FArguments& InArgs);
+	UE_API void Construct(const FArguments& InArgs);
 
 	/** Get the filter text widget, e.g. for focus */
-	TSharedPtr<SWidget> GetFilterTextWidget();
+	UE_API TSharedPtr<SWidget> GetFilterTextWidget();
 
 private:
 	// SWidget interface
-	virtual FReply OnKeyDown(const FGeometry& MyGeometry, const FKeyEvent& InKeyEvent) override;
+	UE_API virtual FReply OnKeyDown(const FGeometry& MyGeometry, const FKeyEvent& InKeyEvent) override;
 
 	// Using the current filter, repopulate the tree view
-	void RebuildBoneList(const FName& SelectedBone);
+	UE_API void RebuildBoneList(const FName& SelectedBone);
 
 	// Make a single tree row widget
-	TSharedRef<ITableRow> MakeTreeRowWidget(TSharedPtr<FBoneNameInfo> InInfo, const TSharedRef<STableViewBase>& OwnerTable);
+	UE_API TSharedRef<ITableRow> MakeTreeRowWidget(TSharedPtr<FBoneNameInfo> InInfo, const TSharedRef<STableViewBase>& OwnerTable);
 
 	// Get the children for the provided bone info
-	void GetChildrenForInfo(TSharedPtr<FBoneNameInfo> InInfo, TArray< TSharedPtr<FBoneNameInfo> >& OutChildren);
+	UE_API void GetChildrenForInfo(TSharedPtr<FBoneNameInfo> InInfo, TArray< TSharedPtr<FBoneNameInfo> >& OutChildren);
 
 	// Called when the user changes the search filter
-	void OnFilterTextChanged(const FText& InFilterText);
+	UE_API void OnFilterTextChanged(const FText& InFilterText);
 
 	// Handle the tree view selection changing
-	void OnSelectionChanged(TSharedPtr<SBoneTreeMenu::FBoneNameInfo> BoneInfo, ESelectInfo::Type SelectInfo);
+	UE_API void OnSelectionChanged(TSharedPtr<SBoneTreeMenu::FBoneNameInfo> BoneInfo, ESelectInfo::Type SelectInfo);
 
 	// Select a specific bone, helper for UI handler functions
-	void SelectBone(TSharedPtr<SBoneTreeMenu::FBoneNameInfo> BoneInfo);
+	UE_API void SelectBone(TSharedPtr<SBoneTreeMenu::FBoneNameInfo> BoneInfo);
 	
 	// Tree info entries for bone picker
 	TArray<TSharedPtr<FBoneNameInfo>> SkeletonTreeInfo;
@@ -106,7 +108,7 @@ private:
 	bool bShowNone;
 };
 
-class PERSONA_API SBoneSelectionWidget : public SCompoundWidget
+class SBoneSelectionWidget : public SCompoundWidget
 {
 public: 
 
@@ -147,19 +149,19 @@ public:
 	 *
 	 * @param	InArgs	The declaration data for this widget
 	 */
-	void Construct( const FArguments& InArgs );
+	UE_API void Construct( const FArguments& InArgs );
 
 private: 
 
 	// Creates the combo button menu when clicked
-	TSharedRef<SWidget> CreateSkeletonWidgetMenu();
+	UE_API TSharedRef<SWidget> CreateSkeletonWidgetMenu();
 	// Called when the user selects a bone name
-	void OnSelectionChanged(FName BoneName);
+	UE_API void OnSelectionChanged(FName BoneName);
 
 	// Gets the current bone name, used to get the right name for the combo button
-	FText GetCurrentBoneName() const;
+	UE_API FText GetCurrentBoneName() const;
 
-	FText GetFinalToolTip() const;
+	UE_API FText GetFinalToolTip() const;
 
 	// Base combo button 
 	TSharedPtr<SComboButton> BonePickerButton;
@@ -176,3 +178,5 @@ private:
 	// Cache supplied tooltip
 	FText SuppliedToolTip;
 };
+
+#undef UE_API
diff --git a/Engine/Source/Editor/Persona/Public/IPersonaEditMode.h b/Engine/Source/Editor/Persona/Public/IPersonaEditMode.h
--- a/Engine/Source/Editor/Persona/Public/IPersonaEditMode.h
+++ b/Engine/Source/Editor/Persona/Public/IPersonaEditMode.h
@@ -4,14 +4,18 @@
 
 #include "AnimationEditMode.h"
 
-class PERSONA_API IPersonaEditMode : public FAnimationEditMode
+#define UE_API PERSONA_API
+
+class IPersonaEditMode : public FAnimationEditMode
 {
 	// These implementations are duplicated in the Persona module so that code that links only to Persona and not to
 	// the AnimationEditMode module can work properly
 public:
-	IPersonaEditMode();
-	virtual ~IPersonaEditMode() override;
-	virtual void Enter() override;
-	virtual void Exit() override;
-	virtual void AddReferencedObjects(FReferenceCollector& Collector) override;
+	UE_API IPersonaEditMode();
+	UE_API virtual ~IPersonaEditMode() override;
+	UE_API virtual void Enter() override;
+	UE_API virtual void Exit() override;
+	UE_API virtual void AddReferencedObjects(FReferenceCollector& Collector) override;
 };
+
+#undef UE_API
diff --git a/Engine/Source/Editor/Persona/Public/IPersonaEditorModeManager.h b/Engine/Source/Editor/Persona/Public/IPersonaEditorModeManager.h
--- a/Engine/Source/Editor/Persona/Public/IPersonaEditorModeManager.h
+++ b/Engine/Source/Editor/Persona/Public/IPersonaEditorModeManager.h
@@ -7,6 +7,8 @@
 
 #include "IPersonaEditorModeManager.generated.h"
 
+#define UE_API PERSONA_API
+
 UINTERFACE(MinimalAPI)
 class UPersonaManagerContext : public UInterface
 {
@@ -34,16 +36,16 @@ public:
 
 class IPersonaEditorModeManager;
 
-UCLASS()
-class PERSONA_API UPersonaEditorModeManagerContext : public UObject, public IPersonaManagerContext
+UCLASS(MinimalAPI)
+class UPersonaEditorModeManagerContext : public UObject, public IPersonaManagerContext
 {
 	GENERATED_BODY()
 public:
 	// Only for use by FAnimationViewportClient::GetPersonaModeManager for compatibility
-	IPersonaEditorModeManager* GetPersonaEditorModeManager() const;
-	virtual bool GetCameraTarget(FSphere& OutTarget) const override;
-	virtual void GetOnScreenDebugInfo(TArray<FText>& OutDebugText) const override;
-	virtual void SetFocusInViewport();
+	UE_API IPersonaEditorModeManager* GetPersonaEditorModeManager() const;
+	UE_API virtual bool GetCameraTarget(FSphere& OutTarget) const override;
+	UE_API virtual void GetOnScreenDebugInfo(TArray<FText>& OutDebugText) const override;
+	UE_API virtual void SetFocusInViewport();
 private:
 	IPersonaEditorModeManager* ModeManager = nullptr;
 	static UPersonaEditorModeManagerContext* CreateFor(IPersonaEditorModeManager* InModeManager)
@@ -64,3 +66,5 @@ public:
 private:
 	TObjectPtr<UPersonaEditorModeManagerContext> PersonaModeManagerContext{UPersonaEditorModeManagerContext::CreateFor(this)};
 };
+
+#undef UE_API
diff --git a/Engine/Source/Editor/Persona/Public/PersonaAssetEditorToolkit.h b/Engine/Source/Editor/Persona/Public/PersonaAssetEditorToolkit.h
--- a/Engine/Source/Editor/Persona/Public/PersonaAssetEditorToolkit.h
+++ b/Engine/Source/Editor/Persona/Public/PersonaAssetEditorToolkit.h
@@ -4,19 +4,21 @@
 
 #include "WorkflowOrientedApp/WorkflowCentricApplication.h"
 
+#define UE_API PERSONA_API
+
 DECLARE_MULTICAST_DELEGATE_OneParam(FOnAttachToolkit, const TSharedRef<IToolkit>&);
 DECLARE_MULTICAST_DELEGATE_OneParam(FOnDetachToolkit, const TSharedRef<IToolkit>&);
 
 /** Persona asset editor toolkit wrapper, used to auto inject the persona editor mode manager */
-class PERSONA_API FPersonaAssetEditorToolkit : public FWorkflowCentricApplication
+class FPersonaAssetEditorToolkit : public FWorkflowCentricApplication
 {
 public:
 	/** FAssetEditorToolkit interface  */
-	virtual void CreateEditorModeManager() override;
+	UE_API virtual void CreateEditorModeManager() override;
 
 	// IToolkitHost Interface
-	virtual void OnToolkitHostingStarted(const TSharedRef<class IToolkit>& Toolkit) override;
-	virtual void OnToolkitHostingFinished(const TSharedRef<class IToolkit>& Toolkit) override;
+	UE_API virtual void OnToolkitHostingStarted(const TSharedRef<class IToolkit>& Toolkit) override;
+	UE_API virtual void OnToolkitHostingFinished(const TSharedRef<class IToolkit>& Toolkit) override;
 
 	// Returns the currently hosted toolkit. Can be invalid if no toolkit is being hosted.
 	TSharedPtr<IToolkit> GetHostedToolkit() const { return HostedToolkit; }
@@ -24,7 +26,7 @@ public:
 	FOnDetachToolkit& GetOnDetachToolkit() { return OnDetachToolkit; }
 	
 protected:
-	virtual void OnEditorModeIdChanged(const FEditorModeID& ModeChangedID, bool bIsEnteringMode);
+	UE_API virtual void OnEditorModeIdChanged(const FEditorModeID& ModeChangedID, bool bIsEnteringMode);
 	
 protected:
 	FOnAttachToolkit OnAttachToolkit;
@@ -34,3 +36,5 @@ protected:
 	TSharedPtr<IToolkit> HostedToolkit;
 		
 };
+
+#undef UE_API
diff --git a/Engine/Source/Editor/Persona/Public/PersonaModule.h b/Engine/Source/Editor/Persona/Public/PersonaModule.h
--- a/Engine/Source/Editor/Persona/Public/PersonaModule.h
+++ b/Engine/Source/Editor/Persona/Public/PersonaModule.h
@@ -13,6 +13,8 @@
 #include "SPersonaToolBox.h"
 #include "IBlendProfilePickerExtender.h"
 
+#define UE_API PERSONA_API
+
 class IMorphTargetViewer;
 class FBlueprintEditor;
 class IDetailsView;
@@ -36,10 +38,10 @@ struct FAnimAssetFindReplaceConfig;
 extern const FName PersonaAppName;
 
 // Editor mode constants
-struct PERSONA_API FPersonaEditModes
+struct FPersonaEditModes
 {
 	/** Selection/manipulation of bones & sockets */
-	const static FEditorModeID SkeletonSelection;
+	UE_API const static FEditorModeID SkeletonSelection;
 };
 
 DECLARE_DELEGATE_TwoParams(FIsRecordingActive, USkeletalMeshComponent* /*Component*/, bool& /* bIsRecording */);
@@ -657,3 +659,4 @@ private:
 	uint32 NotifyHostAssetVersion = 0;
 };
 
+#undef UE_API
diff --git a/Engine/Source/Editor/Persona/Public/PersonaPreviewSceneController.h b/Engine/Source/Editor/Persona/Public/PersonaPreviewSceneController.h
--- a/Engine/Source/Editor/Persona/Public/PersonaPreviewSceneController.h
+++ b/Engine/Source/Editor/Persona/Public/PersonaPreviewSceneController.h
@@ -9,12 +9,14 @@
 #include "DetailCategoryBuilder.h"
 #include "PersonaPreviewSceneController.generated.h"
 
+#define UE_API PERSONA_API
+
 class UPersonaPreviewSceneDescription;
 class IPersonaPreviewScene;
 
 // Base class for preview scene controller (controls what the preview scene in persona does) 
-UCLASS(Abstract)
-class PERSONA_API UPersonaPreviewSceneController : public UObject
+UCLASS(MinimalAPI, Abstract)
+class UPersonaPreviewSceneController : public UObject
 {
 public:
 	GENERATED_BODY()
@@ -25,5 +27,7 @@ public:
 	virtual void UninitializeView(UPersonaPreviewSceneDescription* SceneDescription, IPersonaPreviewScene* PreviewScene) const PURE_VIRTUAL(UPersonaPreviewSceneController::UninitializeView, );
 
 	//Called when populating the preview scene settings details panel to allow customization of the controllers properties
-	virtual IDetailPropertyRow* AddPreviewControllerPropertyToDetails(const TSharedRef<class IPersonaToolkit>& PersonaToolkit, IDetailLayoutBuilder& DetailBuilder, IDetailCategoryBuilder& Category, const FProperty* Property, const EPropertyLocation::Type PropertyLocation);
+	UE_API virtual IDetailPropertyRow* AddPreviewControllerPropertyToDetails(const TSharedRef<class IPersonaToolkit>& PersonaToolkit, IDetailLayoutBuilder& DetailBuilder, IDetailCategoryBuilder& Category, const FProperty* Property, const EPropertyLocation::Type PropertyLocation);
 };
+
+#undef UE_API
diff --git a/Engine/Source/Editor/Persona/Public/PersonaSelectionProxies.h b/Engine/Source/Editor/Persona/Public/PersonaSelectionProxies.h
--- a/Engine/Source/Editor/Persona/Public/PersonaSelectionProxies.h
+++ b/Engine/Source/Editor/Persona/Public/PersonaSelectionProxies.h
@@ -7,12 +7,14 @@
 #include "HitProxies.h"
 #include "UObject/NameTypes.h"
 
+#define UE_API PERSONA_API
+
 class USkeletalMeshSocket;
 
 
-struct PERSONA_API HPersonaBoneHitProxy : public HHitProxy
+struct HPersonaBoneHitProxy : public HHitProxy
 {
-	DECLARE_HIT_PROXY()
+	DECLARE_HIT_PROXY( UE_API )
 
 	int32 BoneIndex;
 	FName BoneName;
@@ -28,9 +30,9 @@ struct PERSONA_API HPersonaBoneHitProxy : public HHitProxy
 	// End of HHitProxy interface
 };
 
-struct PERSONA_API HPersonaSocketHitProxy : public HHitProxy
+struct HPersonaSocketHitProxy : public HHitProxy
 {
-	DECLARE_HIT_PROXY()
+	DECLARE_HIT_PROXY( UE_API )
 
 	USkeletalMeshSocket* Socket;
 
@@ -43,4 +45,6 @@ struct PERSONA_API HPersonaSocketHitProxy : public HHitProxy
 	// HHitProxy interface
 	virtual EMouseCursor::Type GetMouseCursor() override { return EMouseCursor::Crosshairs; }
 	// End of HHitProxy interface
-};
\ No newline at end of file
+};
+
+#undef UE_API
diff --git a/Engine/Source/Editor/Persona/Public/PersonaTabs.h b/Engine/Source/Editor/Persona/Public/PersonaTabs.h
--- a/Engine/Source/Editor/Persona/Public/PersonaTabs.h
+++ b/Engine/Source/Editor/Persona/Public/PersonaTabs.h
@@ -4,52 +4,56 @@
 
 #include "CoreMinimal.h"
 
+#define UE_API PERSONA_API
+
 // Tab constants
-struct PERSONA_API FPersonaTabs
+struct FPersonaTabs
 {
 	// Selection Details
-	static const FName MorphTargetsID;
-	static const FName AnimCurveViewID;
-	static const FName AnimCurveMetadataEditorID;
-	static const FName SkeletonTreeViewID;
+	static UE_API const FName MorphTargetsID;
+	static UE_API const FName AnimCurveViewID;
+	static UE_API const FName AnimCurveMetadataEditorID;
+	static UE_API const FName SkeletonTreeViewID;
 	// Skeleton Pose manager
-	static const FName RetargetManagerID;
-	static const FName RigManagerID;
+	static UE_API const FName RetargetManagerID;
+	static UE_API const FName RigManagerID;
 	// Skeleton/Sockets
 	// Anim Blueprint Params
 	// Explorer
 	// Class Defaults
-	static const FName AnimBlueprintPreviewEditorID;
-	static const FName AnimBlueprintParentPlayerEditorID;
+	static UE_API const FName AnimBlueprintPreviewEditorID;
+	static UE_API const FName AnimBlueprintParentPlayerEditorID;
 	// Anim Document
-	static const FName ScrubberID;
+	static UE_API const FName ScrubberID;
 	// Toolbar
-	static const FName PreviewViewportID;
-	static const FName PreviewViewport1ID;
-	static const FName PreviewViewport2ID;
-	static const FName PreviewViewport3ID;
-	static const FName AssetBrowserID;
-	static const FName MirrorSetupID;
-	static const FName AnimBlueprintDebugHistoryID;
-	static const FName AnimAssetPropertiesID;
-	static const FName MeshAssetPropertiesID;
-	static const FName PreviewManagerID;
-	static const FName SkeletonAnimNotifiesID;
-	static const FName SkeletonSlotNamesID;
-	static const FName SkeletonSlotGroupNamesID;
-	static const FName CurveNameManagerID;
-	static const FName BlendProfileManagerID;
-	static const FName AnimMontageSectionsID;
-	static const FName PoseWatchManagerID;
+	static UE_API const FName PreviewViewportID;
+	static UE_API const FName PreviewViewport1ID;
+	static UE_API const FName PreviewViewport2ID;
+	static UE_API const FName PreviewViewport3ID;
+	static UE_API const FName AssetBrowserID;
+	static UE_API const FName MirrorSetupID;
+	static UE_API const FName AnimBlueprintDebugHistoryID;
+	static UE_API const FName AnimAssetPropertiesID;
+	static UE_API const FName MeshAssetPropertiesID;
+	static UE_API const FName PreviewManagerID;
+	static UE_API const FName SkeletonAnimNotifiesID;
+	static UE_API const FName SkeletonSlotNamesID;
+	static UE_API const FName SkeletonSlotGroupNamesID;
+	static UE_API const FName CurveNameManagerID;
+	static UE_API const FName BlendProfileManagerID;
+	static UE_API const FName AnimMontageSectionsID;
+	static UE_API const FName PoseWatchManagerID;
 
 	// Advanced Preview Scene
-	static const FName AdvancedPreviewSceneSettingsID;
-	static const FName DetailsID;
-	static const FName FindReplaceID;
+	static UE_API const FName AdvancedPreviewSceneSettingsID;
+	static UE_API const FName DetailsID;
+	static UE_API const FName FindReplaceID;
 
-	static const FName ToolboxID;
+	static UE_API const FName ToolboxID;
 
-	static const FName AnimAttributeViewID;
+	static UE_API const FName AnimAttributeViewID;
 private:
 	FPersonaTabs() {}
-};
\ No newline at end of file
+};
+
+#undef UE_API
diff --git a/Engine/Source/Editor/Persona/Public/PersonaToolMenuContext.h b/Engine/Source/Editor/Persona/Public/PersonaToolMenuContext.h
--- a/Engine/Source/Editor/Persona/Public/PersonaToolMenuContext.h
+++ b/Engine/Source/Editor/Persona/Public/PersonaToolMenuContext.h
@@ -11,6 +11,8 @@
 
 #include "PersonaToolMenuContext.generated.h"
 
+#define UE_API PERSONA_API
+
 class IPersonaToolkit;
 class UAnimBlueprint;
 class UAnimationAsset;
@@ -19,40 +21,42 @@ class USkeletalMesh;
 class USkeleton;
 struct FFrame;
 
-UCLASS(BlueprintType)
-class PERSONA_API UPersonaToolMenuContext : public UObject
+UCLASS(MinimalAPI, BlueprintType)
+class UPersonaToolMenuContext : public UObject
 {
 	GENERATED_BODY()
 public:
 	
 	/** Get the skeleton that we are editing */
 	UFUNCTION(BlueprintCallable, Category = PersonaEditorExtensions)
-	USkeleton* GetSkeleton() const;
+	UE_API USkeleton* GetSkeleton() const;
 
 	/** Get the preview component that we are using */
 	UFUNCTION(BlueprintCallable, Category = PersonaEditorExtensions)
-	UDebugSkelMeshComponent* GetPreviewMeshComponent() const;
+	UE_API UDebugSkelMeshComponent* GetPreviewMeshComponent() const;
 
 	/** Get the skeletal mesh that we are editing */
 	UFUNCTION(BlueprintCallable, Category = PersonaEditorExtensions)
-	USkeletalMesh* GetMesh() const;
+	UE_API USkeletalMesh* GetMesh() const;
 
 	/** Get the anim blueprint that we are editing */
 	UFUNCTION(BlueprintCallable, Category = PersonaEditorExtensions)
-	UAnimBlueprint* GetAnimBlueprint() const;
+	UE_API UAnimBlueprint* GetAnimBlueprint() const;
 
 	/** Get the animation asset that we are editing */
 	UFUNCTION(BlueprintCallable, Category = PersonaEditorExtensions)
-	UAnimationAsset* GetAnimationAsset() const;
+	UE_API UAnimationAsset* GetAnimationAsset() const;
 
 	/** Get the persona toolkit */
-	void SetToolkit(TSharedRef<IPersonaToolkit> InToolkit);
+	UE_API void SetToolkit(TSharedRef<IPersonaToolkit> InToolkit);
 
 	/** Get a weak ptr to the persona toolkit */
 	TWeakPtr<IPersonaToolkit> GetToolkit() const { return WeakToolkit; }
 
 protected:
-	bool HasValidToolkit() const;
+	UE_API bool HasValidToolkit() const;
 private:
 	TWeakPtr<IPersonaToolkit> WeakToolkit;
 };
+
+#undef UE_API
diff --git a/Engine/Source/Editor/Persona/Public/PhysicsAssetRenderUtils.h b/Engine/Source/Editor/Persona/Public/PhysicsAssetRenderUtils.h
--- a/Engine/Source/Editor/Persona/Public/PhysicsAssetRenderUtils.h
+++ b/Engine/Source/Editor/Persona/Public/PhysicsAssetRenderUtils.h
@@ -11,6 +11,8 @@
 #include "SceneManagement.h"
 #include "PhysicsAssetRenderUtils.generated.h"
 
+#define UE_API PERSONA_API
+
 enum class EPhysicsAssetEditorCenterOfMassViewMode : uint8;
 enum class EPhysicsAssetEditorConstraintViewMode : uint8;
 enum class EPhysicsAssetEditorCollisionViewMode : uint8;
@@ -31,48 +33,48 @@ class USkeletalMeshComponent;
 	physics assets whilst switching between editor modes.
 */
 USTRUCT()
-struct PERSONA_API FPhysicsAssetRenderSettings
+struct FPhysicsAssetRenderSettings
 {
 	GENERATED_BODY()
 
 public:
 
-	FPhysicsAssetRenderSettings();
-	void InitPhysicsAssetRenderSettings(class UMaterialInterface* InBoneUnselectedMaterial, class UMaterialInterface* InBoneNoCollisionMaterial);
+	UE_API FPhysicsAssetRenderSettings();
+	UE_API void InitPhysicsAssetRenderSettings(class UMaterialInterface* InBoneUnselectedMaterial, class UMaterialInterface* InBoneNoCollisionMaterial);
 
 	/** Accessors/helper methods */
-	bool IsBodyHidden(const int32 BodyIndex) const;
-	bool IsConstraintHidden(const int32 ConstraintIndex) const;
-	bool AreAnyBodiesHidden() const;
-	bool AreAnyConstraintsHidden() const;
-	void HideBody(const int32 BodyIndex);
-	void ShowBody(const int32 BodyIndex);
-	void HideConstraint(const int32 ConstraintIndex);
-	void ShowConstraint(const int32 ConstraintIndex);
-	void ShowAllBodies();
-	void ShowAllConstraints();
-	void ShowAll();
-	void HideAllBodies(const UPhysicsAsset* const PhysicsAsset);
-	void HideAllConstraints(const UPhysicsAsset* const PhysicsAsset);
-	void HideAll(const UPhysicsAsset* const PhysicsAsset);
-	void ToggleShowBody(const int32 BodyIndex);
-	void ToggleShowConstraint(const int32 ConstraintIndex);
-	void ToggleShowAllBodies(const UPhysicsAsset* const PhysicsAsset);
-	void ToggleShowAllConstraints(const UPhysicsAsset* const PhysicsAsset);
-
-	void SetHiddenBodies(const TArray<int32>& InHiddenBodies);
-	void SetHiddenConstraints(const TArray<int32>& InHiddenConstraints);
+	UE_API bool IsBodyHidden(const int32 BodyIndex) const;
+	UE_API bool IsConstraintHidden(const int32 ConstraintIndex) const;
+	UE_API bool AreAnyBodiesHidden() const;
+	UE_API bool AreAnyConstraintsHidden() const;
+	UE_API void HideBody(const int32 BodyIndex);
+	UE_API void ShowBody(const int32 BodyIndex);
+	UE_API void HideConstraint(const int32 ConstraintIndex);
+	UE_API void ShowConstraint(const int32 ConstraintIndex);
+	UE_API void ShowAllBodies();
+	UE_API void ShowAllConstraints();
+	UE_API void ShowAll();
+	UE_API void HideAllBodies(const UPhysicsAsset* const PhysicsAsset);
+	UE_API void HideAllConstraints(const UPhysicsAsset* const PhysicsAsset);
+	UE_API void HideAll(const UPhysicsAsset* const PhysicsAsset);
+	UE_API void ToggleShowBody(const int32 BodyIndex);
+	UE_API void ToggleShowConstraint(const int32 ConstraintIndex);
+	UE_API void ToggleShowAllBodies(const UPhysicsAsset* const PhysicsAsset);
+	UE_API void ToggleShowAllConstraints(const UPhysicsAsset* const PhysicsAsset);
+
+	UE_API void SetHiddenBodies(const TArray<int32>& InHiddenBodies);
+	UE_API void SetHiddenConstraints(const TArray<int32>& InHiddenConstraints);
 
 	/** Returns a set of flags describing which components of the selected constraint's transforms are being manipulated in the view port. */
-	EConstraintTransformComponentFlags GetConstraintViewportManipulationFlags() const;
+	UE_API EConstraintTransformComponentFlags GetConstraintViewportManipulationFlags() const;
 
 	/** Returns true if all the constraint transform components specified by the flags should be displayed as an offset from the default (snapped) transforms. */
-	bool IsDisplayingConstraintTransformComponentRelativeToDefault(const EConstraintTransformComponentFlags ComponentFlags) const;
+	UE_API bool IsDisplayingConstraintTransformComponentRelativeToDefault(const EConstraintTransformComponentFlags ComponentFlags) const;
 
 	/** Set how the constraint transform components specified by the flags should be displayed, in the frame of their associated physics body (false) or as an offset from the default (snapped) transforms (true). */
-	void SetDisplayConstraintTransformComponentRelativeToDefault(const EConstraintTransformComponentFlags ComponentFlags, const bool bShouldDisplayRelativeToDefault);
+	UE_API void SetDisplayConstraintTransformComponentRelativeToDefault(const EConstraintTransformComponentFlags ComponentFlags, const bool bShouldDisplayRelativeToDefault);
 
-	void ResetEditorViewportOptions();
+	UE_API void ResetEditorViewportOptions();
 
 	// Physics Asset Editor Viewport Options
 	UPROPERTY()
@@ -166,33 +168,33 @@ public:
 // UPhysicsAssetRenderUtilities
 
 /** Factory class for FPhysicsAssetRenderSettings. */
-UCLASS(config=EditorPerProjectUserSettings)
-class PERSONA_API UPhysicsAssetRenderUtilities : public UObject
+UCLASS(MinimalAPI, config=EditorPerProjectUserSettings)
+class UPhysicsAssetRenderUtilities : public UObject
 {
 	GENERATED_BODY()
 
 public:
 
-	UPhysicsAssetRenderUtilities();
-	virtual ~UPhysicsAssetRenderUtilities();
+	UE_API UPhysicsAssetRenderUtilities();
+	UE_API virtual ~UPhysicsAssetRenderUtilities();
 
-	static void Initialise();
+	static UE_API void Initialise();
 
 	/** Returns an existing render settings object or creates and returns a new one if none exists. */
-	static FPhysicsAssetRenderSettings* GetSettings(const UPhysicsAsset* InPhysicsAsset);
-	static FPhysicsAssetRenderSettings* GetSettings(const FString& InPhysicsAssetPathName);
-	static FPhysicsAssetRenderSettings* GetSettings(const uint32 InPhysicsAssetPathNameHash);
+	static UE_API FPhysicsAssetRenderSettings* GetSettings(const UPhysicsAsset* InPhysicsAsset);
+	static UE_API FPhysicsAssetRenderSettings* GetSettings(const FString& InPhysicsAssetPathName);
+	static UE_API FPhysicsAssetRenderSettings* GetSettings(const uint32 InPhysicsAssetPathNameHash);
 	
-	static uint32 GetPathNameHash(const UPhysicsAsset* InPhysicsAsset);
-	static uint32 GetPathNameHash(const FString& InPhysicsAssetPathName);
+	static UE_API uint32 GetPathNameHash(const UPhysicsAsset* InPhysicsAsset);
+	static UE_API uint32 GetPathNameHash(const FString& InPhysicsAssetPathName);
 
-	void OnAssetRenamed(FAssetData const& AssetInfo, const FString& InOldPhysicsAssetPathName);
-	void OnAssetRemoved(UObject* Object);
+	UE_API void OnAssetRenamed(FAssetData const& AssetInfo, const FString& InOldPhysicsAssetPathName);
+	UE_API void OnAssetRemoved(UObject* Object);
 
 private:
 
-	void InitialiseImpl();
-	FPhysicsAssetRenderSettings* GetSettingsImpl(const uint32 InPhysicsAssetPathNameHash);
+	UE_API void InitialiseImpl();
+	UE_API FPhysicsAssetRenderSettings* GetSettingsImpl(const uint32 InPhysicsAssetPathNameHash);
 
 	UPROPERTY()
 	TMap< uint32, FPhysicsAssetRenderSettings > IdToSettingsMap;
@@ -264,4 +266,6 @@ public:
 	virtual void SetDisplayConstraintTransformComponentRelativeToDefault(class UPhysicsAsset* const PhysicsAsset, const EConstraintTransformComponentFlags ComponentFlags, const bool bShouldDisplayRelativeToDefault) override;
 };
 
-PERSONA_API bool IsBodyKinematic(const UPhysicsAsset* const PhysicsAsset, const int32 BodyIndex); // TODO - where is the best place for this fn to live ?
\ No newline at end of file
+PERSONA_API bool IsBodyKinematic(const UPhysicsAsset* const PhysicsAsset, const int32 BodyIndex); // TODO - where is the best place for this fn to live ?
+
+#undef UE_API
diff --git a/Engine/Source/Editor/Persona/Public/ReferenceSectionSelectionWidget.h b/Engine/Source/Editor/Persona/Public/ReferenceSectionSelectionWidget.h
--- a/Engine/Source/Editor/Persona/Public/ReferenceSectionSelectionWidget.h
+++ b/Engine/Source/Editor/Persona/Public/ReferenceSectionSelectionWidget.h
@@ -13,6 +13,8 @@
 #include "Widgets/Views/STreeView.h"
 #include "Widgets/Views/SListView.h"
 
+#define UE_API PERSONA_API
+
 class FSkeletalMeshLODModel;
 class IEditableSkeleton;
 class SComboButton;
@@ -21,7 +23,7 @@ DECLARE_DELEGATE_OneParam(FOnSectionSelectionChanged, int32);
 DECLARE_DELEGATE_RetVal_OneParam(int32, FGetSelectedSection, bool& /*bMultipleValues*/);
 DECLARE_DELEGATE_RetVal(const FSkeletalMeshLODModel&, FGetLodModel);
 
-class PERSONA_API SReferenceSectionSelectionWidget : public SCompoundWidget
+class SReferenceSectionSelectionWidget : public SCompoundWidget
 {
 public:
 
@@ -49,21 +51,21 @@ public:
 	 *
 	 * @param	InArgs	The declaration data for this widget
 	 */
-	void Construct( const FArguments& InArgs );
+	UE_API void Construct( const FArguments& InArgs );
 
 private: 
 
 	// Creates the combo button menu when clicked
-	TSharedRef<SWidget> CreateSectionListWidgetMenu();
-	TSharedRef<ITableRow> MakeIntegerDisplayWidget(const TSharedPtr<int32> SectionIndex, const TSharedRef<STableViewBase>& OwnerTable) const;
+	UE_API TSharedRef<SWidget> CreateSectionListWidgetMenu();
+	UE_API TSharedRef<ITableRow> MakeIntegerDisplayWidget(const TSharedPtr<int32> SectionIndex, const TSharedRef<STableViewBase>& OwnerTable) const;
 
 	TSharedPtr<SListView<TSharedPtr<int32>>> SectionListView;
 	// Called when the user selects a bone name
-	void OnSelectionChanged(TSharedPtr<int32> NewSectionIndex, ESelectInfo::Type SelectInfo);
+	UE_API void OnSelectionChanged(TSharedPtr<int32> NewSectionIndex, ESelectInfo::Type SelectInfo);
 
-	const TArray<TSharedPtr<int32>>* GetSections() const;
+	UE_API const TArray<TSharedPtr<int32>>* GetSections() const;
 
-	FText GetCurrentSectionIndex() const;
+	UE_API FText GetCurrentSectionIndex() const;
 
 	// Base combo button 
 	TSharedPtr<SComboButton> SectionPickerButton;
@@ -76,3 +78,5 @@ private:
 	bool bHideChunkedSections = false;
 	mutable TArray<TSharedPtr<int32>> CacheSectionList;
 };
+
+#undef UE_API
diff --git a/Engine/Source/Editor/Persona/Public/SAnimAttributeView.h b/Engine/Source/Editor/Persona/Public/SAnimAttributeView.h
--- a/Engine/Source/Editor/Persona/Public/SAnimAttributeView.h
+++ b/Engine/Source/Editor/Persona/Public/SAnimAttributeView.h
@@ -12,6 +12,8 @@
 #include "Components/SkeletalMeshComponent.h"
 #include "Engine/SkeletalMesh.h"
 
+#define UE_API PERSONA_API
+
 class ITableRow;
 class STableViewBase;
 class IStructureDetailsView;
@@ -95,16 +97,16 @@ private:
 	TWeakPtr<FAnimAttributeEntry> Entry;
 };
 
-class PERSONA_API SAnimAttributeView : public SCompoundWidget
+class SAnimAttributeView : public SCompoundWidget
 {
 private:	
-	static TSharedRef<IStructureDetailsView> CreateValueViewWidget();
+	static UE_API TSharedRef<IStructureDetailsView> CreateValueViewWidget();
 	
-	static TSharedRef<ITableRow> MakeTableRowWidget(
+	static UE_API TSharedRef<ITableRow> MakeTableRowWidget(
 		TSharedPtr<FAnimAttributeEntry> InItem,
 		const TSharedRef<STableViewBase>& InOwnerTable);
 
-	static FName GetSnapshotColumnDisplayName(const TArray<FName>& InSnapshotNames);
+	static UE_API FName GetSnapshotColumnDisplayName(const TArray<FName>& InSnapshotNames);
 
 public:
 	DECLARE_DELEGATE_RetVal_OneParam(FName, FOnGetAttributeSnapshotColumnDisplayName, const TArray<FName>& /** Snapshot Names */)
@@ -119,9 +121,9 @@ public:
 		SLATE_ATTRIBUTE(FText, SnapshotColumnLabelOverride)
 	SLATE_END_ARGS()
 	
-	SAnimAttributeView();
+	UE_API SAnimAttributeView();
 	
-	void Construct(const FArguments& InArgs);
+	UE_API void Construct(const FArguments& InArgs);
 
 	void DisplayNewAttributeContainerSnapshots(
 		const TArray<TTuple<FName, const UE::Anim::FMeshAttributeContainer&>>& InSnapshots, 
@@ -214,7 +216,7 @@ public:
 
 	void ClearListView();
 	
-	virtual void Tick(const FGeometry& AllottedGeometry, const double InCurrentTime, const float InDeltaTime) override;
+	UE_API virtual void Tick(const FGeometry& AllottedGeometry, const double InCurrentTime, const float InDeltaTime) override;
 	
 private:
 
@@ -285,22 +287,22 @@ private:
 	}
 
 	
-	void OnSelectionChanged(TSharedPtr<FAnimAttributeEntry> InEntry, ESelectInfo::Type InSelectType);
+	UE_API void OnSelectionChanged(TSharedPtr<FAnimAttributeEntry> InEntry, ESelectInfo::Type InSelectType);
 
-	void OnFilterTextChanged(const FText& InText);
+	UE_API void OnFilterTextChanged(const FText& InText);
 
-	EColumnSortMode::Type GetSortModeForColumn(FName InColumnId) const;
+	UE_API EColumnSortMode::Type GetSortModeForColumn(FName InColumnId) const;
 	
-	void OnSortAttributeEntries(
+	UE_API void OnSortAttributeEntries(
 		EColumnSortPriority::Type InPriority,
 		const FName& InColumnId,
 		EColumnSortMode::Type InSortMode);
 
-	void ExecuteSort();
+	UE_API void ExecuteSort();
 	
-	void RefreshFilteredAttributeEntries();
+	UE_API void RefreshFilteredAttributeEntries();
 
-	void RefreshValueView();
+	UE_API void RefreshValueView();
 
 private:
 	/** list view */
@@ -376,17 +378,17 @@ inline void SAnimAttributeView::ClearListView()
 	bShouldRefreshListView = true;
 }
 
-class PERSONA_API SAnimAttributeViewer : public SCompoundWidget
+class SAnimAttributeViewer : public SCompoundWidget
 {
 	SLATE_BEGIN_ARGS( SAnimAttributeViewer ) {}
 	SLATE_END_ARGS()
 	
-	void Construct( const FArguments& InArgs, const TSharedRef<class IPersonaPreviewScene>& InPreviewScene);	
+	UE_API void Construct( const FArguments& InArgs, const TSharedRef<class IPersonaPreviewScene>& InPreviewScene);	
 	virtual ~SAnimAttributeViewer() override {}	
-	virtual void Tick(const FGeometry& AllottedGeometry, const double InCurrentTime, const float InDeltaTime) override;
+	UE_API virtual void Tick(const FGeometry& AllottedGeometry, const double InCurrentTime, const float InDeltaTime) override;
 
 protected:
-	UAnimInstance* GetAnimInstance() const;
+	UE_API UAnimInstance* GetAnimInstance() const;
 protected:
 	/** Pointer to the preview scene we are bound to */
 	TWeakPtr<class IPersonaPreviewScene> PreviewScenePtr;
@@ -394,3 +396,5 @@ protected:
 	TSharedPtr<SAnimAttributeView> AttributeView;	
 	TSharedPtr<SPoseWatchPicker> PoseWatchPicker;
 };
+
+#undef UE_API
diff --git a/Engine/Source/Editor/Persona/Public/SBoneMappingBase.h b/Engine/Source/Editor/Persona/Public/SBoneMappingBase.h
--- a/Engine/Source/Editor/Persona/Public/SBoneMappingBase.h
+++ b/Engine/Source/Editor/Persona/Public/SBoneMappingBase.h
@@ -13,6 +13,8 @@
 #include "Widgets/Input/SComboButton.h"
 #include "BoneSelectionWidget.h"
 
+#define UE_API PERSONA_API
+
 class IEditableSkeleton;
 class URig;
 class USkeleton;
@@ -117,7 +119,7 @@ private:
 //////////////////////////////////////////////////////////////////////////
 // SBoneMappingBase
 
-class PERSONA_API SBoneMappingBase : public SCompoundWidget
+class SBoneMappingBase : public SCompoundWidget
 {
 public:
 	SLATE_BEGIN_ARGS( SBoneMappingBase )
@@ -136,7 +138,7 @@ public:
 	* @param InArgs - Arguments passed from Slate
 	*
 	*/
-	void Construct( const FArguments& InArgs, FSimpleMulticastDelegate& InOnPostUndo );
+	UE_API void Construct( const FArguments& InArgs, FSimpleMulticastDelegate& InOnPostUndo );
 
 	/**
 	* Filters the SListView when the user changes the search text box (NameFilterBox)
@@ -144,7 +146,7 @@ public:
 	* @param SearchText - The text the user has typed
 	*
 	*/
-	void OnFilterTextChanged( const FText& SearchText );
+	UE_API void OnFilterTextChanged( const FText& SearchText );
 
 	/**
 	* Filters the SListView when the user hits enter or clears the search box
@@ -154,7 +156,7 @@ public:
 	* @param CommitInfo - Not used
 	*
 	*/
-	void OnFilterTextCommitted( const FText& SearchText, ETextCommit::Type CommitInfo );
+	UE_API void OnFilterTextCommitted( const FText& SearchText, ETextCommit::Type CommitInfo );
 
 	/**
 	* Create a widget for an entry in the tree from an info
@@ -164,12 +166,12 @@ public:
 	*
 	* @return A new Slate widget, containing the UI for this row
 	*/
-	TSharedRef<ITableRow> GenerateBoneMappingRow(TSharedPtr<FDisplayedBoneMappingInfo> InInfo, const TSharedRef<STableViewBase>& OwnerTable);
+	UE_API TSharedRef<ITableRow> GenerateBoneMappingRow(TSharedPtr<FDisplayedBoneMappingInfo> InInfo, const TSharedRef<STableViewBase>& OwnerTable);
 
 	/**
 	* Handler for the delete of retarget source
 	*/
-	void RefreshBoneMappingList();
+	UE_API void RefreshBoneMappingList();
 private:
 
 	/**
@@ -191,13 +193,15 @@ private:
 	FText FilterText;
 
 	/** Delegate for undo/redo transaction **/
-	void PostUndo();
+	UE_API void PostUndo();
 
-	void OnBoneMappingChanged( FName NodeName, FName BoneName );
-	FName GetBoneMapping( FName NodeName );
+	UE_API void OnBoneMappingChanged( FName NodeName, FName BoneName );
+	UE_API FName GetBoneMapping( FName NodeName );
 
 	FGetReferenceSkeleton	OnGetReferenceSkeletonDelegate;
 	FOnBoneMappingChanged	OnBoneMappingChangedDelegate;
 	FOnGetBoneMapping		OnGetBoneMappingDelegate;
 	FOnCreateBoneMapping	OnCreateBoneMappingDelegate;
 };
+
+#undef UE_API
diff --git a/Engine/Source/Editor/Persona/Public/SCreateClothingSettingsPanel.h b/Engine/Source/Editor/Persona/Public/SCreateClothingSettingsPanel.h
--- a/Engine/Source/Editor/Persona/Public/SCreateClothingSettingsPanel.h
+++ b/Engine/Source/Editor/Persona/Public/SCreateClothingSettingsPanel.h
@@ -8,6 +8,8 @@
 #include "Widgets/DeclarativeSyntaxSupport.h"
 #include "Engine/SkeletalMesh.h"
 
+#define UE_API PERSONA_API
+
 DECLARE_DELEGATE_OneParam(FOnCreateClothingRequested, FSkeletalMeshClothBuildParams&);
 
 class FClothCreateSettingsCustomization : public IDetailCustomization
@@ -43,7 +45,7 @@ protected:
 	FSkeletalMeshClothBuildParams* ParamsStruct;
 };
 
-class PERSONA_API SCreateClothingSettingsPanel : public SCompoundWidget
+class SCreateClothingSettingsPanel : public SCompoundWidget
 {
 public:
 	SLATE_BEGIN_ARGS(SCreateClothingSettingsPanel)
@@ -67,7 +69,7 @@ public:
 
 	SLATE_END_ARGS()
 
-	void Construct(const FArguments& InArgs);
+	UE_API void Construct(const FArguments& InArgs);
 
 private:
 
@@ -78,13 +80,15 @@ private:
 	bool bIsSubImport;
 
 	// Create button functionality
-	FText GetCreateButtonTooltip() const;
-	bool CanCreateClothing() const;
+	UE_API FText GetCreateButtonTooltip() const;
+	UE_API bool CanCreateClothing() const;
 
 	// Handlers for panel buttons
-	FReply OnCreateClicked();
+	UE_API FReply OnCreateClicked();
 
 	// Called when the create button is clicked, so external caller can handle the request
 	FOnCreateClothingRequested OnCreateDelegate;
 
 };
+
+#undef UE_API
diff --git a/Engine/Source/Editor/Persona/Public/SPersonaToolBox.h b/Engine/Source/Editor/Persona/Public/SPersonaToolBox.h
--- a/Engine/Source/Editor/Persona/Public/SPersonaToolBox.h
+++ b/Engine/Source/Editor/Persona/Public/SPersonaToolBox.h
@@ -15,6 +15,8 @@
 #include "Templates/SharedPointer.h"
 #include "Toolkits/IToolkit.h"
 
+#define UE_API PERSONA_API
+
 class FModeToolkit;
 class FPersonaAssetEditorToolkit;
 class IToolkit;
@@ -35,32 +37,32 @@ protected:
 };
 
 /** Widget to surface an Edit Mode its widget, similar to LevelEditor approach */
-class PERSONA_API SPersonaToolbox : public SCompoundWidget
+class SPersonaToolbox : public SCompoundWidget
 {
 public:
 	SLATE_BEGIN_ARGS( SPersonaToolbox ) {}
 	SLATE_END_ARGS()
 
-	~SPersonaToolbox();
+	UE_API ~SPersonaToolbox();
 
-	void Construct(const FArguments& InArgs, const TSharedRef<FPersonaAssetEditorToolkit>& InOwningEditor);
+	UE_API void Construct(const FArguments& InArgs, const TSharedRef<FPersonaAssetEditorToolkit>& InOwningEditor);
 
 	// Called to attach a toolkit to this toolbox widget (e.g. when the mode changes).
-	void AttachToolkit(const TSharedRef<IToolkit>& InToolkit);
+	UE_API void AttachToolkit(const TSharedRef<IToolkit>& InToolkit);
 
 	// Called to detach a toolkit from this toolbox widget (e.g. when the mode changes).
-	void DetachToolkit(const TSharedRef<IToolkit>& InToolkit);
+	UE_API void DetachToolkit(const TSharedRef<IToolkit>& InToolkit);
 
 	// Sets the owning tab so that we can update the tab icon based on what the toolkit provides.
-	void SetOwningTab(const TSharedRef<SDockTab>& InOwningTab);
+	UE_API void SetOwningTab(const TSharedRef<SDockTab>& InOwningTab);
 
 private:	
-	void UpdateInlineContent(const TSharedPtr<IToolkit>& Toolkit, TSharedPtr<SWidget> InlineContent);
+	UE_API void UpdateInlineContent(const TSharedPtr<IToolkit>& Toolkit, TSharedPtr<SWidget> InlineContent);
 
-	void UpdatePalette(const TWeakPtr<FModeToolkit>& InWeakToolkit) const;
+	UE_API void UpdatePalette(const TWeakPtr<FModeToolkit>& InWeakToolkit) const;
 
 	// Whether to show the inline content or not.
-	EVisibility GetInlineContentHolderVisibility() const;
+	UE_API EVisibility GetInlineContentHolderVisibility() const;
 
 	// The owning tab where the toolkit's interface is hosted. We need this to update the icons
 	// and title as the toolkit changes.
@@ -79,4 +81,6 @@ private:
 
 	// The global status bar to show the name of the active tool.
 	FStatusBarMessageHandle StatusBarMessageHandle;
-};
\ No newline at end of file
+};
+
+#undef UE_API
diff --git a/Engine/Source/Editor/PinnedCommandList/Public/IPinnedCommandList.h b/Engine/Source/Editor/PinnedCommandList/Public/IPinnedCommandList.h
--- a/Engine/Source/Editor/PinnedCommandList/Public/IPinnedCommandList.h
+++ b/Engine/Source/Editor/PinnedCommandList/Public/IPinnedCommandList.h
@@ -32,7 +32,7 @@ class ISlateStyle;
  *    to be called when widgets are interacted with to automatically add the custom widget to 
  *    the pinned commands list.
  */
-class PINNEDCOMMANDLIST_API IPinnedCommandList : public SCompoundWidget
+class IPinnedCommandList : public SCompoundWidget
 {
 public:
 	/** Delegate used to generate custom widgets */
@@ -98,4 +98,4 @@ public:
 	 * @param	InStyleName		The name of the menu button style to use, e.g. "Menu"
 	 */
 	virtual void SetStyle(const ISlateStyle* InStyleSet, const FName& InStyleName) = 0;
-};
\ No newline at end of file
+};
diff --git a/Engine/Source/Editor/PinnedCommandList/Public/UICommandList_Pinnable.h b/Engine/Source/Editor/PinnedCommandList/Public/UICommandList_Pinnable.h
--- a/Engine/Source/Editor/PinnedCommandList/Public/UICommandList_Pinnable.h
+++ b/Engine/Source/Editor/PinnedCommandList/Public/UICommandList_Pinnable.h
@@ -9,6 +9,8 @@
 #include "Templates/SharedPointer.h"
 #include "UObject/NameTypes.h"
 
+#define UE_API PINNEDCOMMANDLIST_API
+
 class FUICommandInfo;
 struct FUIAction;
 
@@ -18,10 +20,10 @@ struct FUIAction;
  * - Indexing when mapping to provide a sort key
  * - Optional grouping of actions
  */
-class PINNEDCOMMANDLIST_API FUICommandList_Pinnable : public FUICommandList
+class FUICommandList_Pinnable : public FUICommandList
 {
 public:
-	FUICommandList_Pinnable();
+	UE_API FUICommandList_Pinnable();
 
 	using FUICommandList::MapAction;
 
@@ -32,11 +34,11 @@ public:
 	DECLARE_MULTICAST_DELEGATE_TwoParams(FOnCustomWidgetInteraction, FName /*InCustomWidgetIdentifier*/, const TSharedRef<const FUICommandList_Pinnable>& /*InUICommandList*/);
 
 	/** FUICommandList interface */
-	virtual bool ExecuteAction(const TSharedRef< const FUICommandInfo > InUICommandInfo) const override;
-	virtual void MapAction( const TSharedPtr< const FUICommandInfo > InUICommandInfo, const FUIAction& InUIAction ) override;
+	UE_API virtual bool ExecuteAction(const TSharedRef< const FUICommandInfo > InUICommandInfo) const override;
+	UE_API virtual void MapAction( const TSharedPtr< const FUICommandInfo > InUICommandInfo, const FUIAction& InUIAction ) override;
 
 	/** Mark that a custom widget was interacted with */
-	void WidgetInteraction(FName InCustomWidgetIdentifier) const;
+	UE_API void WidgetInteraction(FName InCustomWidgetIdentifier) const;
 
 	/** Delegate called when an action is executed */
 	FOnExecuteAction& OnExecuteAction() { return OnExecuteActionDelegate; }
@@ -49,23 +51,23 @@ public:
 	 * @param	InCommandName		The command name to check
 	 * @return the mapped index, or INDEX_NONE if the command is not mapped to this command list
 	 */
-	int32 GetMappedCommandIndex( FName InCommandName ) const;
+	UE_API int32 GetMappedCommandIndex( FName InCommandName ) const;
 
 	/** 
 	 * Get the group this command was mapped into.
 	 * @param	InUICommandInfo		The command to check
 	 * @return the command's group, or NAME_None if the command is not grouped
 	 */
-	FName GetMappedCommandGroup( FName InCommandName ) const;
+	UE_API FName GetMappedCommandGroup( FName InCommandName ) const;
 
 	/** 
 	 * Start a group scope. All calls to MapAction will be grouped into the specified group.
 	 * Should be paired with EndGroup()
 	 */
-	void BeginGroup(FName InGroupName);
+	UE_API void BeginGroup(FName InGroupName);
 
 	/** End a group scope. Should be paired with BeginGroup(). */
-	void EndGroup();
+	UE_API void EndGroup();
 
 protected:
 	/** Index used to provide a sort key for actions when mapped */
@@ -86,3 +88,5 @@ protected:
 	/** Map providing a group for each command that was mapped via this command list */
 	TMap<FName, FName> CommandGroupMap;
 };
+
+#undef UE_API
diff --git a/Engine/Source/Editor/PropertyEditor/Public/AsyncDetailViewDiff.h b/Engine/Source/Editor/PropertyEditor/Public/AsyncDetailViewDiff.h
--- a/Engine/Source/Editor/PropertyEditor/Public/AsyncDetailViewDiff.h
+++ b/Engine/Source/Editor/PropertyEditor/Public/AsyncDetailViewDiff.h
@@ -4,8 +4,10 @@
 #include "DiffUtils.h"
 #include "IDetailsView.h"
 
+#define UE_API PROPERTYEDITOR_API
+
 template<>
-class PROPERTYEDITOR_API TTreeDiffSpecification<TWeakPtr<FDetailTreeNode>>
+class TTreeDiffSpecification<TWeakPtr<FDetailTreeNode>>
 {
 public:
 	virtual ~TTreeDiffSpecification() = default;
@@ -14,9 +16,9 @@ public:
 	 * @param TreeNodeA node from the first user provided tree (guaranteed not to be null)
 	 * @param TreeNodeB node from the second user provided tree (guaranteed not to be null)
 	 */
-	virtual bool AreValuesEqual(const TWeakPtr<FDetailTreeNode>& TreeNodeA, const TWeakPtr<FDetailTreeNode>& TreeNodeB) const;
+	UE_API virtual bool AreValuesEqual(const TWeakPtr<FDetailTreeNode>& TreeNodeA, const TWeakPtr<FDetailTreeNode>& TreeNodeB) const;
 	UE_DEPRECATED(5.6, "Output parameters wasn't generic and is no longer used.")
-	virtual bool AreValuesEqual(const TWeakPtr<FDetailTreeNode>& TreeNodeA, const TWeakPtr<FDetailTreeNode>& TreeNodeB, TArray<FPropertySoftPath>*) const final;
+	UE_API virtual bool AreValuesEqual(const TWeakPtr<FDetailTreeNode>& TreeNodeA, const TWeakPtr<FDetailTreeNode>& TreeNodeB, TArray<FPropertySoftPath>*) const final;
 
 	/**
 	 * determine whether two nodes occupy the same space in their trees
@@ -25,52 +27,54 @@ public:
 	 * @param TreeNodeA node from the first user provided tree (guaranteed not to be null)
 	 * @param TreeNodeB node from the second user provided tree (guaranteed not to be null)
 	 */
-	virtual bool AreMatching(const TWeakPtr<FDetailTreeNode>& TreeNodeA, const TWeakPtr<FDetailTreeNode>& TreeNodeB) const;
+	UE_API virtual bool AreMatching(const TWeakPtr<FDetailTreeNode>& TreeNodeA, const TWeakPtr<FDetailTreeNode>& TreeNodeB) const;
 	UE_DEPRECATED(5.6, "Output parameters wasn't generic and is no longer used.")
-	virtual bool AreMatching(const TWeakPtr<FDetailTreeNode>& TreeNodeA, const TWeakPtr<FDetailTreeNode>& TreeNodeB, TArray<FPropertySoftPath>*) const final;
+	UE_API virtual bool AreMatching(const TWeakPtr<FDetailTreeNode>& TreeNodeA, const TWeakPtr<FDetailTreeNode>& TreeNodeB, TArray<FPropertySoftPath>*) const final;
 
 	/**
 	 * retrieves an array of children nodes from the parent node
 	 * @param InParent node from one of the two user provided trees (guaranteed not to be null)
 	 * @param[out] OutChildren to be filled with the children of parent
 	 */
-	virtual void GetChildren(const TWeakPtr<FDetailTreeNode>& InParent, TArray<TWeakPtr<FDetailTreeNode>>& OutChildren) const;
+	UE_API virtual void GetChildren(const TWeakPtr<FDetailTreeNode>& InParent, TArray<TWeakPtr<FDetailTreeNode>>& OutChildren) const;
 
 	/**
 	 * return true for nodes that match using AreValuesEqual first, and pair up by position second
 	 * this is useful for arrays since we often want to keep elements with the same data paired while diffing other elements in order
 	 * @param TreeNode node from one of the two user provided trees (guaranteed not to be null)
 	*/
-	virtual bool ShouldMatchByValue(const TWeakPtr<FDetailTreeNode>& TreeNode) const;
+	UE_API virtual bool ShouldMatchByValue(const TWeakPtr<FDetailTreeNode>& TreeNode) const;
 
 	/**
 	 * return true if TreeNode is considered equal when all it's children are equal.
 	 * This avoids an unnecessary call to AreValuesEqual
 	*/
-	virtual bool ShouldInheritEqualFromChildren(const TWeakPtr<FDetailTreeNode>& TreeNodeA, const TWeakPtr<FDetailTreeNode>& TreeNodeB) const;
+	UE_API virtual bool ShouldInheritEqualFromChildren(const TWeakPtr<FDetailTreeNode>& TreeNodeA, const TWeakPtr<FDetailTreeNode>& TreeNodeB) const;
 };
 
 // asynchronously update a difference tree with changes to details views
 // note: users are expected to call TAsyncTreeDifferences::Tick or TAsyncTreeDifferences::FlushQueue to get accurate data
-class PROPERTYEDITOR_API FAsyncDetailViewDiff : public TAsyncTreeDifferences<TWeakPtr<FDetailTreeNode>>
+class FAsyncDetailViewDiff : public TAsyncTreeDifferences<TWeakPtr<FDetailTreeNode>>
 {
 public:
-	FAsyncDetailViewDiff(TSharedRef<IDetailsView> LeftView, TSharedRef<IDetailsView> RightView);
+	UE_API FAsyncDetailViewDiff(TSharedRef<IDetailsView> LeftView, TSharedRef<IDetailsView> RightView);
 
 	// Note: by default this list is not necessarily accurate or exhaustive.
 	// if you need a perfectly accurate list, call FlushQueue() first
-	void GetPropertyDifferences(TArray<FSingleObjectDiffEntry>& OutDiffEntries) const;
+	UE_API void GetPropertyDifferences(TArray<FSingleObjectDiffEntry>& OutDiffEntries) const;
 
 	// execute a method for each node associated with a row in at least one details view
 	// method provides current left and right row number
 	// returns the number of rows in each details view
-	TPair<int32, int32> ForEachRow(const TFunction<ETreeTraverseControl(const TUniquePtr<DiffNodeType>&, int32, int32)>& Method) const;
+	UE_API TPair<int32, int32> ForEachRow(const TFunction<ETreeTraverseControl(const TUniquePtr<DiffNodeType>&, int32, int32)>& Method) const;
 
 	// generates scroll data used by LinkableScrollbar
-	TArray<FVector2f> GenerateScrollSyncRate() const;
+	UE_API TArray<FVector2f> GenerateScrollSyncRate() const;
 
 private:
-	static TAttribute<TArray<TWeakPtr<FDetailTreeNode>>> RootNodesAttribute(TWeakPtr<IDetailsView> DetailsView);
+	static UE_API TAttribute<TArray<TWeakPtr<FDetailTreeNode>>> RootNodesAttribute(TWeakPtr<IDetailsView> DetailsView);
 	TWeakPtr<IDetailsView> LeftView;
 	TWeakPtr<IDetailsView> RightView;
 };
+
+#undef UE_API
diff --git a/Engine/Source/Editor/PropertyEditor/Public/DetailRowMenuContext.h b/Engine/Source/Editor/PropertyEditor/Public/DetailRowMenuContext.h
--- a/Engine/Source/Editor/PropertyEditor/Public/DetailRowMenuContext.h
+++ b/Engine/Source/Editor/PropertyEditor/Public/DetailRowMenuContext.h
@@ -9,8 +9,8 @@
 class IDetailsView;
 class IPropertyHandle;
 
-UCLASS()
-class PROPERTYEDITOR_API UDetailRowMenuContext : public UObject
+UCLASS(MinimalAPI)
+class UDetailRowMenuContext : public UObject
 {
 	GENERATED_BODY()
 
diff --git a/Engine/Source/Editor/PropertyEditor/Public/DetailsDisplayManager.h b/Engine/Source/Editor/PropertyEditor/Public/DetailsDisplayManager.h
--- a/Engine/Source/Editor/PropertyEditor/Public/DetailsDisplayManager.h
+++ b/Engine/Source/Editor/PropertyEditor/Public/DetailsDisplayManager.h
@@ -97,7 +97,7 @@ public:
 	/**
 	 * Returns whether or not this display manager provides a widget to overlay over the detail tree.  This can be used to facilitate interactions like drag and drop.
 	 */
-	PROPERTYEDITOR_API virtual bool SupportsDetailTreeOverlay() const { return false; }
+	virtual bool SupportsDetailTreeOverlay() const { return false; }
 
 	/**
 	 * Returns a widget which will be put in an overlay slot which is over the detail tree portion of the widget.
@@ -195,4 +195,4 @@ struct FConstructPropertyUpdatedWidgetBuilderArgs
 	FName Category = NAME_None;
 	FExecuteAction InvalidateCachedState;
 	TArray<TWeakObjectPtr<UObject>>* Objects = nullptr;
-};
\ No newline at end of file
+};
diff --git a/Engine/Source/Editor/PropertyEditor/Public/DetailsNameWidgetOverrideCustomization.h b/Engine/Source/Editor/PropertyEditor/Public/DetailsNameWidgetOverrideCustomization.h
--- a/Engine/Source/Editor/PropertyEditor/Public/DetailsNameWidgetOverrideCustomization.h
+++ b/Engine/Source/Editor/PropertyEditor/Public/DetailsNameWidgetOverrideCustomization.h
@@ -21,7 +21,7 @@ public:
 	/**
 	 * destructor
 	 */
-	PROPERTYEDITOR_API virtual ~FDetailsNameWidgetOverrideCustomization() = default;
+	virtual ~FDetailsNameWidgetOverrideCustomization() = default;
 
 	/**
 	 * @param InnerNameContent the unaltered content of the Name widget. It shows what would
@@ -31,4 +31,4 @@ public:
 	 * the default Name widget for the current property path
 	 */
 	 PROPERTYEDITOR_API virtual TSharedRef<SWidget> CustomizeName( TSharedRef<SWidget> InnerNameContent,  FPropertyPath& Path);
-};
\ No newline at end of file
+};
diff --git a/Engine/Source/Editor/PropertyEditor/Public/MaterialList.h b/Engine/Source/Editor/PropertyEditor/Public/MaterialList.h
--- a/Engine/Source/Editor/PropertyEditor/Public/MaterialList.h
+++ b/Engine/Source/Editor/PropertyEditor/Public/MaterialList.h
@@ -7,6 +7,8 @@
 #include "DetailWidgetRow.h"
 #include "Materials/MaterialInterface.h"
 
+#define UE_API PROPERTYEDITOR_API
+
 class FMaterialItemView;
 class FMaterialListBuilder;
 class IDetailChildrenBuilder;
@@ -149,7 +151,7 @@ struct FMaterialListItem
 /**
  * A view of a single item in an FMaterialList
  */
-class PROPERTYEDITOR_API FMaterialItemView : public TSharedFromThis<FMaterialItemView>
+class FMaterialItemView : public TSharedFromThis<FMaterialItemView>
 {
 public:
 	/** Virtual destructor */
@@ -161,7 +163,7 @@ public:
 	* @param Material				The material to view
 	* @param InOnMaterialChanged	Delegate for when the material changes
 	*/
-	static TSharedRef<FMaterialItemView> Create(
+	static UE_API TSharedRef<FMaterialItemView> Create(
 		const FMaterialListItem& Material, 
 		FOnMaterialChanged InOnMaterialChanged,
 		FOnGenerateWidgetsForMaterial InOnGenerateNameWidgetsForMaterial, 
@@ -174,12 +176,12 @@ public:
 	/**
 	 * Create a Name Content for Material Item Widget
 	 */
-	virtual TSharedRef<SWidget> CreateNameContent();
+	UE_API virtual TSharedRef<SWidget> CreateNameContent();
 
 	/**
 	 * Create a Value Content for Material Item Widget
 	 */
-	virtual TSharedRef<SWidget> CreateValueContent(IDetailLayoutBuilder& InDetailBuilder, const TArray<FAssetData>& OwnerAssetDataArray, UActorComponent* InActorComponent);
+	UE_API virtual TSharedRef<SWidget> CreateValueContent(IDetailLayoutBuilder& InDetailBuilder, const TArray<FAssetData>& OwnerAssetDataArray, UActorComponent* InActorComponent);
 
 	/**
 	 * Get Material item in a material list slot
@@ -192,20 +194,20 @@ public:
 	 * @param NewMaterial			The new material apply to slot
 	 * @param bReplaceAll			Whether it should be applied to all slots
 	 */
-	void ReplaceMaterial( UMaterialInterface* NewMaterial, bool bReplaceAll = false );
+	UE_API void ReplaceMaterial( UMaterialInterface* NewMaterial, bool bReplaceAll = false );
 
 	/**
 	 * Called to get the visibility of the reset to base button
 	 */
-	bool GetResetToBaseVisibility() const;
+	UE_API bool GetResetToBaseVisibility() const;
 
 	/**
 	 * Called when reset to base is clicked
 	 */
-	void OnResetToBaseClicked();
+	UE_API void OnResetToBaseClicked();
 
 private:
-	FMaterialItemView(const FMaterialListItem& InMaterial, 
+	UE_API FMaterialItemView(const FMaterialListItem& InMaterial, 
 					FOnMaterialChanged& InOnMaterialChanged, 
 					FOnGenerateWidgetsForMaterial& InOnGenerateNameWidgets, 
 					FOnGenerateWidgetsForMaterial& InOnGenerateMaterialWidgets, 
@@ -217,37 +219,37 @@ private:
 	/**
 	 * Handler when the material asset has been changed
 	 */
-	void OnSetObject( const FAssetData& AssetData );
+	UE_API void OnSetObject( const FAssetData& AssetData );
 
 	/**
 	 * @return Whether or not the textures menu is enabled
 	 */
-	bool IsTexturesMenuEnabled() const;
+	UE_API bool IsTexturesMenuEnabled() const;
 
 	/**
 	 * Generate list of the textures for material selection list
 	 */
-	TSharedRef<SWidget> OnGetTexturesMenuForMaterial();
+	UE_API TSharedRef<SWidget> OnGetTexturesMenuForMaterial();
 
 	/**
 	 * Create button for browse to nanite override material
 	 */
-	TSharedRef<SWidget> MakeBrowseNaniteOverrideMaterialButton() const;
+	UE_API TSharedRef<SWidget> MakeBrowseNaniteOverrideMaterialButton() const;
 
 	/**
 	 * @return The tool tip for the textures menuuuuu
 	 */
-	FText GetTexturesMenuToolTipText() const;
+	UE_API FText GetTexturesMenuToolTipText() const;
 
 	/**
 	 * On Get object path handler
 	 */
-	FString OnGetObjectPath() const;
+	UE_API FString OnGetObjectPath() const;
 
 	/**
 	 * Finds the asset in the content browser
 	 */
-	void GoToAssetInContentBrowser( TWeakObjectPtr<UObject> Object );
+	UE_API void GoToAssetInContentBrowser( TWeakObjectPtr<UObject> Object );
 
 	/**
 	 * Get Property Row Extenssion Arguments.
@@ -257,7 +259,7 @@ private:
 	 * @param InCurrentComponent The current component using the current material.
 	 * @return Generate Global Row Extension Args
 	 */
-	struct FOnGenerateGlobalRowExtensionArgs GetGlobalRowExtensionArgs(IDetailLayoutBuilder& InDetailBuilder, UActorComponent* InCurrentComponent) const;
+	UE_API struct FOnGenerateGlobalRowExtensionArgs GetGlobalRowExtensionArgs(IDetailLayoutBuilder& InDetailBuilder, UActorComponent* InCurrentComponent) const;
 
 	/**
 	 * Get row extension widget
@@ -266,7 +268,7 @@ private:
 	 * @param InCurrentComponent The current component using the current material.
 	 * @return Pointer to widget
 	 */
-	TSharedPtr<SWidget> GetGlobalRowExtensionWidget(IDetailLayoutBuilder& InDetailBuilder, UActorComponent* InCurrentComponent) const;
+	UE_API TSharedPtr<SWidget> GetGlobalRowExtensionWidget(IDetailLayoutBuilder& InDetailBuilder, UActorComponent* InCurrentComponent) const;
 
 private:
 	/** The Material item in a material list slot */
@@ -387,3 +389,5 @@ private:
 	/** The mesh asset that owns these materials */
 	TArray<FAssetData> OwnerAssetDataArray;
 };
+
+#undef UE_API
diff --git a/Engine/Source/Editor/PropertyEditor/Public/PropertyEditorClipboard.h b/Engine/Source/Editor/PropertyEditor/Public/PropertyEditorClipboard.h
--- a/Engine/Source/Editor/PropertyEditor/Public/PropertyEditorClipboard.h
+++ b/Engine/Source/Editor/PropertyEditor/Public/PropertyEditorClipboard.h
@@ -1,4 +1,4 @@
-﻿// Copyright Epic Games, Inc. All Rights Reserved.
+// Copyright Epic Games, Inc. All Rights Reserved.
 
 #pragma once
 
@@ -9,38 +9,42 @@
 #include "Templates/SharedPointerFwd.h"
 #include "UObject/NameTypes.h"
 
+#define UE_API PROPERTYEDITOR_API
+
 /**
  * Specializes clipboard handling to allow tagged entries,
  * where you can partition the clipboard contents and look them up by tag (name).
  * The clipboard is persisted using both the Platform implementation, and as a stored virtual clipboard.
  */
-class PROPERTYEDITOR_API FPropertyEditorClipboard
+class FPropertyEditorClipboard
 {
 public:
 	/** Perform a copy. */
-	static void ClipboardCopy(const TCHAR* Str);
+	static UE_API void ClipboardCopy(const TCHAR* Str);
 
 	/** Perform a copy, optionally specifying the tag. */
-	static void ClipboardCopy(const TCHAR* Str, FName Tag);
+	static UE_API void ClipboardCopy(const TCHAR* Str, FName Tag);
 
 	/** Perform a copy, supplying a function that maps the result to multiple tags. */
-	static void ClipboardCopy(TUniqueFunction<void(TMap<FName, FString>&)>&& TagMappingFunc);
+	static UE_API void ClipboardCopy(TUniqueFunction<void(TMap<FName, FString>&)>&& TagMappingFunc);
 
 	/** Append the given values to the current clipboard entry. If none exists, it will create it. */
-	static void ClipboardAppend(TUniqueFunction<void(TMap<FName, FString>&)>&& TagMappingFunc);
+	static UE_API void ClipboardAppend(TUniqueFunction<void(TMap<FName, FString>&)>&& TagMappingFunc);
 
 	/** Gets the clipboard contents. */
-	static bool ClipboardPaste(FString& Dest);
+	static UE_API bool ClipboardPaste(FString& Dest);
 	
 	/**
 	 * Will return true if the tag is specified and exists. 
 	 * If the tag doesn't exist, Dest will be populated with the untagged clipboard contents.  
 	 */
-	static bool ClipboardPaste(FString& Dest, FName Tag);
+	static UE_API bool ClipboardPaste(FString& Dest, FName Tag);
 
 private:
-	static TSharedRef<FPropertyEditorClipboard> Get();
+	static UE_API TSharedRef<FPropertyEditorClipboard> Get();
 
 	/** Item per-tag, cleared when a new copy operation occurs. */
 	TMap<FName, FString> ClipboardContents;
 };
+
+#undef UE_API
diff --git a/Engine/Source/Editor/PropertyEditor/Public/PropertyPermissionList.h b/Engine/Source/Editor/PropertyEditor/Public/PropertyPermissionList.h
--- a/Engine/Source/Editor/PropertyEditor/Public/PropertyPermissionList.h
+++ b/Engine/Source/Editor/PropertyEditor/Public/PropertyPermissionList.h
@@ -8,6 +8,8 @@
 #include "UObject/WeakObjectPtr.h"
 #include "Misc/NamePermissionList.h"
 
+#define UE_API PROPERTYEDITOR_API
+
 class FDetailTreeNode;
 
 DECLARE_LOG_CATEGORY_EXTERN(LogPropertyEditorPermissionList, Log, All);
@@ -49,35 +51,35 @@ struct FPropertyPermissionListEntry
 	EPropertyPermissionListRules Rules = EPropertyPermissionListRules::UseExistingPermissionList;
 };
 
-class PROPERTYEDITOR_API FPropertyPermissionList
+class FPropertyPermissionList
 {
 public:
 	/** Add a set of rules for a specific base UStruct to determine which properties are visible in all details panels */
 	UE_DEPRECATED(5.5, "Call AddPermissionList with additional required arguments instead.")
-	void AddPermissionList(TSoftObjectPtr<const UStruct> Struct, const FNamePermissionList& PermissionList, const EPropertyPermissionListRules Rules = EPropertyPermissionListRules::UseExistingPermissionList);
+	UE_API void AddPermissionList(TSoftObjectPtr<const UStruct> Struct, const FNamePermissionList& PermissionList, const EPropertyPermissionListRules Rules = EPropertyPermissionListRules::UseExistingPermissionList);
 	/** Add a set of rules for a specific base UStruct to determine which properties are visible in all details panels */
-	void AddPermissionList(TSoftObjectPtr<const UStruct> Struct, const FNamePermissionList& PermissionList, const EPropertyPermissionListRules Rules, const TConstArrayView<FName> InAdditionalOwnerNames);
+	UE_API void AddPermissionList(TSoftObjectPtr<const UStruct> Struct, const FNamePermissionList& PermissionList, const EPropertyPermissionListRules Rules, const TConstArrayView<FName> InAdditionalOwnerNames);
 	/** Remove a set of rules for a specific base UStruct to determine which properties are visible in all details panels */
-	void RemovePermissionList(TSoftObjectPtr<const UStruct> Struct);
+	UE_API void RemovePermissionList(TSoftObjectPtr<const UStruct> Struct);
 	/** Remove all rules */
-	void ClearPermissionList();
+	UE_API void ClearPermissionList();
 
 	/** Unregister an owner from all permission lists currently stored */
-	void UnregisterOwner(const FName Owner);
+	UE_API void UnregisterOwner(const FName Owner);
 
 	/** Add a specific property to a UStruct's AllowList */
-	void AddToAllowList(TSoftObjectPtr<const UStruct> Struct, const FName PropertyName, const FName Owner);
+	UE_API void AddToAllowList(TSoftObjectPtr<const UStruct> Struct, const FName PropertyName, const FName Owner);
 	/** Add a list of properties to a UStruct's AllowList */
-	void AddToAllowList(TSoftObjectPtr<const UStruct> Struct, const TArray<FName>& PropertyNames, const FName Owner);
+	UE_API void AddToAllowList(TSoftObjectPtr<const UStruct> Struct, const TArray<FName>& PropertyNames, const FName Owner);
 	/** Remove a specific property from a UStruct's AllowList */
-	void RemoveFromAllowList(TSoftObjectPtr<const UStruct> Struct, const FName PropertyName, const FName Owner);
+	UE_API void RemoveFromAllowList(TSoftObjectPtr<const UStruct> Struct, const FName PropertyName, const FName Owner);
 	/** Add a specific property to a UStruct's DenyList */
-	void AddToDenyList(TSoftObjectPtr<const UStruct> Struct, const FName PropertyName, const FName Owner);
+	UE_API void AddToDenyList(TSoftObjectPtr<const UStruct> Struct, const FName PropertyName, const FName Owner);
 	/** Remove a specific property from a UStruct's DenyList */
-	void RemoveFromDenyList(TSoftObjectPtr<const UStruct> Struct, const FName PropertyName, const FName Owner);
+	UE_API void RemoveFromDenyList(TSoftObjectPtr<const UStruct> Struct, const FName PropertyName, const FName Owner);
 	/** Registers a delegate that is capable of handling more advanced use cases, such as details panel rows whose identifiers are created dynamically. */
-	void RegisterCustomDetailsTreeNodePermissionDelegate(TSoftObjectPtr<const UStruct> Struct, FCustomDetailTreeNodePermissionDelegate InDelegate);
-	void UnregisterCustomDetailsTreeNodePermissionDelegate(TSoftObjectPtr<const UStruct> Struct);
+	UE_API void RegisterCustomDetailsTreeNodePermissionDelegate(TSoftObjectPtr<const UStruct> Struct, FCustomDetailTreeNodePermissionDelegate InDelegate);
+	UE_API void UnregisterCustomDetailsTreeNodePermissionDelegate(TSoftObjectPtr<const UStruct> Struct);
 
 	/** When the PermissionList or DenyList for any struct was added to or removed from. */
 	FPermissionListUpdated PermissionListUpdatedDelegate;
@@ -92,45 +94,45 @@ public:
 	 * Optionally the change can be announced through a broadcast, which is typically recommended, but may not be needed
 	 * if the permission list is only temporary enabled, e.g. during saving.
 	 */
-	void SetEnabled(bool bEnable, bool bAnnounce = true);
+	UE_API void SetEnabled(bool bEnable, bool bAnnounce = true);
 
 	/** Checks if the provided struct has any entries in the allow, deny and deny all lists and returns true in that case. */
-	bool HasFiltering(const UStruct* ObjectStruct) const;
+	UE_API bool HasFiltering(const UStruct* ObjectStruct) const;
 
 	/**
 	 * Checks if a property passes the PermissionList/DenyList filtering specified by PropertyPermissionLists
 	 * This should be relatively fast as it maintains a flattened cache of all inherited PermissionLists for every UStruct (which is generated lazily).
 	 */
-	bool DoesPropertyPassFilter(const UStruct* ObjectStruct, FName PropertyName) const;
+	UE_API bool DoesPropertyPassFilter(const UStruct* ObjectStruct, FName PropertyName) const;
 	/** Checks if a given detail tree node passes the filters. This is more flexible than the above function but should only be considered for advanced use cases,
 	 * such as dynamically generated rows from customizations that don't fit the structure of simply checking against an FName. */
-	bool DoesDetailTreeNodePassFilter(const UStruct* ObjectStruct, TSharedRef<FDetailTreeNode> DetailTreeNode);
+	UE_API bool DoesDetailTreeNodePassFilter(const UStruct* ObjectStruct, TSharedRef<FDetailTreeNode> DetailTreeNode);
 
 	/** */
-	bool HasSpecificList(const UStruct* ObjectStruct) const;
+	UE_API bool HasSpecificList(const UStruct* ObjectStruct) const;
 
 	/** Check whether a property exists on the PermissionList for a specific Struct - this will return false if the property is AllowListed on a parent Struct */
-	bool IsSpecificPropertyAllowListed(const UStruct* ObjectStruct, FName PropertyName) const;
+	UE_API bool IsSpecificPropertyAllowListed(const UStruct* ObjectStruct, FName PropertyName) const;
 	/** Check whether a property exists on the DenyList for a specific Struct - this will return false if the property is DenyListed on a parent Struct */
-	bool IsSpecificPropertyDenyListed(const UStruct* ObjectStruct, FName PropertyName) const;
+	UE_API bool IsSpecificPropertyDenyListed(const UStruct* ObjectStruct, FName PropertyName) const;
 
 	/** Gets a read-only copy of the original, un-flattened PermissionList. */
 	const TMap<TSoftObjectPtr<const UStruct>, FPropertyPermissionListEntry>& GetRawPermissionList() const { return RawPropertyPermissionList; }
 
 	/** Clear CachedPropertyPermissionList to cause the PermissionListed property list to be regenerated next time it's queried */
-	void ClearCache();
+	UE_API void ClearCache();
 
 	/** If true, PermissionListUpdatedDelegate will not broadcast when modifying the permission lists */
 	bool bSuppressUpdateDelegate = false;
 
 protected:
-	FPropertyPermissionList();
-	~FPropertyPermissionList();
+	UE_API FPropertyPermissionList();
+	UE_API ~FPropertyPermissionList();
 
 private:
-	bool Tick(float DeltaTime);
-	void RegisterOnBlueprintCompiled();
-	void ClearCacheAndBroadcast(TSoftObjectPtr<const UStruct> ObjectStruct = nullptr, FName OwnerName = NAME_None);
+	UE_API bool Tick(float DeltaTime);
+	UE_API void RegisterOnBlueprintCompiled();
+	UE_API void ClearCacheAndBroadcast(TSoftObjectPtr<const UStruct> ObjectStruct = nullptr, FName OwnerName = NAME_None);
 
 	/** Whether DoesPropertyPassFilter should perform its PermissionList check or always return true */
 	bool bEnablePermissionList = false;
@@ -154,11 +156,11 @@ private:
 	mutable TMap<TWeakObjectPtr<const UStruct>, FPropertyPermissionListEntry> CachedPropertyPermissionList;
 
 	/** Get or create the cached PermissionList for a specific UStruct */
-	const FNamePermissionList& GetCachedPermissionListForStruct(const UStruct* Struct) const;
-	const FNamePermissionList& GetCachedPermissionListForStructHelper(const UStruct* Struct, bool& bInOutShouldAllowListAllProperties) const;
+	UE_API const FNamePermissionList& GetCachedPermissionListForStruct(const UStruct* Struct) const;
+	UE_API const FNamePermissionList& GetCachedPermissionListForStructHelper(const UStruct* Struct, bool& bInOutShouldAllowListAllProperties) const;
 };
 
-class PROPERTYEDITOR_API FPropertyEditorPermissionList : public FPropertyPermissionList
+class FPropertyEditorPermissionList : public FPropertyPermissionList
 {
 public:
 	static FPropertyEditorPermissionList& Get()
@@ -180,7 +182,7 @@ private:
 	bool bShouldShowMenuEntries = false;
 };
 
-class PROPERTYEDITOR_API FHiddenPropertyPermissionList : public FPropertyPermissionList
+class FHiddenPropertyPermissionList : public FPropertyPermissionList
 {
 public:
 	static FHiddenPropertyPermissionList& Get()
@@ -192,3 +194,5 @@ private:
 	FHiddenPropertyPermissionList() = default;
 	~FHiddenPropertyPermissionList() = default;
 };
+
+#undef UE_API
diff --git a/Engine/Source/Editor/PropertyEditor/Public/PropertyRestriction.h b/Engine/Source/Editor/PropertyEditor/Public/PropertyRestriction.h
--- a/Engine/Source/Editor/PropertyEditor/Public/PropertyRestriction.h
+++ b/Engine/Source/Editor/PropertyEditor/Public/PropertyRestriction.h
@@ -7,9 +7,11 @@
 #include "CoreMinimal.h"
 #include "Internationalization/Text.h"
 
+#define UE_API PROPERTYEDITOR_API
+
 class IClassViewerFilter;
 
-class PROPERTYEDITOR_API FPropertyRestriction
+class FPropertyRestriction
 {
 public:
 
@@ -20,16 +22,16 @@ public:
 
 	const FText& GetReason() const { return Reason; }
 
-	bool IsValueHidden(const FString& Value) const;
-	bool IsValueDisabled(const FString& Value) const;
-	void AddHiddenValue(FString Value);
-	void AddDisabledValue(FString Value);
-	void AddClassFilter(TSharedRef<IClassViewerFilter> ClassFilter);
+	UE_API bool IsValueHidden(const FString& Value) const;
+	UE_API bool IsValueDisabled(const FString& Value) const;
+	UE_API void AddHiddenValue(FString Value);
+	UE_API void AddDisabledValue(FString Value);
+	UE_API void AddClassFilter(TSharedRef<IClassViewerFilter> ClassFilter);
 
-	void RemoveHiddenValue(FString Value);
-	void RemoveDisabledValue(FString Value);
-	void RemoveClassFilter(TSharedRef<IClassViewerFilter> ClassFilter);
-	void RemoveAll();
+	UE_API void RemoveHiddenValue(FString Value);
+	UE_API void RemoveDisabledValue(FString Value);
+	UE_API void RemoveClassFilter(TSharedRef<IClassViewerFilter> ClassFilter);
+	UE_API void RemoveAll();
 
 	TArray<FString>::TConstIterator GetHiddenValuesIterator() const 
 	{
@@ -52,3 +54,5 @@ private:
 	TArray<TSharedRef<IClassViewerFilter>> ClassViewFilter;
 	FText Reason;
 };
+
+#undef UE_API
diff --git a/Engine/Source/Editor/PropertyEditor/Public/SDetailsSplitter.h b/Engine/Source/Editor/PropertyEditor/Public/SDetailsSplitter.h
--- a/Engine/Source/Editor/PropertyEditor/Public/SDetailsSplitter.h
+++ b/Engine/Source/Editor/PropertyEditor/Public/SDetailsSplitter.h
@@ -5,8 +5,10 @@
 #include "Templates/SharedPointer.h"
 #include "IDetailsView.h"
 
+#define UE_API PROPERTYEDITOR_API
+
 // Splitter that allows you to provide an FAsyncDetailViewDiff to connect like-properties between two or more details panels
-class PROPERTYEDITOR_API SDetailsSplitter : public SCompoundWidget
+class SDetailsSplitter : public SCompoundWidget
 {
 public:
 	// Callback that's called before highlighting a row. If true, the row is skipped.
@@ -35,7 +37,7 @@ public:
 		TAttribute<TSharedPtr<FAsyncDetailViewDiff>> DiffRight;
 		FShouldIgnoreRow ShouldIgnoreRow;
 	};
-	static FSlot::FSlotArguments Slot();
+	static UE_API FSlot::FSlotArguments Slot();
 	
 	SLATE_BEGIN_ARGS(SDetailsSplitter)
 	{}
@@ -44,22 +46,22 @@ public:
 		SLATE_SLOT_ARGUMENT(FSlot, Slots)
 	SLATE_END_ARGS()
 
-	void Construct(const FArguments& InArgs);
+	UE_API void Construct(const FArguments& InArgs);
 
-	virtual int32 OnPaint(const FPaintArgs& Args, const FGeometry& AllottedGeometry, const FSlateRect& MyCullingRect,
+	UE_API virtual int32 OnPaint(const FPaintArgs& Args, const FGeometry& AllottedGeometry, const FSlateRect& MyCullingRect,
 	                      FSlateWindowElementList& OutDrawElements, int32 LayerId, const FWidgetStyle& InWidgetStyle,
 	                      bool bParentEnabled) const override;
 
-	virtual FReply OnMouseMove(const FGeometry& MyGeometry, const FPointerEvent& MouseEvent) override;
-	virtual void OnMouseLeave(const FPointerEvent& MouseEvent) override;
-	virtual FReply OnMouseButtonDown(const FGeometry& MyGeometry, const FPointerEvent& MouseEvent) override;
+	UE_API virtual FReply OnMouseMove(const FGeometry& MyGeometry, const FPointerEvent& MouseEvent) override;
+	UE_API virtual void OnMouseLeave(const FPointerEvent& MouseEvent) override;
+	UE_API virtual FReply OnMouseButtonDown(const FGeometry& MyGeometry, const FPointerEvent& MouseEvent) override;
 	
-	void AddSlot(const FSlot::FSlotArguments& SlotArgs, int32 Index = INDEX_NONE);
-	FPanel& GetPanel(int32 Index);
+	UE_API void AddSlot(const FSlot::FSlotArguments& SlotArgs, int32 Index = INDEX_NONE);
+	UE_API FPanel& GetPanel(int32 Index);
 
-	void SetRowHighlightColorDelegate(const FRowHighlightColor& Delegate);
+	UE_API void SetRowHighlightColorDelegate(const FRowHighlightColor& Delegate);
 	// generate a highlight color delegate that utilizes the merge results to determine the color
-	void HighlightFromMergeResults(const TMap<FString, TMap<FPropertySoftPath, ETreeDiffResult>>& Highlights);
+	UE_API void HighlightFromMergeResults(const TMap<FString, TMap<FPropertySoftPath, ETreeDiffResult>>& Highlights);
 
 private:
 	
@@ -78,10 +80,10 @@ private:
 		EPropertyCopyDirection CopyDirection = EPropertyCopyDirection::Copy_None;
 	};
 	
-	void PaintPropertyConnector(FSlateWindowElementList& OutDrawElements, int32 LayerId, const FSlateRect& LeftPropertyRect, 
+	UE_API void PaintPropertyConnector(FSlateWindowElementList& OutDrawElements, int32 LayerId, const FSlateRect& LeftPropertyRect, 
 		const FSlateRect& RightPropertyRect, const FLinearColor& FillColor, const FLinearColor& OutlineColor) const;
 	
-	void PaintCopyPropertyButton(FSlateWindowElementList& OutDrawElements, int32 LayerId, const TUniquePtr<FAsyncDetailViewDiff::DiffNodeType>& DiffNode,
+	UE_API void PaintCopyPropertyButton(FSlateWindowElementList& OutDrawElements, int32 LayerId, const TUniquePtr<FAsyncDetailViewDiff::DiffNodeType>& DiffNode,
 		const FSlateRect& LeftPropertyRect, const FSlateRect& RightPropertyRect, EPropertyCopyDirection CopyDirection) const;
 	
 	TSharedPtr<SSplitter> Splitter;
@@ -90,3 +92,5 @@ private:
 	FRowHighlightColor GetRowHighlightColor;
 	FShouldHighlightRow GetShouldHighlightRow;
 };
+
+#undef UE_API
diff --git a/Engine/Source/Editor/PropertyEditor/Public/SPerQualityLevelPropertiesWidget.h b/Engine/Source/Editor/PropertyEditor/Public/SPerQualityLevelPropertiesWidget.h
--- a/Engine/Source/Editor/PropertyEditor/Public/SPerQualityLevelPropertiesWidget.h
+++ b/Engine/Source/Editor/PropertyEditor/Public/SPerQualityLevelPropertiesWidget.h
@@ -8,9 +8,11 @@
 #include "Widgets/SCompoundWidget.h"
 #include "Framework/SlateDelegates.h"
 
+#define UE_API PROPERTYEDITOR_API
+
 class FMenuBuilder;
 
-class PROPERTYEDITOR_API SOverridePropertiesWidget : public SCompoundWidget
+class SOverridePropertiesWidget : public SCompoundWidget
 {
 public:
 	typedef typename TSlateDelegates<FName>::FOnGenerateWidget FOnGenerateWidget;
@@ -34,18 +36,18 @@ public:
 	*
 	* @param	InArgs	The declaration data for this widget
 	*/
-	virtual void Construct(const typename SOverridePropertiesWidget::FArguments& InArgs);
+	UE_API virtual void Construct(const typename SOverridePropertiesWidget::FArguments& InArgs);
 
 protected:
 	virtual void ConstructChildren() = 0;
 
-	TSharedRef<SWidget> MakeOverrideWidget(FName InName, FText InDisplayText, const TArray<FName>& InEntries, FMenuBuilder& InAddMenuBuilder);
+	UE_API TSharedRef<SWidget> MakeOverrideWidget(FName InName, FText InDisplayText, const TArray<FName>& InEntries, FMenuBuilder& InAddMenuBuilder);
 
-	void AddEntry(FName EntryName);
+	UE_API void AddEntry(FName EntryName);
 
-	FReply RemoveEntry(FName EntryName);
+	UE_API FReply RemoveEntry(FName EntryName);
 
-	void AddEntryToMenu(const FName& EntryName, const FTextFormat Format, FMenuBuilder& AddEntryMenuBuilder);
+	UE_API void AddEntryToMenu(const FName& EntryName, const FTextFormat Format, FMenuBuilder& AddEntryMenuBuilder);
 
 	FOnGenerateWidget OnGenerateWidget;
 	FOnOverrideAction OnAddEntry;
@@ -60,17 +62,18 @@ protected:
 /**
 * SPerQualityLevelPropertiesWidget
 */
-class PROPERTYEDITOR_API SPerQualityLevelPropertiesWidget : public SOverridePropertiesWidget
+class SPerQualityLevelPropertiesWidget : public SOverridePropertiesWidget
 {
 public:
 
-	virtual void Construct(const typename SPerQualityLevelPropertiesWidget::FArguments& InArgs);
+	UE_API virtual void Construct(const typename SPerQualityLevelPropertiesWidget::FArguments& InArgs);
 
 protected:
-	virtual void ConstructChildren() override;
+	UE_API virtual void ConstructChildren() override;
 
 };
 
 
 
 
+#undef UE_API
diff --git a/Engine/Source/Editor/PropertyEditor/Public/SResetToDefaultMenu.h b/Engine/Source/Editor/PropertyEditor/Public/SResetToDefaultMenu.h
--- a/Engine/Source/Editor/PropertyEditor/Public/SResetToDefaultMenu.h
+++ b/Engine/Source/Editor/PropertyEditor/Public/SResetToDefaultMenu.h
@@ -12,6 +12,8 @@
 #include "Widgets/DeclarativeSyntaxSupport.h"
 #include "Widgets/SCompoundWidget.h"
 
+#define UE_API PROPERTYEDITOR_API
+
 struct FGeometry;
 
 /** 
@@ -24,7 +26,7 @@ DECLARE_DELEGATE_RetVal(FText, FOnGetResetToDefaultText);
  * Displays a reset to default menu for a set of properties
  * Will hide itself automatically when all of the properties being observed are already at their default values
  */
-class PROPERTYEDITOR_API SResetToDefaultMenu : public SCompoundWidget
+class SResetToDefaultMenu : public SCompoundWidget
 {
 public:
 
@@ -51,37 +53,37 @@ public:
 		SLATE_ARGUMENT(FOnGetResetToDefaultText, OnGetResetToDefaultText)
 	SLATE_END_ARGS()
 
-	void Construct( const FArguments& InArgs );
+	UE_API void Construct( const FArguments& InArgs );
 
 	/**
 	 * Adds a new property to the menu that is displayed when a user clicks the reset to default button
 	 *
 	 * @param InProperty	The property handle to display
 	 */
-	void AddProperty( TSharedRef<class IPropertyHandle> InProperty );
+	UE_API void AddProperty( TSharedRef<class IPropertyHandle> InProperty );
 
 private:
-	virtual void Tick( const FGeometry& AllottedGeometry, const double InCurrentTime, const float InDeltaTime ) override;
+	UE_API virtual void Tick( const FGeometry& AllottedGeometry, const double InCurrentTime, const float InDeltaTime ) override;
 
 	/**
 	 * @return The visibility of the reset to default widget
 	 */
-	EVisibility GetResetToDefaultVisibility() const;
+	UE_API EVisibility GetResetToDefaultVisibility() const;
 
 	/**
 	 * Called when the menu is open to regenerate the default menu content
 	 */
-	TSharedRef<class SWidget> OnGenerateResetToDefaultMenuContent();
+	UE_API TSharedRef<class SWidget> OnGenerateResetToDefaultMenuContent();
 
 	/**
 	 * Resets property to default, only used when our property handle is not valid.
 	 */
-	void ResetToDefault();
+	UE_API void ResetToDefault();
 
 	/**
 	 * Resets all properties to default
 	 */
-	void ResetAllToDefault();
+	UE_API void ResetAllToDefault();
 private:
 
 	/** Properties that should be displayed in the menu */
@@ -102,3 +104,5 @@ private:
 	/** Delegate used to get the text to display when resetting our value - only used when we don't have a property handle */
 	FOnGetResetToDefaultText OnGetResetToDefaultText;
 };
+
+#undef UE_API
diff --git a/Engine/Source/Editor/PropertyEditor/Public/SResetToDefaultPropertyEditor.h b/Engine/Source/Editor/PropertyEditor/Public/SResetToDefaultPropertyEditor.h
--- a/Engine/Source/Editor/PropertyEditor/Public/SResetToDefaultPropertyEditor.h
+++ b/Engine/Source/Editor/PropertyEditor/Public/SResetToDefaultPropertyEditor.h
@@ -11,12 +11,14 @@
 #include "Widgets/DeclarativeSyntaxSupport.h"
 #include "Widgets/SCompoundWidget.h"
 
+#define UE_API PROPERTYEDITOR_API
+
 // Forward decl
 class FResetToDefaultOverride;
 struct FGeometry;
 
 /** Widget showing the reset to default value button */
-class PROPERTYEDITOR_API SResetToDefaultPropertyEditor : public SCompoundWidget
+class SResetToDefaultPropertyEditor : public SCompoundWidget
 {
 public:
 
@@ -27,20 +29,20 @@ public:
 		SLATE_ARGUMENT( TOptional<FResetToDefaultOverride>, CustomResetToDefault )
 	SLATE_END_ARGS()
 
-	~SResetToDefaultPropertyEditor();
+	UE_API ~SResetToDefaultPropertyEditor();
 
-	void Construct( const FArguments& InArgs, const TSharedPtr< class IPropertyHandle>& InPropertyHandle );
+	UE_API void Construct( const FArguments& InArgs, const TSharedPtr< class IPropertyHandle>& InPropertyHandle );
 
 private:
-	FText GetResetToolTip() const;
+	UE_API FText GetResetToolTip() const;
 
-	EVisibility GetDiffersFromDefaultAsVisibility() const;
+	UE_API EVisibility GetDiffersFromDefaultAsVisibility() const;
 
-	void Tick( const FGeometry& AllottedGeometry, const double InCurrentTime, const float InDeltaTime );
+	UE_API void Tick( const FGeometry& AllottedGeometry, const double InCurrentTime, const float InDeltaTime );
 
-	FReply OnResetClicked();
+	UE_API FReply OnResetClicked();
 
-	void UpdateDiffersFromDefaultState();
+	UE_API void UpdateDiffersFromDefaultState();
 
 private:
 	TOptional<FResetToDefaultOverride> OptionalCustomResetToDefault;
@@ -51,3 +53,5 @@ private:
 
 	bool bValueDiffersFromDefault;
 };
+
+#undef UE_API
diff --git a/Engine/Source/Editor/PropertyEditor/Public/SStandaloneCustomizedValueWidget.h b/Engine/Source/Editor/PropertyEditor/Public/SStandaloneCustomizedValueWidget.h
--- a/Engine/Source/Editor/PropertyEditor/Public/SStandaloneCustomizedValueWidget.h
+++ b/Engine/Source/Editor/PropertyEditor/Public/SStandaloneCustomizedValueWidget.h
@@ -6,6 +6,8 @@
 #include "IPropertyTypeCustomization.h"
 #include "Widgets/SCompoundWidget.h"
 
+#define UE_API PROPERTYEDITOR_API
+
 class FAssetThumbnailPool;
 class FDetailCategoryImpl;
 
@@ -14,7 +16,7 @@ class FDetailCategoryImpl;
  * This widget is generally used in the property editor to display a struct as a single row, like with an FColor.
  * Properties passed in that do not have a header customization will return a null widget.
  */
-class PROPERTYEDITOR_API SStandaloneCustomizedValueWidget : public SCompoundWidget, public IPropertyTypeCustomizationUtils
+class SStandaloneCustomizedValueWidget : public SCompoundWidget, public IPropertyTypeCustomizationUtils
 {
 public:
 	SLATE_BEGIN_ARGS( SStandaloneCustomizedValueWidget )
@@ -23,10 +25,10 @@ public:
 		SLATE_ARGUMENT( TSharedPtr<FDetailCategoryImpl>, ParentCategory)
 	SLATE_END_ARGS()
 	
-	void Construct(const FArguments& InArgs,
+	UE_API void Construct(const FArguments& InArgs,
 		TSharedPtr<IPropertyTypeCustomization> InCustomizationInterface, TSharedRef<IPropertyHandle> InPropertyHandle);
 
-	virtual TSharedPtr<FAssetThumbnailPool> GetThumbnailPool() const override;
+	UE_API virtual TSharedPtr<FAssetThumbnailPool> GetThumbnailPool() const override;
 
 private:
 	TWeakPtr<FDetailCategoryImpl> ParentCategory;
@@ -34,3 +36,5 @@ private:
 	TSharedPtr<IPropertyHandle> PropertyHandle;
 	TSharedPtr<FDetailWidgetRow> CustomPropertyWidget;
 };
+
+#undef UE_API
diff --git a/Engine/Source/Editor/RenderResourceViewer/Public/RenderResourceViewerSettings.h b/Engine/Source/Editor/RenderResourceViewer/Public/RenderResourceViewerSettings.h
--- a/Engine/Source/Editor/RenderResourceViewer/Public/RenderResourceViewerSettings.h
+++ b/Engine/Source/Editor/RenderResourceViewer/Public/RenderResourceViewerSettings.h
@@ -18,8 +18,8 @@ struct FRenderResourceViewerTreemapFilter
 	FString DisplayName;
 };
 
-UCLASS(config = Editor, DefaultConfig, meta = (DisplayName = "Render Resource Viewer"))
-class RENDERRESOURCEVIEWER_API URenderResourceViewerSettings : public UDeveloperSettings
+UCLASS(MinimalAPI, config = Editor, DefaultConfig, meta = (DisplayName = "Render Resource Viewer"))
+class URenderResourceViewerSettings : public UDeveloperSettings
 {
 	GENERATED_UCLASS_BODY()
 
diff --git a/Engine/Source/Editor/RewindDebuggerInterface/Public/IRewindDebuggerDoubleClickHandler.h b/Engine/Source/Editor/RewindDebuggerInterface/Public/IRewindDebuggerDoubleClickHandler.h
--- a/Engine/Source/Editor/RewindDebuggerInterface/Public/IRewindDebuggerDoubleClickHandler.h
+++ b/Engine/Source/Editor/RewindDebuggerInterface/Public/IRewindDebuggerDoubleClickHandler.h
@@ -6,11 +6,13 @@
 #include "IRewindDebugger.h"
 #include "Features/IModularFeatures.h"
 
+#define UE_API REWINDDEBUGGERINTERFACE_API
+
 // Interface class for adding custom double click handling to components
-class REWINDDEBUGGERINTERFACE_API IRewindDebuggerDoubleClickHandler : public IModularFeature
+class IRewindDebuggerDoubleClickHandler : public IModularFeature
 {
 public:
-	static const FName ModularFeatureName;
+	static UE_API const FName ModularFeatureName;
 	
 	// returns the name of a type of UObject for which this menu extender will be called
 	virtual FName GetTargetTypeName() const = 0;
@@ -18,3 +20,5 @@ public:
 	// called when the selected component has been double clicked.  return true if something was opened
 	virtual bool HandleDoubleClick(IRewindDebugger* IRewindDebugger) = 0;
 };
+
+#undef UE_API
diff --git a/Engine/Source/Editor/RewindDebuggerInterface/Public/IRewindDebuggerExtension.h b/Engine/Source/Editor/RewindDebuggerInterface/Public/IRewindDebuggerExtension.h
--- a/Engine/Source/Editor/RewindDebuggerInterface/Public/IRewindDebuggerExtension.h
+++ b/Engine/Source/Editor/RewindDebuggerInterface/Public/IRewindDebuggerExtension.h
@@ -4,6 +4,8 @@
 #include "CoreMinimal.h"
 #include "Features/IModularFeatures.h"
 
+#define UE_API REWINDDEBUGGERINTERFACE_API
+
 
 // IRewindDebuggerExtension 
 //
@@ -13,10 +15,10 @@
 
 class IRewindDebugger;
 
-class REWINDDEBUGGERINTERFACE_API IRewindDebuggerExtension : public IModularFeature
+class IRewindDebuggerExtension : public IModularFeature
 {
 public:
-	static const FName ModularFeatureName;
+	static UE_API const FName ModularFeatureName;
 
 	// debugging name for this extension
 	virtual FString GetName() { return TEXT("RewindDebugger Extension"); }
@@ -33,3 +35,5 @@ public:
 	// called when recording is unloaded
 	virtual void Clear(IRewindDebugger* RewindDebugger) {};
 };
+
+#undef UE_API
diff --git a/Engine/Source/Editor/RewindDebuggerInterface/Public/IRewindDebuggerView.h b/Engine/Source/Editor/RewindDebuggerInterface/Public/IRewindDebuggerView.h
--- a/Engine/Source/Editor/RewindDebuggerInterface/Public/IRewindDebuggerView.h
+++ b/Engine/Source/Editor/RewindDebuggerInterface/Public/IRewindDebuggerView.h
@@ -6,7 +6,7 @@
 #include "Widgets/SCompoundWidget.h"
 
 // Interface class for views of recorded debug data
-class REWINDDEBUGGERINTERFACE_API IRewindDebuggerView : public SCompoundWidget
+class IRewindDebuggerView : public SCompoundWidget
 {
 	public:
 		// unique name for widget
diff --git a/Engine/Source/Editor/RewindDebuggerInterface/Public/IRewindDebuggerViewCreator.h b/Engine/Source/Editor/RewindDebuggerInterface/Public/IRewindDebuggerViewCreator.h
--- a/Engine/Source/Editor/RewindDebuggerInterface/Public/IRewindDebuggerViewCreator.h
+++ b/Engine/Source/Editor/RewindDebuggerInterface/Public/IRewindDebuggerViewCreator.h
@@ -6,16 +6,18 @@
 #include "IRewindDebuggerView.h"
 #include "Textures/SlateIcon.h"
 
+#define UE_API REWINDDEBUGGERINTERFACE_API
+
 namespace TraceServices
 {
 	class IAnalysisSession;
 }
 
 // Interface class which creates debug widgets
-class REWINDDEBUGGERINTERFACE_API IRewindDebuggerViewCreator : public IModularFeature
+class IRewindDebuggerViewCreator : public IModularFeature
 {
 public: 
-	static const FName ModularFeatureName;
+	static UE_API const FName ModularFeatureName;
 
 	// returns the name of a type of UObject for which this debug view will be created
 	virtual FName GetTargetTypeName() const = 0;
@@ -37,4 +39,6 @@ public:
 
 	// creates and returns a widget, which will be displayed in Rewind Debugger
 	virtual TSharedPtr<IRewindDebuggerView> CreateDebugView(uint64 ObjectId, double CurrentTime, const TraceServices::IAnalysisSession& InAnalysisSession) const = 0;
-};
\ No newline at end of file
+};
+
+#undef UE_API
diff --git a/Engine/Source/Editor/SVGDistanceField/Public/SvgDistanceFieldGenerator.h b/Engine/Source/Editor/SVGDistanceField/Public/SvgDistanceFieldGenerator.h
--- a/Engine/Source/Editor/SVGDistanceField/Public/SvgDistanceFieldGenerator.h
+++ b/Engine/Source/Editor/SVGDistanceField/Public/SvgDistanceFieldGenerator.h
@@ -9,8 +9,10 @@
 
 #include "SvgDistanceFieldGenerator.generated.h"
 
-UCLASS()
-class SVGDISTANCEFIELD_API USvgDistanceFieldGenerator : public UBlueprintFunctionLibrary
+#define UE_API SVGDISTANCEFIELD_API
+
+UCLASS(MinimalAPI)
+class USvgDistanceFieldGenerator : public UBlueprintFunctionLibrary
 {
 	GENERATED_BODY()
 
@@ -18,7 +20,9 @@ public:
 
 #if WITH_EDITORONLY_DATA
 	UFUNCTION(BlueprintCallable, Category=SvgDistanceField)
-	static UTexture2D* GenerateTextureFromSvgFile(const FString& SvgFilePath, const FSvgDistanceFieldConfiguration& Configuration);
+	static UE_API UTexture2D* GenerateTextureFromSvgFile(const FString& SvgFilePath, const FSvgDistanceFieldConfiguration& Configuration);
 #endif
 
 };
+
+#undef UE_API
diff --git a/Engine/Source/Editor/SceneDepthPickerMode/Public/SceneDepthPickerMode.h b/Engine/Source/Editor/SceneDepthPickerMode/Public/SceneDepthPickerMode.h
--- a/Engine/Source/Editor/SceneDepthPickerMode/Public/SceneDepthPickerMode.h
+++ b/Engine/Source/Editor/SceneDepthPickerMode/Public/SceneDepthPickerMode.h
@@ -5,28 +5,32 @@
 #include "CoreMinimal.h"
 #include "Modules/ModuleInterface.h"
 
+#define UE_API SCENEDEPTHPICKERMODE_API
+
 DECLARE_DELEGATE_OneParam( FOnSceneDepthLocationSelected, FVector );
 
 /**
  * Scene depth picker mode module
  */
-class SCENEDEPTHPICKERMODE_API FSceneDepthPickerModeModule : public IModuleInterface
+class FSceneDepthPickerModeModule : public IModuleInterface
 {
 public:
 	// IModuleInterface
-	virtual void StartupModule() override;
-	virtual void ShutdownModule() override;
+	UE_API virtual void StartupModule() override;
+	UE_API virtual void ShutdownModule() override;
 	// End of IModuleInterface
 
 	/** 
 	 * Enter scene depth picking mode (note: will cancel any current scene depth picking)
 	 * @param	InOnDepthSelected		Delegate to call when user has selected a location to sample the scene depth
 	 */
-	void BeginSceneDepthPickingMode(FOnSceneDepthLocationSelected InOnSceneDepthLocationSelected);
+	UE_API void BeginSceneDepthPickingMode(FOnSceneDepthLocationSelected InOnSceneDepthLocationSelected);
 
 	/** Exit scene depth picking mode */
-	void EndSceneDepthPickingMode();
+	UE_API void EndSceneDepthPickingMode();
 
 	/** @return Whether or not scene depth picking mode is currently active */
-	bool IsInSceneDepthPickingMode() const;
+	UE_API bool IsInSceneDepthPickingMode() const;
 };
+
+#undef UE_API
diff --git a/Engine/Source/Editor/SceneOutliner/Public/ActorBaseTreeItem.h b/Engine/Source/Editor/SceneOutliner/Public/ActorBaseTreeItem.h
--- a/Engine/Source/Editor/SceneOutliner/Public/ActorBaseTreeItem.h
+++ b/Engine/Source/Editor/SceneOutliner/Public/ActorBaseTreeItem.h
@@ -6,16 +6,20 @@
 #include "ISceneOutlinerTreeItem.h"
 #include "SceneOutlinerStandaloneTypes.h"
 
+#define UE_API SCENEOUTLINER_API
+
 struct FGuid;
 
 /** A tree item that represents an actor, loaded or unloaded */
-struct SCENEOUTLINER_API IActorBaseTreeItem : ISceneOutlinerTreeItem
+struct IActorBaseTreeItem : ISceneOutlinerTreeItem
 {
 public:
 	IActorBaseTreeItem(FSceneOutlinerTreeItemType InType) : ISceneOutlinerTreeItem(InType) {}
 
 	/** Static type identifier for the base class tree item */
-	static const FSceneOutlinerTreeItemType Type;
+	static UE_API const FSceneOutlinerTreeItemType Type;
 
 	virtual const FGuid& GetGuid() const =0;
 };
+
+#undef UE_API
diff --git a/Engine/Source/Editor/SceneOutliner/Public/ActorBrowsingMode.h b/Engine/Source/Editor/SceneOutliner/Public/ActorBrowsingMode.h
--- a/Engine/Source/Editor/SceneOutliner/Public/ActorBrowsingMode.h
+++ b/Engine/Source/Editor/SceneOutliner/Public/ActorBrowsingMode.h
@@ -32,19 +32,19 @@ public:
 	SCENEOUTLINER_API virtual bool CanRenameItem(const ISceneOutlinerTreeItem& Item) const override;
 	SCENEOUTLINER_API virtual FText GetStatusText() const override;
 	SCENEOUTLINER_API virtual FSlateColor GetStatusTextColor() const override;
-	SCENEOUTLINER_API virtual ESelectionMode::Type GetSelectionMode() const override { return ESelectionMode::Multi; }
-	SCENEOUTLINER_API virtual bool SupportsKeyboardFocus() const override { return true; }
-	SCENEOUTLINER_API virtual bool ShouldShowFolders() const override { return true; }
-	SCENEOUTLINER_API virtual bool SupportsCreateNewFolder() const override { return true; }
-	SCENEOUTLINER_API virtual bool ShowStatusBar() const override { return true; }
-	SCENEOUTLINER_API virtual bool ShowViewButton() const override { return true; }
-	SCENEOUTLINER_API virtual bool ShowFilterOptions() const override { return true; }
+	virtual ESelectionMode::Type GetSelectionMode() const override { return ESelectionMode::Multi; }
+	virtual bool SupportsKeyboardFocus() const override { return true; }
+	virtual bool ShouldShowFolders() const override { return true; }
+	virtual bool SupportsCreateNewFolder() const override { return true; }
+	virtual bool ShowStatusBar() const override { return true; }
+	virtual bool ShowViewButton() const override { return true; }
+	virtual bool ShowFilterOptions() const override { return true; }
 	SCENEOUTLINER_API virtual bool CanDelete() const override;
 	SCENEOUTLINER_API virtual bool CanRename() const override;
 	SCENEOUTLINER_API virtual bool CanCut() const override;
 	SCENEOUTLINER_API virtual bool CanCopy() const override;
 	SCENEOUTLINER_API virtual bool CanPaste() const override;
-	SCENEOUTLINER_API virtual bool CanSupportDragAndDrop() const { return true; }
+	virtual bool CanSupportDragAndDrop() const { return true; }
 	SCENEOUTLINER_API virtual bool HasErrors() const;
 	SCENEOUTLINER_API virtual FText GetErrorsText() const;
 	SCENEOUTLINER_API virtual void RepairErrors() const;	
