Fix property permission list no longer working
#rnx

diff --git a/Engine/Source/Editor/PropertyEditor/Private/PropertyPermissionList.cpp b/Engine/Source/Editor/PropertyEditor/Private/PropertyPermissionList.cpp
--- a/Engine/Source/Editor/PropertyEditor/Private/PropertyPermissionList.cpp
+++ b/Engine/Source/Editor/PropertyEditor/Private/PropertyPermissionList.cpp
@@ -20,6 +20,12 @@ namespace PropertyEditorPermissionList
 		TEXT("PropertyEditor.PropertyEditorPermissionList.SupportAllowAllFunctions"),
 		bSupportAllowAllFunctions,
 		TEXT("Enable allowing all zero parameter call in editor functions when allowing all properties"));
+
+	static FString ClassToDebug;
+	static FAutoConsoleVariableRef CVarClassToDebug(
+		TEXT("PropertyEditor.PropertyEditorPermissionList.ClassToDebug"),
+		ClassToDebug,
+		TEXT("Debug when a class that contains this string is cached"));
 }
 
 bool operator==(const FPropertyPermissionList::FPermissionListUpdate& A, const FPropertyPermissionList::FPermissionListUpdate& B)
@@ -293,6 +299,18 @@ const FNamePermissionList& FPropertyPermissionList::GetCachedPermissionListForSt
 {
 	check(Struct);
 
+#if !UE_BUILD_SHIPPING
+	if (!PropertyEditorPermissionList::ClassToDebug.IsEmpty())
+	{
+		FNameBuilder StructPathName;
+		Struct->GetPathName(nullptr, StructPathName);
+		if (StructPathName.ToView().Contains(PropertyEditorPermissionList::ClassToDebug))
+		{
+			bool bSetBreakpointHere = true;
+		}
+	}
+#endif
+
 	auto IsAllowListAllProps = [](const FPropertyPermissionListEntry* InEntry)
 		{
 			return InEntry ? (InEntry->PermissionList.GetAllowList().Num() == 0 || InEntry->Rules == EPropertyPermissionListRules::AllowListAllProperties) : true;
@@ -434,3 +452,15 @@ bool FPropertyPermissionList::IsSpecificPropertyDenyListed(const UStruct* Object
 	}
 	return false;
 }
+
+FPropertyEditorPermissionList& FPropertyEditorPermissionList::Get()
+{
+	static FPropertyEditorPermissionList PermissionList;
+	return PermissionList;
+}
+
+FHiddenPropertyPermissionList& FHiddenPropertyPermissionList::Get()
+{
+	static FHiddenPropertyPermissionList PermissionList;
+	return PermissionList;
+}
\ No newline at end of file
diff --git a/Engine/Source/Editor/PropertyEditor/Public/PropertyPermissionList.h b/Engine/Source/Editor/PropertyEditor/Public/PropertyPermissionList.h
--- a/Engine/Source/Editor/PropertyEditor/Public/PropertyPermissionList.h
+++ b/Engine/Source/Editor/PropertyEditor/Public/PropertyPermissionList.h
@@ -163,11 +163,7 @@ private:
 class FPropertyEditorPermissionList : public FPropertyPermissionList
 {
 public:
-	static FPropertyEditorPermissionList& Get()
-	{
-		static FPropertyEditorPermissionList PermissionList;
-		return PermissionList;
-	}
+	static UE_API FPropertyEditorPermissionList& Get();
 
 	/** Whether the Details View should show special menu entries to add/remove items in the PermissionList */
 	bool ShouldShowMenuEntries() const { return bShouldShowMenuEntries; }
@@ -185,11 +181,7 @@ private:
 class FHiddenPropertyPermissionList : public FPropertyPermissionList
 {
 public:
-	static FHiddenPropertyPermissionList& Get()
-	{
-		static FHiddenPropertyPermissionList PermissionList;
-		return PermissionList;
-	}
+	static UE_API FHiddenPropertyPermissionList& Get();
 private:
 	FHiddenPropertyPermissionList() = default;
 	~FHiddenPropertyPermissionList() = default;
