[PropertyEditor] Added optional sort priority to Property Sections

diff --git a/Engine/Source/Editor/PropertyEditor/Private/PropertyEditorModule.cpp b/Engine/Source/Editor/PropertyEditor/Private/PropertyEditorModule.cpp
--- a/Engine/Source/Editor/PropertyEditor/Private/PropertyEditorModule.cpp
+++ b/Engine/Source/Editor/PropertyEditor/Private/PropertyEditorModule.cpp
@@ -596,12 +596,12 @@ TSharedPtr<FPropertySection> FClassSectionMapping::FindSection(FName SectionName
 	return *Section;
 }
 
-TSharedRef<FPropertySection> FClassSectionMapping::FindOrAddSection(FName SectionName, FText DisplayName)
+TSharedRef<FPropertySection> FClassSectionMapping::FindOrAddSection(FName SectionName, FText DisplayName, const int32 Order)
 {
 	TSharedPtr<FPropertySection> Section = FindSection(SectionName);
 	if (!Section.IsValid())
 	{
-		Section = MakeShared<FPropertySection>(SectionName, DisplayName);
+		Section = MakeShared<FPropertySection>(SectionName, DisplayName, Order);
 		DefinedSections.Add(SectionName, Section);
 	}
 
@@ -656,7 +656,7 @@ void FPropertyEditorModule::RemoveSection(FName ClassName, FName SectionName)
 }
 
 
-TSharedRef<FPropertySection> FPropertyEditorModule::FindOrCreateSection(FName ClassName, FName SectionName, FText DisplayName)
+TSharedRef<FPropertySection> FPropertyEditorModule::FindOrCreateSection(FName ClassName, FName SectionName, FText DisplayName, const int32 Order)
 {
 	checkf(!ClassName.IsNone(), TEXT("Invalid class name given."));
 	checkf(!SectionName.IsNone(), TEXT("Invalid section name given."));
@@ -673,7 +673,7 @@ TSharedRef<FPropertySection> FPropertyEditorModule::FindOrCreateSection(FName Cl
 		ClassMapping = *ExistingMapping;
 	}
 
-	return ClassMapping->FindOrAddSection(SectionName, DisplayName);
+	return ClassMapping->FindOrAddSection(SectionName, DisplayName, Order);
 }
 
 TArray<TSharedPtr<FPropertySection>> FPropertyEditorModule::FindSectionsForCategory(const UStruct* Struct, FName CategoryName) const
diff --git a/Engine/Source/Editor/PropertyEditor/Private/SDetailsView.cpp b/Engine/Source/Editor/PropertyEditor/Private/SDetailsView.cpp
--- a/Engine/Source/Editor/PropertyEditor/Private/SDetailsView.cpp
+++ b/Engine/Source/Editor/PropertyEditor/Private/SDetailsView.cpp
@@ -1328,7 +1328,7 @@ void SDetailsView::RebuildSectionSelector()
 		return;
 	}
 
-	const TMap<FName, FText> AllSections = GetAllSections();
+	const TMap<FName, SDetailsView::FSection> AllSections = GetAllSections();
 	if (AllSections.IsEmpty())
 	{
 		// we've selected something that has no sections - rather than show just "All", hide the box and clear visible sections
@@ -1354,30 +1354,46 @@ void SDetailsView::RebuildSectionSelector()
 			];
 	};
 
-	TArray<FName> SortedKeys;
-	AllSections.GenerateKeyArray(SortedKeys);
-	SortedKeys.Sort([](FName A, FName B)
+	struct FOrderedSection
+	{
+		FName Name;
+		int32 Order = 0;
+	};
+
+	TArray<FOrderedSection> SortedKeys;
+	SortedKeys.Reserve(AllSections.Num());
+
+	for (const TPair<FName, SDetailsView::FSection>& SectionPair : AllSections)
+	{
+		SortedKeys.Emplace(FOrderedSection{ SectionPair.Key, SectionPair.Value.Order });
+	}
+
+	SortedKeys.Sort([](const FOrderedSection& A, const FOrderedSection& B)
 		{
 			static const FName General("General");
 			static const FName All("All");
 
-			// General first, All last, rest alphabetical
-			if (A.IsEqual(General) || B.IsEqual(All))
+			// General first, All last, rest ordered then alphabetical
+			if (A.Name.IsEqual(General) || B.Name.IsEqual(All))
 			{
 				return true;
 			}
-			if (A.IsEqual(All) || B.IsEqual(General))
+			if (A.Name.IsEqual(All) || B.Name.IsEqual(General))
 			{
 				return false;
 			}
-			return A.LexicalLess(B);
+			if (A.Order != B.Order)
+			{
+				return A.Order < B.Order;
+			}
+			return A.Name.LexicalLess(B.Name);
 		});
 
-	for (const FName& Key : SortedKeys)
+	for (const FOrderedSection& Key : SortedKeys)
 	{
 		SectionSelectorBox->AddSlot()
 		[
-			CreateSection(Key, AllSections[Key])
+			CreateSection(Key.Name, AllSections[Key.Name].DisplayName)
 		];
 	}
 
@@ -1472,12 +1488,12 @@ ECheckBoxState SDetailsView::IsSectionChecked(FName Section) const
 	return ECheckBoxState::Unchecked;
 }
 
-TMap<FName, FText> SDetailsView::GetAllSections() const
+TMap<FName, SDetailsView::FSection> SDetailsView::GetAllSections() const
 {
 	static const FName PropertyEditor("PropertyEditor");
 	FPropertyEditorModule& PropertyModule = FModuleManager::GetModuleChecked<FPropertyEditorModule>(PropertyEditor);
 
-	TMap<FName, FText> AllSections;
+	TMap<FName, SDetailsView::FSection> AllSections;
 
 	// for every category, check every base struct and find the associated section
 	// if one exists, add it to the set of valid section names
@@ -1516,7 +1532,7 @@ TMap<FName, FText> SDetailsView::GetAllSections() const
 			TArray<TSharedPtr<FPropertySection>> SectionsForCategory = PropertyModule.FindSectionsForCategory(RootBaseStruct, Category);
 			for (const TSharedPtr<FPropertySection>& Section : SectionsForCategory)
 			{
-				AllSections.Add(Section->GetName(), Section->GetDisplayName());
+				AllSections.Add(Section->GetName(), SDetailsView::FSection{ Section->GetDisplayName(), Section->GetOrder() });
 			}
 		}
 	}
diff --git a/Engine/Source/Editor/PropertyEditor/Private/SDetailsView.h b/Engine/Source/Editor/PropertyEditor/Private/SDetailsView.h
--- a/Engine/Source/Editor/PropertyEditor/Private/SDetailsView.h
+++ b/Engine/Source/Editor/PropertyEditor/Private/SDetailsView.h
@@ -132,6 +132,14 @@ public:
 	*/
 	virtual bool IsDefaultStyle() const override;
 
+private:
+	/** Minimal struct to hold section information. Note that the Name is the key, so is not present here. */
+	struct FSection
+	{
+		FText DisplayName;
+		int32 Order = 0;
+	};
+
 private:
 	/**
 	 * returns a pointer to a Customization that can be applied to all name widgets, if one exists
@@ -200,8 +208,8 @@ private:
 	/** Called after an undo or redo operation occurs in the editor. */
 	void OnPostUndoRedo();
 
-	/** Get all section names and display names for the objects currently selected in the view. */
-	TMap<FName, FText> GetAllSections() const;
+	/** Get all sections, keyed by name for the objects currently selected in the view. */
+	TMap<FName, FSection> GetAllSections() const;
 
 	/** Rebuild the section selector widget after a selection has been changed. */
 	void RebuildSectionSelector();
diff --git a/Engine/Source/Editor/PropertyEditor/Public/PropertyEditorModule.h b/Engine/Source/Editor/PropertyEditor/Public/PropertyEditorModule.h
--- a/Engine/Source/Editor/PropertyEditor/Public/PropertyEditorModule.h
+++ b/Engine/Source/Editor/PropertyEditor/Public/PropertyEditorModule.h
@@ -132,10 +132,12 @@ public:
 	/**
 	 * @param InName		The internal name of this section.
 	 * @param InDisplayName	The localizable display name to show to the user.
+	 * @param InOrder		The order in which this section should be displayed. Lower values are displayed first.
 	 */
-	FPropertySection(FName InName, FText InDisplayName) : 
+	FPropertySection(FName InName, FText InDisplayName, const int32 InOrder = 0) : 
 		Name(InName),
-		DisplayName(InDisplayName)
+		DisplayName(InDisplayName),
+		Order(InOrder)
 	{
 	}
 
@@ -160,6 +162,9 @@ public:
 	/** Get the display name of this section. */
 	virtual FText GetDisplayName() const { return DisplayName; }
 
+	/** Get the (sort) order of this section, where the default is 0. */
+	virtual int32 GetOrder() const { return Order; }
+	
 private:
 
 	/** The internal name to use for this section. */
@@ -167,6 +172,9 @@ private:
 
 	/** The display name to use for this section. */
 	FText DisplayName;
+
+	/** The order in which this section should be displayed. Lower values are displayed first. */
+	int32 Order = 0;
 	
 	/** The set of categories that are added to this section. */
 	TSet<FName> AddedCategories;
@@ -192,8 +200,10 @@ public:
 
 	/**
 	 * Find or add a section of the given name.
+	 * Optionally, you can specify an order for the section to be displayed in. Lower values are displayed first.
+	 * Note that the default "All" section will always be displayed last.
 	 */
-	TSharedRef<FPropertySection> FindOrAddSection(FName SectionName, FText DisplayName);
+	TSharedRef<FPropertySection> FindOrAddSection(FName SectionName, FText DisplayName, const int32 Order = 0);
 
 	/** 
 	 * Remove a section of the given name.
@@ -318,9 +328,10 @@ public:
 	 * @param ClassName		The class to add a section mapping for.
 	 * @param SectionName	The section to find or create.
 	 * @param DisplayName	The display name to use for the section. If the section already exists for this class, the display name will not be replaced.
+	 * @param Order			The order in which this section should be displayed. Lower values are displayed first. The default "All" section will always be displayed last.
 	 * @return				A new section, or the existing one. 
 	 */
-	virtual TSharedRef<FPropertySection> FindOrCreateSection(FName ClassName, FName SectionName, FText DisplayName);
+	virtual TSharedRef<FPropertySection> FindOrCreateSection(FName ClassName, FName SectionName, FText DisplayName, const int32 Order = 0);
 
 	/** 
 	 * Find the section that the given category in the given struct should be a part of. 
