Editor: Add support to DetailsViewArgs to supply initial sections to be selected in the filter.

diff --git a/Engine/Source/Editor/PropertyEditor/Private/SDetailsView.cpp b/Engine/Source/Editor/PropertyEditor/Private/SDetailsView.cpp
--- a/Engine/Source/Editor/PropertyEditor/Private/SDetailsView.cpp
+++ b/Engine/Source/Editor/PropertyEditor/Private/SDetailsView.cpp
@@ -88,6 +88,11 @@ void SDetailsView::Construct(const FArguments& InArgs, const FDetailsViewArgs& I
 	CurrentFilter.bShowLooseProperties = InDetailsViewArgs.bShowLooseProperties;
 	CurrentFilter.ShouldForceHideProperty = InDetailsViewArgs.ShouldForceHideProperty;
 
+	if (DetailsViewArgs.InitialSections.Num() > 0)
+	{
+		CurrentFilter.VisibleSections = DetailsViewArgs.InitialSections;
+	}
+
 	ColumnSizeData.SetValueColumnWidth(DetailsViewArgs.ColumnWidth);
 	ColumnSizeData.SetRightColumnMinWidth(DetailsViewArgs.RightColumnMinWidth);
 
diff --git a/Engine/Source/Editor/PropertyEditor/Public/DetailsViewArgs.h b/Engine/Source/Editor/PropertyEditor/Public/DetailsViewArgs.h
--- a/Engine/Source/Editor/PropertyEditor/Public/DetailsViewArgs.h
+++ b/Engine/Source/Editor/PropertyEditor/Public/DetailsViewArgs.h
@@ -125,6 +125,9 @@ struct FDetailsViewArgs
 	/** Resolve any UObject* to their IDO counterparts if they exist */
 	TOptional<bool> bResolveInstanceDataObjects;
 
+	/** Initial sections to be shown enabled in the filter */
+	TSet<FName> InitialSections;
+
 public:
 
 	FDetailsViewArgs()
