[State Tree Editor]
new class meta tag for disallowing level actor references, Child class can override the behavior. Default is allowing level actor reference.

Added the tag to relevant StateTree UObject classes

diff --git a/Engine/Plugins/Runtime/StateTree/Source/StateTreeEditorModule/Public/StateTreeEditorData.h b/Engine/Plugins/Runtime/StateTree/Source/StateTreeEditorModule/Public/StateTreeEditorData.h
--- a/Engine/Plugins/Runtime/StateTree/Source/StateTreeEditorModule/Public/StateTreeEditorData.h
+++ b/Engine/Plugins/Runtime/StateTree/Source/StateTreeEditorModule/Public/StateTreeEditorData.h
@@ -55,7 +55,7 @@ enum class EStateTreeVisitor : uint8
 /**
  * Edit time data for StateTree asset. This data gets baked into runtime format before being used by the StateTreeInstance.
  */
-UCLASS(MinimalAPI, BlueprintType, EditInlineNew, CollapseCategories, Within = "StateTree")
+UCLASS(MinimalAPI, BlueprintType, EditInlineNew, CollapseCategories, Within = "StateTree", meta = (DisallowLevelActorReference = true))
 class UStateTreeEditorData : public UObject, public IStateTreeEditorPropertyBindingsOwner
 {
 	GENERATED_BODY()
diff --git a/Engine/Plugins/Runtime/StateTree/Source/StateTreeEditorModule/Public/StateTreeState.h b/Engine/Plugins/Runtime/StateTree/Source/StateTreeEditorModule/Public/StateTreeState.h
--- a/Engine/Plugins/Runtime/StateTree/Source/StateTreeEditorModule/Public/StateTreeState.h
+++ b/Engine/Plugins/Runtime/StateTree/Source/StateTreeEditorModule/Public/StateTreeState.h
@@ -209,7 +209,7 @@ struct FStateTreeStateParameters
 /**
  * Editor representation of a state in StateTree
  */
-UCLASS(MinimalAPI, BlueprintType, EditInlineNew, CollapseCategories)
+UCLASS(MinimalAPI, BlueprintType, EditInlineNew, CollapseCategories, meta = (DisallowLevelActorReference = true))
 class UStateTreeState : public UObject, public IStateTreeSchemaProvider
 {
 	GENERATED_BODY()
diff --git a/Engine/Plugins/Runtime/StateTree/Source/StateTreeModule/Public/Blueprint/StateTreeNodeBlueprintBase.h b/Engine/Plugins/Runtime/StateTree/Source/StateTreeModule/Public/Blueprint/StateTreeNodeBlueprintBase.h
--- a/Engine/Plugins/Runtime/StateTree/Source/StateTreeModule/Public/Blueprint/StateTreeNodeBlueprintBase.h
+++ b/Engine/Plugins/Runtime/StateTree/Source/StateTreeModule/Public/Blueprint/StateTreeNodeBlueprintBase.h
@@ -37,7 +37,7 @@ struct FStateTreeBlueprintExternalDataHandle
 };
 
 
-UCLASS(MinimalAPI, Abstract)
+UCLASS(MinimalAPI, Abstract, meta = (DisallowLevelActorReference = true))
 class UStateTreeNodeBlueprintBase : public UObject
 {
 	GENERATED_BODY()
diff --git a/Engine/Source/Editor/DetailCustomizations/Private/ComponentReferenceCustomization.cpp b/Engine/Source/Editor/DetailCustomizations/Private/ComponentReferenceCustomization.cpp
--- a/Engine/Source/Editor/DetailCustomizations/Private/ComponentReferenceCustomization.cpp
+++ b/Engine/Source/Editor/DetailCustomizations/Private/ComponentReferenceCustomization.cpp
@@ -51,13 +51,14 @@
 #include "Widgets/SNullWidget.h"
 #include "Widgets/Text/STextBlock.h"
 
-static const FName NAME_AllowAnyActor = "AllowAnyActor";
-static const FName NAME_AllowedClasses = "AllowedClasses";
-static const FName NAME_DisallowedClasses = "DisallowedClasses";
-static const FName NAME_UseComponentPicker = "UseComponentPicker";
-
 #define LOCTEXT_NAMESPACE "ComponentReferenceCustomization"
 
+static const FLazyName NAME_AllowAnyActor = FLazyName("AllowAnyActor");
+static const FLazyName NAME_AllowedClasses = FLazyName("AllowedClasses");
+static const FLazyName NAME_DisallowedClasses = FLazyName("DisallowedClasses");
+static const FLazyName NAME_UseComponentPicker = FLazyName("UseComponentPicker");
+static const FLazyName NAME_DisallowLevelActorReference = FLazyName("DisallowLevelActorReference");
+
 TSharedRef<IPropertyTypeCustomization> FComponentReferenceCustomization::MakeInstance()
 {
 	return MakeShareable( new FComponentReferenceCustomization);
@@ -207,10 +208,12 @@ void FComponentReferenceCustomization::BuildComboBox()
 {
 	TAttribute<bool> IsEnabledAttribute(this, &FComponentReferenceCustomization::CanEdit);
 	TAttribute<FText> TooltipAttribute;
+
+	TArray<UObject*> ObjectList;
+	PropertyHandle->GetOuterObjects(ObjectList);
+
 	if (PropertyHandle->GetMetaDataProperty()->HasAnyPropertyFlags(CPF_EditConst | CPF_DisableEditOnTemplate))
 	{
-		TArray<UObject*> ObjectList;
-		PropertyHandle->GetOuterObjects(ObjectList);
 
 		// If there is no objects, that means we must have a struct asset managing this property
 		if (ObjectList.Num() == 0)
@@ -221,7 +224,7 @@ void FComponentReferenceCustomization::BuildComboBox()
 		else
 		{
 			// Go through all the found objects and see if any are a CDO, we can't set an actor in a CDO default.
-			for (UObject* Obj : ObjectList)
+			for (const UObject* Obj : ObjectList)
 			{
 				if (Obj->IsTemplate() && !Obj->IsA<ALevelScriptActor>())
 				{
@@ -234,6 +237,22 @@ void FComponentReferenceCustomization::BuildComboBox()
 		}
 	}
 
+	if (IsEnabledAttribute.Get())
+	{
+		for (const UObject* Obj : ObjectList)
+		{
+			if (const UClass* Class = Obj->GetClass())
+			{
+				if (Class->GetBoolMetaDataHierarchical(NAME_DisallowLevelActorReference))
+				{
+					IsEnabledAttribute.Set(false);
+					TooltipAttribute.Set(LOCTEXT("ObjectDisallowLevelActorReference", "Object doesn't allow level actor reference"));
+					break;
+				}
+			}
+		}
+	}
+
 	TSharedRef<SVerticalBox> ObjectContent = SNew(SVerticalBox);
 	if (bAllowAnyActor)
 	{
diff --git a/Engine/Source/Editor/PropertyEditor/Private/UserInterface/PropertyEditor/SPropertyEditorAsset.cpp b/Engine/Source/Editor/PropertyEditor/Private/UserInterface/PropertyEditor/SPropertyEditorAsset.cpp
--- a/Engine/Source/Editor/PropertyEditor/Private/UserInterface/PropertyEditor/SPropertyEditorAsset.cpp
+++ b/Engine/Source/Editor/PropertyEditor/Private/UserInterface/PropertyEditor/SPropertyEditorAsset.cpp
@@ -47,6 +47,8 @@
 
 #define LOCTEXT_NAMESPACE "PropertyEditor"
 
+static const FLazyName NAME_DisallowLevelActorReference = FLazyName("DisallowLevelActorReference");
+
 // Helper to retrieve the correct property that has the applicable metadata.
 static const FProperty* GetActualMetadataProperty(const FProperty* Property)
 {
@@ -540,15 +542,15 @@ void SPropertyEditorAsset::Construct(const FArguments& InArgs, const TSharedPtr<
 
 	if (Property)
 	{
+		TArray<UObject*> ObjectList;
+		if (PropertyEditor.IsValid())
+		{
+			PropertyEditor->GetPropertyHandle()->GetOuterObjects(ObjectList);
+		}
+
 		const FProperty* PropToConsider = GetActualMetadataProperty(Property);
 		if (PropToConsider->HasAnyPropertyFlags(CPF_EditConst | CPF_DisableEditOnTemplate))
 		{
-			TArray<UObject*> ObjectList;
-			if (PropertyEditor.IsValid())
-			{
-				PropertyEditor->GetPropertyHandle()->GetOuterObjects(ObjectList);
-			}
-
 			// NOTE: This code decides that 99% of structs are "defaults" which is not technically correct, 
 			// but we want to stop hard actor references from being set in places like data tables without banning soft references.
 			// The actor check should get refactored to be independent of EditOnTemplate and do more explicit checks for world-owned object references
@@ -560,7 +562,7 @@ void SPropertyEditorAsset::Construct(const FArguments& InArgs, const TSharedPtr<
 			else
 			{
 				// Go through all the found objects and see if any are a CDO, we can't set an actor in a CDO default.
-				for (UObject* Obj : ObjectList)
+				for (const UObject* Obj : ObjectList)
 				{
 					if (Obj->IsTemplate() && !Obj->IsA<ALevelScriptActor>())
 					{
@@ -571,6 +573,23 @@ void SPropertyEditorAsset::Construct(const FArguments& InArgs, const TSharedPtr<
 
 				}
 			}
+
+		}
+
+		if (bIsActor && IsEnabledAttribute.Get())
+		{
+			for (const UObject* Obj : ObjectList)
+			{
+				if (const UClass* Class = Obj->GetClass())
+				{
+					if (Class->GetBoolMetaDataHierarchical(NAME_DisallowLevelActorReference))
+					{
+						IsEnabledAttribute.Set(false);
+						TooltipAttribute.Set(LOCTEXT("ObjectDisallowLevelActorReference", "Object doesn't allow level actor reference"));
+						break;
+					}
+				}
+			}
 		}
 	}
 
@@ -797,7 +816,7 @@ void SPropertyEditorAsset::Construct(const FArguments& InArgs, const TSharedPtr<
 			PropertyCustomizationHelpers::MakeBrowseButton(
 				OnBrowseDelegate,
 				TAttribute<FText>( this, &SPropertyEditorAsset::GetOnBrowseToolTip ),
-				true,
+				IsEnabledAttribute,
 				bIsActor,
 				IconDelegate
 				)
