Expose RequestForceRefresh via the public DetailsView interface.

diff --git a/Engine/Source/Editor/PropertyEditor/Private/IDetailsViewPrivate.h b/Engine/Source/Editor/PropertyEditor/Private/IDetailsViewPrivate.h
--- a/Engine/Source/Editor/PropertyEditor/Private/IDetailsViewPrivate.h
+++ b/Engine/Source/Editor/PropertyEditor/Private/IDetailsViewPrivate.h
@@ -48,12 +48,6 @@ public:
 	 */
 	virtual TSharedPtr<IPropertyUtilities> GetPropertyUtilities() = 0;
 
-	/** Request the details view to be refreshed (new widgets generated) with the current set of objects on the next Tick */
-	virtual void RequestForceRefresh() = 0;
-
-	/** Causes the details view to be refreshed (new widgets generated) with the current set of objects */
-	virtual void ForceRefresh() = 0;
-
 	/** Causes the details view to move the scroll offset (by item)
 	 * @param DeltaOffset	We add this value to the current scroll offset if the result is in the scrolling range
 	 */
diff --git a/Engine/Source/Editor/PropertyEditor/Public/IDetailsView.h b/Engine/Source/Editor/PropertyEditor/Public/IDetailsView.h
--- a/Engine/Source/Editor/PropertyEditor/Public/IDetailsView.h
+++ b/Engine/Source/Editor/PropertyEditor/Public/IDetailsView.h
@@ -276,7 +276,10 @@ public:
 	/** Sets the tab manager of the hosting toolkit (can be nullptr if the details panel is not hosted within a tab) */
 	virtual void SetHostTabManager(TSharedPtr<FTabManager> InTabManager) = 0;
 
-	/** Force refresh */
+	/** Request the details view to be refreshed (new widgets generated) with the current set of objects on the next Tick */
+	virtual void RequestForceRefresh() = 0;
+
+	/** Causes the details view to be refreshed (new widgets generated) with the current set of objects */
 	virtual void ForceRefresh() = 0;
 
 	/** Invalidates any cached state without necessarily doing a complete rebuild. */
