Add a GetExpansionStructFunc UCLASS meta tag with which it's possible to use a member struct of a wrapper class as expansion root.

diff --git a/Engine/Source/Editor/PropertyEditor/Private/SDetailsViewBase.cpp b/Engine/Source/Editor/PropertyEditor/Private/SDetailsViewBase.cpp
--- a/Engine/Source/Editor/PropertyEditor/Private/SDetailsViewBase.cpp
+++ b/Engine/Source/Editor/PropertyEditor/Private/SDetailsViewBase.cpp
@@ -669,7 +669,7 @@ void SDetailsViewBase::UpdatePropertyMapsWithItemExpansions()
 	for (TSharedPtr<FComplexPropertyNode>& RootPropertyNode : GetRootNodes())
 	{
 		check(RootPropertyNode.IsValid());
-		UStruct* BaseStruct = RootPropertyNode->FindComplexParent()->GetBaseStructure();
+		UStruct* BaseStruct = GetExpansionStruct(RootPropertyNode.ToSharedRef());
 		FStringPrefixTree& PrefixTree = ExpansionStates.Add(BaseStruct);
 		GetExpansionState(BaseStruct, PrefixTree);
 	}
@@ -680,7 +680,7 @@ void SDetailsViewBase::UpdatePropertyMapsWithItemExpansions()
 		for (TSharedPtr<FComplexPropertyNode>& ExternalRootPropertyNode : ExternalRootPropertyNodes)
 		{
 			check(ExternalRootPropertyNode.IsValid());
-			UStruct* BaseStruct = ExternalRootPropertyNode->FindComplexParent()->GetBaseStructure();
+			UStruct* BaseStruct = GetExpansionStruct(ExternalRootPropertyNode.ToSharedRef());
 			FStringPrefixTree& PrefixTree = ExpansionStates.Add(BaseStruct);
 			GetExpansionState(BaseStruct, PrefixTree);
 		}
@@ -1572,6 +1572,55 @@ void SDetailsViewBase::RestorePreSearchExpandedItems()
 	PreSearchExpandedCategories.Clear();
 }
 
+UStruct* SDetailsViewBase::GetExpansionStruct(const TSharedRef<FPropertyNode>& StartNode)
+{
+	FComplexPropertyNode* ComplexParent = StartNode->FindComplexParent();
+	UStruct* BaseStructure = ComplexParent->GetBaseStructure();
+
+	FString FuncName;
+	if (BaseStructure && BaseStructure->GetStringMetaDataHierarchical("GetExpansionStructFunc", &FuncName))
+	{
+		if (const FObjectPropertyNode* ObjectNode = ComplexParent->AsObjectNode())
+		{
+			const UClass* BaseClass = Cast<UClass>(BaseStructure);
+			const int32 InstancesNum = ObjectNode->GetInstancesNum();
+
+			if (InstancesNum > 0 && BaseClass)
+			{
+				if (UFunction* Function = BaseClass->FindFunctionByName(FName(FuncName)))
+				{
+					if (Function->HasAnyFunctionFlags(FUNC_Static))
+					{
+						struct FFunctionParams
+						{
+							TArray<UObject*> Objects;
+							const UStruct* Struct;
+						} Params;
+
+						Params.Objects.Reserve(InstancesNum);
+						for (int32 i = 0; i < InstancesNum; ++i)
+						{
+							if (UObject* Object = ObjectNode->GetInstanceAsUObject(i).Get())
+							{
+								Params.Objects.Add(Object);
+							}
+						}
+
+						BaseStructure->ProcessEvent(Function, &Params);
+
+						if (Params.Struct)
+						{
+							return const_cast<UStruct*>(Params.Struct);
+						}
+					}
+				}
+			}
+		}
+	}
+
+	return BaseStructure;
+}
+
 void SDetailsViewBase::SaveExpandedItems(TSharedRef<FPropertyNode> StartNode)
 {
 	TRACE_CPUPROFILER_EVENT_SCOPE(SDetailsView::SaveExpandedItems);
@@ -1584,7 +1633,7 @@ void SDetailsViewBase::SaveExpandedItems(TSharedRef<FPropertyNode> StartNode)
 		return;
 	}
 
-	UStruct* BestBaseStruct = StartNode->FindComplexParent()->GetBaseStructure();
+	UStruct* BestBaseStruct = GetExpansionStruct(StartNode);
 
 	FStringPrefixTree ExpandedPropertyItemSet;
 	GetExpandedItems(StartNode, ExpandedPropertyItemSet);
@@ -1687,7 +1736,7 @@ void SDetailsViewBase::RestoreExpandedItems(TSharedRef<FPropertyNode> StartNode,
 		return;
 	}
 
-	UStruct* BestBaseStruct = StartNode->FindComplexParent()->GetBaseStructure();
+	UStruct* BestBaseStruct = GetExpansionStruct(StartNode);
 
 	if (OptionalExpansionStates)
 	{
diff --git a/Engine/Source/Editor/PropertyEditor/Private/SDetailsViewBase.h b/Engine/Source/Editor/PropertyEditor/Private/SDetailsViewBase.h
--- a/Engine/Source/Editor/PropertyEditor/Private/SDetailsViewBase.h
+++ b/Engine/Source/Editor/PropertyEditor/Private/SDetailsViewBase.h
@@ -293,6 +293,7 @@ protected:
 	virtual void CustomUpdatePropertyMap(TSharedPtr<FDetailLayoutBuilderImpl>& InDetailLayout) {}
 
 	void GetExpansionState(UStruct* BaseStruct, FStringPrefixTree& OutPrefixTree) const;
+	UStruct* GetExpansionStruct(const TSharedRef<FPropertyNode>& StartNode);
 
 	/** Updates the property map while keeping the expansion state of all items */
 	void UpdatePropertyMapsWithItemExpansions();
