[PropertyCore] Add safeguard to not retrieve the object address when enumerating the structs to modify if we do not have a one for one mapping

* It was changed recently to support propagating values changes in structs when it is own by an object (example: Instanced Property Bag owned by a BP template, to propagate the change to all its BP instances)
* But it was assuming a 1 for 1 mapping between the ComplexNode (property node for the structure) and the ObjectNode (property node for the owning object)
* Since `GetNumInstances` on property nodes is virtual, there is no guarentee that there is a 1 for 1 mapping. So we will only try to gather the owning object address when the number of instances matches between the two, otherwise we fallback to previous behavior

#rnx

diff --git a/Engine/Source/Editor/PropertyEditor/Private/PropertyHandleImpl.cpp b/Engine/Source/Editor/PropertyEditor/Private/PropertyHandleImpl.cpp
--- a/Engine/Source/Editor/PropertyEditor/Private/PropertyHandleImpl.cpp
+++ b/Engine/Source/Editor/PropertyEditor/Private/PropertyHandleImpl.cpp
@@ -68,19 +68,21 @@ FPropertyValueImpl::FPropertyValueImpl( TSharedPtr<FPropertyNode> InPropertyNode
 {
 }
 
-void FPropertyValueImpl::EnumerateObjectsToModify( FPropertyNode* InPropertyNode, const EnumerateObjectsToModifyFuncRef& InObjectsToModifyCallback ) const
+void FPropertyValueImpl::EnumerateObjectsToModify(FPropertyNode* InPropertyNode, const EnumerateObjectsToModifyFuncRef& InObjectsToModifyCallback) const
 {
 	// EnumerateObjectsToModify catch the first ComplexNode and iterates on its number of instances to get the property address in memory.
 	// A complex node can be an object node or a struct node (at this time of writing).
 	// If the complex node is an object, we also extract the address of this object.
 	// If the complex node is a struct, we still want to know if it is owned by an object, so we look further up to find the first object node (if it exists).
-	// If it exists, we have a complex struct owned by an object, so we extract the object address.
+	// If it exists, we might have a complex struct owned by an object, so we try to extract the object address.
 	FComplexPropertyNode* ComplexNode = InPropertyNode->FindComplexParent();
-	FObjectPropertyNode* ObjectNode = InPropertyNode->FindObjectItemParent();
 	if (ComplexNode)
 	{
 		const bool bIsStruct = FComplexPropertyNode::EPT_StandaloneStructure == ComplexNode->GetPropertyType();
 		const int32 NumInstances = ComplexNode->GetInstancesNum();
+
+		const FObjectPropertyNode* ObjectNode = bIsStruct ? InPropertyNode->FindObjectItemParent() : nullptr;
+		const int32 NumObjectInstances = ObjectNode ? ObjectNode->GetInstancesNum() : 0;
 		for (int32 Index = 0; Index < NumInstances; ++Index)
 		{
 			UObject* Object = nullptr;
@@ -88,7 +90,10 @@ void FPropertyValueImpl::EnumerateObjectsToModify( FPropertyNode* InPropertyNode
 			if (bIsStruct)
 			{
 				StructAddress = ComplexNode->GetMemoryOfInstance(Index);
-				if (ObjectNode)
+
+				// Make sure we have the same number of structs and objects, so there is a 1 for 1 mapping between the struct and its owning object.
+				// If we don't, we don't have a clear way to resolve this, so we won't try to retrieve the object address (which was the previous behavior).
+				if (ObjectNode && NumInstances == NumObjectInstances)
 				{
 					Object = ObjectNode->GetInstanceAsUObject(Index).Get();
 				}
