Add ShowObjectRootClass so that properties leveraging the ClassViewer can include /Script/CoreUObject.Object in their list

[RN] meta=(ShowObjectRootClass) can now be added to class reference properties to include Object in the list of selectable classes

diff --git a/Engine/Source/Editor/PropertyEditor/Private/PropertyEditorUtils.cpp b/Engine/Source/Editor/PropertyEditor/Private/PropertyEditorUtils.cpp
--- a/Engine/Source/Editor/PropertyEditor/Private/PropertyEditorUtils.cpp
+++ b/Engine/Source/Editor/PropertyEditor/Private/PropertyEditorUtils.cpp
@@ -255,4 +255,6 @@ namespace PropertyEditorUtils
 			}
 		}
 	}
+
+	const FLazyName ShowObjectRootClassKey(TEXT("ShowObjectRootClass"));
 }
diff --git a/Engine/Source/Editor/PropertyEditor/Private/UserInterface/PropertyEditor/SPropertyEditorClass.cpp b/Engine/Source/Editor/PropertyEditor/Private/UserInterface/PropertyEditor/SPropertyEditorClass.cpp
--- a/Engine/Source/Editor/PropertyEditor/Private/UserInterface/PropertyEditor/SPropertyEditorClass.cpp
+++ b/Engine/Source/Editor/PropertyEditor/Private/UserInterface/PropertyEditor/SPropertyEditorClass.cpp
@@ -148,6 +148,7 @@ void SPropertyEditorClass::Construct(const FArguments& InArgs, const TSharedPtr<
 		bShowViewOptions = Property->GetOwnerProperty()->HasMetaData(TEXT("HideViewOptions")) ? false : true;
 		bShowTree = Property->GetOwnerProperty()->HasMetaData(TEXT("ShowTreeView"));
 		bShowDisplayNames = Property->GetOwnerProperty()->HasMetaData(TEXT("ShowDisplayNames"));
+		bShowObjectRootClass = Property->GetOwnerProperty()->HasMetaData(PropertyEditorUtils::ShowObjectRootClassKey);
 
 		if (RequiredInterface != nullptr)
 		{
@@ -213,6 +214,7 @@ void SPropertyEditorClass::Construct(const FArguments& InArgs, const TSharedPtr<
 		bShowViewOptions = InArgs._ShowViewOptions;
 		bShowTree = InArgs._ShowTree;
 		bShowDisplayNames = InArgs._ShowDisplayNames;
+		bShowObjectRootClass = false;
 		AllowedClassFilters = InArgs._AllowedClasses;
 		DisallowedClassFilters = InArgs._DisallowedClasses;
 		SelectedClass = InArgs._SelectedClass;
@@ -312,6 +314,7 @@ void SPropertyEditorClass::CreateClassFilter(const TArray<TSharedRef<IClassViewe
 	ClassViewerOptions.DisplayMode = bShowTree ? EClassViewerDisplayMode::TreeView : EClassViewerDisplayMode::ListView;
 	ClassViewerOptions.bAllowViewOptions = bShowViewOptions;
 	ClassViewerOptions.ClassFilters.Append(InClassFilters);
+	ClassViewerOptions.bShowObjectRootClass	= bShowObjectRootClass;
 
 	TSharedRef<FPropertyEditorClassFilter> PropEdClassFilter = MakeShared<FPropertyEditorClassFilter>();
 	PropEdClassFilter->ClassPropertyMetaClass = MetaClass;
diff --git a/Engine/Source/Editor/PropertyEditor/Private/UserInterface/PropertyEditor/SPropertyEditorClass.h b/Engine/Source/Editor/PropertyEditor/Private/UserInterface/PropertyEditor/SPropertyEditorClass.h
--- a/Engine/Source/Editor/PropertyEditor/Private/UserInterface/PropertyEditor/SPropertyEditorClass.h
+++ b/Engine/Source/Editor/PropertyEditor/Private/UserInterface/PropertyEditor/SPropertyEditorClass.h
@@ -142,6 +142,8 @@ private:
 	bool bShowTree;
 	/** Should we prettify class names on the class picker? (ie show their display name) */
 	bool bShowDisplayNames;
+	/** Should we include /Script/CoreUObject.Object in the list? */
+	bool bShowObjectRootClass;
 	/** Classes that can be picked with this property */
 	TArray<const UClass*> AllowedClassFilters;
 	/** Classes that can NOT be picked with this property */
diff --git a/Engine/Source/Editor/PropertyEditor/Public/PropertyEditorUtils.h b/Engine/Source/Editor/PropertyEditor/Public/PropertyEditorUtils.h
--- a/Engine/Source/Editor/PropertyEditor/Public/PropertyEditorUtils.h
+++ b/Engine/Source/Editor/PropertyEditor/Public/PropertyEditorUtils.h
@@ -1,4 +1,4 @@
-﻿// Copyright Epic Games, Inc. All Rights Reserved.
+// Copyright Epic Games, Inc. All Rights Reserved.
 
 #pragma once
 
@@ -37,4 +37,14 @@ namespace PropertyEditorUtils
 	 * @param	ObjectClass			The class to add to the disallowed classes if there is no allowed classes
 	 */
 	PROPERTYEDITOR_API void GetAllowedAndDisallowedClasses(const TArray<UObject*>& ObjectList, const FProperty& Property, TArray<const UClass*>& AllowedClasses, TArray<const UClass*>& DisallowedClasses, bool bExactClass, const UClass* ObjectClass = UObject::StaticClass());
+
+	/** 
+	 * Metadata key for specifying that the property editor's class picker should include the root Object 
+	 * type (/Script/CoreUObject.Object) in the displayed list of classes. By default Object is filtered out.
+	 * This metadata key has no value, behavior is determined by its presence or absence, e.g.:
+	 * 
+	 * UPROPERTY(meta=(ShowObjectRootClass))
+	 * TSoftClassPtr<UObject> Value;
+	 */
+	extern const FLazyName ShowObjectRootClassKey;
 }
