[UE] Detect and deep copy inner object property values for non-explicit instanced references when duplicating array property value entries in the Property Editor.

diff --git a/Engine/Source/Editor/PropertyEditor/Private/PropertyHandleImpl.cpp b/Engine/Source/Editor/PropertyEditor/Private/PropertyHandleImpl.cpp
--- a/Engine/Source/Editor/PropertyEditor/Private/PropertyHandleImpl.cpp
+++ b/Engine/Source/Editor/PropertyEditor/Private/PropertyHandleImpl.cpp
@@ -2336,7 +2336,7 @@ void FPropertyValueImpl::DuplicateChild( TSharedPtr<FPropertyNode> ChildNodeToDu
 		}
 		
 		FScriptArrayHelper ArrayHelper(ArrayProperty, Addr);
-		FPropertyNode::DuplicateArrayEntry(NodeProperty, ArrayHelper, Index);
+		FPropertyNode::DuplicateArrayEntry(Obj, NodeProperty, ArrayHelper, Index);
 		
 		FPropertyChangedEvent ChangeEvent(ParentNode->GetProperty(), EPropertyChangeType::Duplicate, MakeArrayView(TopLevelObjects));
 		ChangeEvent.SetArrayIndexPerObject(ArrayIndicesPerObject);
diff --git a/Engine/Source/Editor/PropertyEditor/Private/PropertyNode.cpp b/Engine/Source/Editor/PropertyEditor/Private/PropertyNode.cpp
--- a/Engine/Source/Editor/PropertyEditor/Private/PropertyNode.cpp
+++ b/Engine/Source/Editor/PropertyEditor/Private/PropertyNode.cpp
@@ -3888,7 +3888,7 @@ void FPropertyNode::GatherInstancesAffectedByContainerPropertyChange(UObject* Mo
 	}
 }
 
-void FPropertyNode::DuplicateArrayEntry(FProperty* NodeProperty, FScriptArrayHelper& ArrayHelper, int32 Index)
+void FPropertyNode::DuplicateArrayEntry(UObject* ModifiedObject, FProperty* NodeProperty, FScriptArrayHelper& ArrayHelper, int32 Index)
 {
 	ArrayHelper.InsertValues(Index);
 
@@ -3900,34 +3900,38 @@ void FPropertyNode::DuplicateArrayEntry(FProperty* NodeProperty, FScriptArrayHel
 	// Copy the selected item's value to the new item.
 	NodeProperty->CopyCompleteValue(DestAddress, SrcAddress);
 
-	if (FObjectProperty* ObjProp = CastField<FObjectProperty>(NodeProperty))
+	bool bHasDynamicInstancedReferences = false;
+	if (ModifiedObject && ModifiedObject->GetClass()->ShouldUseDynamicSubobjectInstancing())
 	{
-		if (ObjProp->HasAnyPropertyFlags(CPF_InstancedReference))
-		{
-			UObject* CurrentObject = ObjProp->GetObjectPropertyValue(DestAddress);
-
-			// Make a deep copy
-			UObject* DuplicatedObject = DuplicateObject(CurrentObject, CurrentObject ? CurrentObject->GetOuter() : nullptr);
-			ObjProp->SetObjectPropertyValue(SrcAddress, DuplicatedObject);
-		}
+		bHasDynamicInstancedReferences = true;
 	}
-	else if (NodeProperty->HasAnyPropertyFlags(CPF_ContainsInstancedReference))
+
+	auto IsDynamicInstancedReference = [bHasDynamicInstancedReferences, ModifiedObject](const UObject* ReferencedObject)
 	{
-PRAGMA_DISABLE_DEPRECATION_WARNINGS
-		// If this is a container with instanced references within it the new entry will reference the old subobjects
-		// Go through and duplicate the subobjects so that each container has unique instances
-		FInstancedPropertyPath NodePropertyPath(NodeProperty);
-		FFindInstancedReferenceSubobjectHelper::ForEachInstancedSubObject<void*>(
-			NodePropertyPath,
-			SrcAddress,
-			[](const FInstancedSubObjRef& Ref, void* PropertyValueAddress)
+		return bHasDynamicInstancedReferences && ReferencedObject->IsInOuter(ModifiedObject);
+	};
+
+	// If this is a property with instanced references within it the new entry will reference the old subobjects
+	// Go through and duplicate the subobjects so that each new entry has unique instances
+	FPropertyVisitorPath Path;
+	Path.Push(FPropertyVisitorInfo(NodeProperty));
+	FPropertyVisitorData Data(SrcAddress, nullptr);
+	FPropertyVisitorContext VisitorContext(Path, Data);
+	NodeProperty->Visit(VisitorContext, [&IsDynamicInstancedReference](const FPropertyVisitorContext& Context) -> EPropertyVisitorControlFlow
+	{
+		if (const FObjectPropertyBase* ObjectProperty = CastField<FObjectPropertyBase>(Context.Path.Top().Property))
+		{
+			UObject* ReferencedObject = ObjectProperty->GetObjectPropertyValue(Context.Data.PropertyData);
+			if (ReferencedObject && (ObjectProperty->HasAnyPropertyFlags(CPF_InstancedReference) || IsDynamicInstancedReference(ReferencedObject)))
 			{
-				UObject* Obj = Ref;
-				((FObjectProperty*)Ref.PropertyPath.Head())->SetObjectPropertyValue(PropertyValueAddress, DuplicateObject(Obj, Obj->GetOuter()));
+				// Make a deep copy
+				UObject* DuplicatedObject = DuplicateObject(ReferencedObject, ReferencedObject->GetOuter());
+				ObjectProperty->SetObjectPropertyValue(Context.Data.PropertyData, DuplicatedObject);
 			}
-		);
-PRAGMA_ENABLE_DEPRECATION_WARNINGS
-	}
+		}
+
+		return EPropertyVisitorControlFlow::StepInto;
+	});
 }
 
 void FPropertyNode::PropagateContainerPropertyChange( UObject* ModifiedObject, const void* OriginalContainerAddr, const TArray<UObject*>& AffectedInstances, EPropertyArrayChangeType::Type ChangeType, int32 Index, int32 SwapIndex /*= INDEX_NONE*/)
@@ -4008,7 +4012,7 @@ void FPropertyNode::PropagateContainerPropertyChange( UObject* ModifiedObject, c
 				ArrayHelper.RemoveValues(ArrayIndex, 1);
 				break;
 			case EPropertyArrayChangeType::Duplicate:
-				DuplicateArrayEntry(NodeProperty, ArrayHelper, ArrayIndex);
+				DuplicateArrayEntry(InstanceToChange, NodeProperty, ArrayHelper, ArrayIndex);
 				break;
 			case EPropertyArrayChangeType::Swap:
 				if (SwapIndex != INDEX_NONE)
diff --git a/Engine/Source/Editor/PropertyEditor/Private/PropertyNode.h b/Engine/Source/Editor/PropertyEditor/Private/PropertyNode.h
--- a/Engine/Source/Editor/PropertyEditor/Private/PropertyNode.h
+++ b/Engine/Source/Editor/PropertyEditor/Private/PropertyNode.h
@@ -744,7 +744,7 @@ public:
 	/**
 	 * Helper function to centralize logic for duplcating an array entry and ensuring that instanced object references are correctly handled
 	 */
-	static void DuplicateArrayEntry(FProperty* NodeProperty, FScriptArrayHelper& ArrayHelper, int32 Index);
+	static void DuplicateArrayEntry(UObject* ModifiedObject, FProperty* NodeProperty, FScriptArrayHelper& ArrayHelper, int32 Index);
 
 	/**
 	 * Gather the list of all instances that will be affected by a container property change
