Use the function name as the transaction name for editor utilities buttons.

diff --git a/Engine/Source/Editor/PropertyEditor/Private/PropertyCustomizationHelpers.cpp b/Engine/Source/Editor/PropertyEditor/Private/PropertyCustomizationHelpers.cpp
--- a/Engine/Source/Editor/PropertyEditor/Private/PropertyCustomizationHelpers.cpp
+++ b/Engine/Source/Editor/PropertyEditor/Private/PropertyCustomizationHelpers.cpp
@@ -409,8 +409,13 @@ namespace PropertyCustomizationHelpers
 					using namespace UE::Reflection;
 					if (UFunction* Function = InWeakFunction.Get())
 					{
-						// @todo: Consider naming the transaction scope after the fully qualified function name for better UX
-						FScopedTransaction Transaction(LOCTEXT("ExecuteCallInEditorMethod", "Call In Editor Action"));
+						FText TransactionName = ObjectTools::GetUserFacingFunctionName(Function);
+						if (TransactionName.IsEmptyOrWhitespace())
+						{
+							TransactionName = LOCTEXT("ExecuteCallInEditorMethod", "Call In Editor Action");
+						}
+
+						FScopedTransaction Transaction(TransactionName);
 						TStrongObjectPtr<UFunction> CallingFunction(Function);
 
 						if (Function->HasMetaData(NAME_WorldContext)
