[UE] Fix a regression that unnecessarily deep-copies nested subobjects on property editor array node duplication.

Notes:
- Also avoids an infinite recursion case while visiting circular object ptr references, which can lead to a stack overflow.

diff --git a/Engine/Source/Editor/PropertyEditor/Private/PropertyNode.cpp b/Engine/Source/Editor/PropertyEditor/Private/PropertyNode.cpp
--- a/Engine/Source/Editor/PropertyEditor/Private/PropertyNode.cpp
+++ b/Engine/Source/Editor/PropertyEditor/Private/PropertyNode.cpp
@@ -3916,7 +3916,7 @@ void FPropertyNode::DuplicateArrayEntry(UObject* ModifiedObject, FProperty* Node
 	FPropertyVisitorPath Path;
 	Path.Push(FPropertyVisitorInfo(NodeProperty));
 	FPropertyVisitorData Data(SrcAddress, nullptr);
-	FPropertyVisitorContext VisitorContext(Path, Data);
+	FPropertyVisitorContext VisitorContext(Path, Data, FPropertyVisitorContext::EScope::ObjectRefs);
 	NodeProperty->Visit(VisitorContext, [&IsDynamicInstancedReference](const FPropertyVisitorContext& Context) -> EPropertyVisitorControlFlow
 	{
 		if (const FObjectPropertyBase* ObjectProperty = CastField<FObjectPropertyBase>(Context.Path.Top().Property))
@@ -3928,8 +3928,12 @@ void FPropertyNode::DuplicateArrayEntry(UObject* ModifiedObject, FProperty* Node
 				UObject* DuplicatedObject = DuplicateObject(ReferencedObject, ReferencedObject->GetOuter());
 				ObjectProperty->SetObjectPropertyValue(Context.Data.PropertyData, DuplicatedObject);
 			}
+			
+			// Continue on to the next property/item.
+			return EPropertyVisitorControlFlow::StepOver;
 		}
 
+		// In the case of container/optional properties, this will recurse into any inner value(s).
 		return EPropertyVisitorControlFlow::StepInto;
 	});
 }
