[UE] Add dynamic reference support to object property value imports from T3D.

Notes:
- Fixes a regression with legacy clipboard workflows when used to paste copied values into dynamic reference properties in the details view.

diff --git a/Engine/Source/Editor/PropertyEditor/Private/PropertyHandleImpl.cpp b/Engine/Source/Editor/PropertyEditor/Private/PropertyHandleImpl.cpp
--- a/Engine/Source/Editor/PropertyEditor/Private/PropertyHandleImpl.cpp
+++ b/Engine/Source/Editor/PropertyEditor/Private/PropertyHandleImpl.cpp
@@ -4435,6 +4435,14 @@ FPropertyAccess::Result FPropertyHandleObject::SetValueFromFormattedString(const
 			{
 				QualifiedClass = QualifiedObject->GetClass();
 			}
+			
+			// Only instance new objects if required/directed
+			if ((Flags & EPropertyValueSetFlags::InstanceObjects) != 0
+				&& !NodeProperty->HasAnyPropertyFlags(CPF_InstancedReference)
+				&& !PropertyNodePin->HasNodeFlags(EPropertyNodeFlags::EditInlineNew))
+			{
+				Flags &= ~EPropertyValueSetFlags::InstanceObjects;
+			}
 
 			const FString& AllowedClassesString = NodeProperty->GetMetaData("AllowedClasses");
 			TArray<FString> AllowedClassNames;
diff --git a/Engine/Source/Runtime/CoreUObject/Private/UObject/PropertyBaseObject.cpp b/Engine/Source/Runtime/CoreUObject/Private/UObject/PropertyBaseObject.cpp
--- a/Engine/Source/Runtime/CoreUObject/Private/UObject/PropertyBaseObject.cpp
+++ b/Engine/Source/Runtime/CoreUObject/Private/UObject/PropertyBaseObject.cpp
@@ -426,7 +426,7 @@ const TCHAR* FObjectPropertyBase::ImportText_Internal( const TCHAR* InBuffer, vo
 
 	bool bOk = ParseObjectPropertyValue(this, Parent, PropertyClass, PortFlags, Buffer, Result, FUObjectThreadContext::Get().GetSerializeContext());
 
-	if (Result && (PortFlags & PPF_InstanceSubobjects) != 0 && HasAnyPropertyFlags(CPF_InstancedReference))
+	if (Result && (PortFlags & PPF_InstanceSubobjects) != 0 && (HasAnyPropertyFlags(CPF_InstancedReference) || (Parent && Parent->GetClass()->ShouldUseDynamicSubobjectInstancing())))
 	{
 		FName DesiredName = Result->GetFName();
 
