[Core] Fix bad use of remove during iteration

#tests Using the editor, running PIE and game client, making builds

diff --git a/Engine/Plugins/Cameras/CameraShakePreviewer/Source/CameraShakePreviewer/Private/CameraShakePreviewerModule.cpp b/Engine/Plugins/Cameras/CameraShakePreviewer/Source/CameraShakePreviewer/Private/CameraShakePreviewerModule.cpp
--- a/Engine/Plugins/Cameras/CameraShakePreviewer/Source/CameraShakePreviewer/Private/CameraShakePreviewerModule.cpp
+++ b/Engine/Plugins/Cameras/CameraShakePreviewer/Source/CameraShakePreviewer/Private/CameraShakePreviewerModule.cpp
@@ -199,7 +199,7 @@ void FCameraShakePreviewerModule::OnLevelViewportClientListChanged()
 	{
 		if (!LevelViewportClients.Contains(It.Key()))
 		{
-			ViewportInfos.Remove(It.Key());
+			It.RemoveCurrent();
 		}
 	}
 	// Add recently added viewports that we don't know about yet.
diff --git a/Engine/Plugins/Media/PixelStreaming/Source/PixelStreaming/Private/AudioSink.cpp b/Engine/Plugins/Media/PixelStreaming/Source/PixelStreaming/Private/AudioSink.cpp
--- a/Engine/Plugins/Media/PixelStreaming/Source/PixelStreaming/Private/AudioSink.cpp
+++ b/Engine/Plugins/Media/PixelStreaming/Source/PixelStreaming/Private/AudioSink.cpp
@@ -7,14 +7,9 @@ namespace UE::PixelStreaming
 {
 	FAudioSink::~FAudioSink()
 	{
-		for (auto Iter = AudioConsumers.CreateIterator(); Iter; ++Iter)
+		for (IPixelStreamingAudioConsumer* AudioConsumer : AudioConsumers)
 		{
-			IPixelStreamingAudioConsumer* AudioConsumer = Iter.ElementIt->Value;
-			Iter.RemoveCurrent();
-			if (AudioConsumer != nullptr)
-			{
-				AudioConsumer->OnConsumerRemoved();
-			}
+			AudioConsumer->OnConsumerRemoved();
 		}
 	}
 
@@ -52,9 +47,8 @@ namespace UE::PixelStreaming
 
 	void FAudioSink::RemoveAudioConsumer(IPixelStreamingAudioConsumer* AudioConsumer)
 	{
-		if (AudioConsumers.Contains(AudioConsumer))
+		if (AudioConsumers.Remove(AudioConsumer) > 0)
 		{
-			AudioConsumers.Remove(AudioConsumer);
 			AudioConsumer->OnConsumerRemoved();
 		}
 	}
diff --git a/Engine/Plugins/Online/OnlineSubsystem/Source/Private/OnlineSubsystemModule.cpp b/Engine/Plugins/Online/OnlineSubsystem/Source/Private/OnlineSubsystemModule.cpp
--- a/Engine/Plugins/Online/OnlineSubsystem/Source/Private/OnlineSubsystemModule.cpp
+++ b/Engine/Plugins/Online/OnlineSubsystem/Source/Private/OnlineSubsystemModule.cpp
@@ -265,6 +265,8 @@ void FOnlineSubsystemModule::ShutdownOnlineSubsystem()
 		// Unloading the module will do proper cleanup
 		FName ModuleName = GetOnlineModuleName(It.Key().ToString(), ModuleRedirects);
 
+		It.RemoveCurrent(); // Remove safely from map otherwise UnloadModule will do it and kill our iteration
+
 		const bool bIsShutdown = true;
 		ModuleManager.UnloadModule(ModuleName, bIsShutdown);
 	} 
diff --git a/Engine/Plugins/Runtime/ComputeFramework/Source/ComputeFramework/Private/ComputeKernelShaderCompilationManager.cpp b/Engine/Plugins/Runtime/ComputeFramework/Source/ComputeFramework/Private/ComputeKernelShaderCompilationManager.cpp
--- a/Engine/Plugins/Runtime/ComputeFramework/Source/ComputeFramework/Private/ComputeKernelShaderCompilationManager.cpp
+++ b/Engine/Plugins/Runtime/ComputeFramework/Source/ComputeFramework/Private/ComputeKernelShaderCompilationManager.cpp
@@ -96,7 +96,7 @@ void FComputeKernelShaderCompilationManager::ProcessAsyncResults()
 		if (Results.FinishedJobs.Num() == Results.NumJobsQueued)
 		{
 			PendingFinalizeComputeKernelShaderMaps.Add(It.Key(), FComputeKernelShaderMapCompileResults(Results));
-			ComputeKernelShaderMapJobs.Remove(It.Key());
+			It.RemoveCurrent();
 		}
 	}
 
@@ -363,7 +363,7 @@ void FComputeKernelShaderCompilationManager::ProcessCompiledComputeKernelShaderM
 
 				// Cleanup shader jobs and compile tracking structures
 				ResultArray.Empty();
-				CompiledShaderMaps.Remove(ShaderMap->GetCompilingId());
+				ProcessIt.RemoveCurrent();
 			}
 
 			if (TimeBudget < 0)
diff --git a/Engine/Source/Developer/DirectoryWatcher/Private/Windows/DirectoryWatcherWindows.cpp b/Engine/Source/Developer/DirectoryWatcher/Private/Windows/DirectoryWatcherWindows.cpp
--- a/Engine/Source/Developer/DirectoryWatcher/Private/Windows/DirectoryWatcherWindows.cpp
+++ b/Engine/Source/Developer/DirectoryWatcher/Private/Windows/DirectoryWatcherWindows.cpp
@@ -81,23 +81,24 @@ bool FDirectoryWatcherWindows::RegisterDirectoryChangedCallback_Handle( const FS
 
 bool FDirectoryWatcherWindows::UnregisterDirectoryChangedCallback_Handle( const FString& Directory, FDelegateHandle InHandle )
 {
-	for (const auto& RequestPair : RequestMap)
+	for (TMap<FDirectoryWithFlags, FDirectoryWatchRequestWindows*>::TIterator RequestIt(RequestMap); RequestIt; ++RequestIt)
 	{
-		if (RequestPair.Key.Key == Directory)
+		if (RequestIt->Key.Key == Directory)
 		{
 			// There should be no NULL entries in the map
-			check(RequestPair.Value);
+			FDirectoryWatchRequestWindows* Watch = RequestIt->Value;
+			check(Watch);
 
-			if (RequestPair.Value->RemoveDelegate(InHandle))
+			if (Watch->RemoveDelegate(InHandle))
 			{
-				if (!RequestPair.Value->HasDelegates())
+				if (!Watch->HasDelegates())
 				{
 					// Remove from the active map and add to the pending delete list
-					RequestMap.Remove(RequestPair.Key);
-					RequestsPendingDelete.AddUnique(RequestPair.Value);
+					RequestIt.RemoveCurrent();
+					RequestsPendingDelete.AddUnique(Watch);
 
 					// Signal to end the watch which will mark this request for deletion
-					RequestPair.Value->EndWatchRequest();
+					Watch->EndWatchRequest();
 				}
 
 				return true;
diff --git a/Engine/Source/Developer/MessageLog/Private/Model/MessageLogModel.cpp b/Engine/Source/Developer/MessageLog/Private/Model/MessageLogModel.cpp
--- a/Engine/Source/Developer/MessageLog/Private/Model/MessageLogModel.cpp
+++ b/Engine/Source/Developer/MessageLog/Private/Model/MessageLogModel.cpp
@@ -4,9 +4,10 @@
 
 FMessageLogModel::~FMessageLogModel()
 {
-	for( auto It = NameToModelMap.CreateIterator(); It; ++It )
+	if (!NameToModelMap.IsEmpty())
 	{
-		UnregisterLogListingModel( It.Key() );
+		NameToModelMap.Empty();
+		Notify();
 	}
 }
 
diff --git a/Engine/Source/Editor/PropertyEditor/Private/DetailLayoutBuilderImpl.cpp b/Engine/Source/Editor/PropertyEditor/Private/DetailLayoutBuilderImpl.cpp
--- a/Engine/Source/Editor/PropertyEditor/Private/DetailLayoutBuilderImpl.cpp
+++ b/Engine/Source/Editor/PropertyEditor/Private/DetailLayoutBuilderImpl.cpp
@@ -897,9 +897,9 @@ FCustomPropertyTypeLayoutMap FDetailLayoutBuilderImpl::GetInstancedPropertyTypeL
 
 void FDetailLayoutBuilderImpl::RefreshNodeVisbility()
 {
-	for(TSet<FDetailTreeNode*>::TIterator It = TickableNodes.CreateIterator(); It; ++It)
+	for(FDetailTreeNode* Node : TickableNodes)
 	{
-		(*It)->RefreshVisibility();
+		Node->RefreshVisibility();
 	}
 }
 
@@ -1057,12 +1057,12 @@ void FDetailLayoutBuilderImpl::SetCustomProperty( const TSharedPtr<FPropertyNode
 	PropertyNode->SetNodeFlags( EPropertyNodeFlags::IsCustomized, true );
 }
 
-void FDetailLayoutBuilderImpl::Tick( float DeltaTime )
+void FDetailLayoutBuilderImpl::Tick(float DeltaTime)
 {
-	for( auto It = TickableNodes.CreateIterator(); It; ++It )
-	{
-		FDetailTreeNode* Node = *It;
+	TGuardValue Guard(bIsIteratingTickableNodes, true);
 
+	auto TickNode = [DeltaTime](FDetailTreeNode* Node)
+	{
 		// Skip ticking tree nodes which point to destroyed property nodes.
 		// This can happen when because the update order is this:
 		//	- update property nodes, calling DestroyTree(), and creating new nodes
@@ -1070,23 +1070,60 @@ void FDetailLayoutBuilderImpl::Tick( float DeltaTime )
 		//  - tick layout builders, which includes the stale builders
 		//  - refresh tree view, which finally gets rid of the stale builders
 		TSharedPtr<FPropertyNode> PropertyNode = Node->GetPropertyNode();
-		if (PropertyNode.IsValid() && PropertyNode->IsDestroyed())
+		if (!(PropertyNode.IsValid() && PropertyNode->IsDestroyed()))
 		{
-			continue;
+			Node->Tick(DeltaTime);
 		}
-		
-		Node->Tick( DeltaTime );
+	};
+
+	for( FDetailTreeNode* Node : TickableNodes)
+	{
+		TickNode(Node);
+	}
+
+	while (!TickableNodesToAdd.IsEmpty())
+	{
+		// Tick the new nodes, but keep them in a local set so we don't modify TickableNodesToAdd while ticking them
+		TSet<FDetailTreeNode*> NodesToAdd = MoveTemp(TickableNodesToAdd);
+		TickableNodesToAdd.Reset();
+
+		for( FDetailTreeNode* Node : NodesToAdd)
+		{
+			TickNode(Node);
+		}
+
+		TickableNodes.Append(MoveTemp(NodesToAdd));
 	}
+
+	for( FDetailTreeNode* Node : TickableNodesToRemove )
+	{
+		TickableNodes.Remove(Node);
+	}
+	TickableNodesToRemove.Reset();
 }
 
 void FDetailLayoutBuilderImpl::AddTickableNode( FDetailTreeNode& TickableNode )
 {
-	TickableNodes.Add( &TickableNode );
+	if (bIsIteratingTickableNodes)
+	{
+		TickableNodesToAdd.Add(&TickableNode);
+	}
+	else
+	{
+		TickableNodes.Add(&TickableNode);
+	}
 }
 
 void FDetailLayoutBuilderImpl::RemoveTickableNode( FDetailTreeNode& TickableNode )
 {
-	TickableNodes.Remove( &TickableNode );
+	if (bIsIteratingTickableNodes)
+	{
+		TickableNodesToRemove.Add(&TickableNode);
+	}
+	else
+	{
+		TickableNodes.Remove(&TickableNode);
+	}
 }
 
 void FDetailLayoutBuilderImpl::SaveExpansionState( const FString& NodePath, bool bIsExpanded )
diff --git a/Engine/Source/Editor/PropertyEditor/Private/DetailLayoutBuilderImpl.h b/Engine/Source/Editor/PropertyEditor/Private/DetailLayoutBuilderImpl.h
--- a/Engine/Source/Editor/PropertyEditor/Private/DetailLayoutBuilderImpl.h
+++ b/Engine/Source/Editor/PropertyEditor/Private/DetailLayoutBuilderImpl.h
@@ -272,6 +272,10 @@ private:
 	TSet<FString> PropertyGenerationAllowListPaths;
 	/** Nodes that require ticking */
 	TSet<FDetailTreeNode*> TickableNodes;
+	/** Nodes to add to being ticked, only used while iterating */
+	TSet<FDetailTreeNode*> TickableNodesToAdd;
+	/** Nodes to remove from being ticked, only used while iterating */
+	TSet<FDetailTreeNode*> TickableNodesToRemove;
 	/** Current filter applied to the view */
 	FDetailFilter CurrentFilter;
 	/** All RootTreeNodes */
@@ -290,6 +294,8 @@ private:
 	UStruct* CurrentCustomizationClass;
 	/** True if the layout is for an external root property node and not in the main set of objects the details panel is observing */
 	bool bLayoutForExternalRoot;
+	/** True if the layout is for an external root property node and not in the main set of objects the details panel is observing */
+	bool bIsIteratingTickableNodes = false;
 	/** A delegate which is called whenever a node owned by this layout builder has it's visibility forcibly changed. */
 	FSimpleMulticastDelegate OnNodeVisibilityChanged;
 
diff --git a/Engine/Source/Editor/UnrealEd/Private/LevelEditorViewport.cpp b/Engine/Source/Editor/UnrealEd/Private/LevelEditorViewport.cpp
--- a/Engine/Source/Editor/UnrealEd/Private/LevelEditorViewport.cpp
+++ b/Engine/Source/Editor/UnrealEd/Private/LevelEditorViewport.cpp
@@ -4798,7 +4798,7 @@ void FLevelEditorViewportClient::UpdateHoveredObjects( const TSet<FViewportHover
 				{
 					// Remove hover effect from object that no longer needs it
 					RemoveHoverEffect( OldHoverTarget );
-					HoveredObjects.Remove( OldHoverTarget );
+					It.RemoveCurrent();
 
 					bAnyHoverChanges = true;
 				}
diff --git a/Engine/Source/Editor/UnrealEd/Private/Subsystems/AssetEditorSubsystem.cpp b/Engine/Source/Editor/UnrealEd/Private/Subsystems/AssetEditorSubsystem.cpp
--- a/Engine/Source/Editor/UnrealEd/Private/Subsystems/AssetEditorSubsystem.cpp
+++ b/Engine/Source/Editor/UnrealEd/Private/Subsystems/AssetEditorSubsystem.cpp
@@ -459,9 +459,10 @@ void UAssetEditorSubsystem::NotifyEditorClosed(IAssetEditorInstance* InInstance)
 bool UAssetEditorSubsystem::CloseAllAssetEditors()
 {
 	bool bAllEditorsClosed = true;
-	for (TMultiMap<IAssetEditorInstance*, FAssetEntry>::TIterator It(OpenedEditors); It; ++It)
+	TArray<IAssetEditorInstance*> OpenedEditorInstances;
+	OpenedEditors.GetKeys(OpenedEditorInstances);
+	for (IAssetEditorInstance* Editor : OpenedEditorInstances) // operate on copy since closing will remove from map
 	{
-		IAssetEditorInstance* Editor = It.Key();
 		if (Editor != nullptr)
 		{
 			if (!Editor->CloseWindow(EAssetEditorCloseReason::CloseAllAssetEditors))
diff --git a/Engine/Source/Runtime/Core/Public/Containers/LruCache.h b/Engine/Source/Runtime/Core/Public/Containers/LruCache.h
--- a/Engine/Source/Runtime/Core/Public/Containers/LruCache.h
+++ b/Engine/Source/Runtime/Core/Public/Containers/LruCache.h
@@ -587,10 +587,11 @@ public:
 					MostRecent = Entry->LessRecent;
 				}
 
+				It.RemoveCurrent(); // remove before delete so hashkey doesn't change
+
 				Entry->Unlink();
 				delete Entry;
 
-				It.RemoveCurrent();
 				++NumRemoved;
 			}
 		}
diff --git a/Engine/Source/Runtime/CoreUObject/Private/UObject/MetaData.cpp b/Engine/Source/Runtime/CoreUObject/Private/UObject/MetaData.cpp
--- a/Engine/Source/Runtime/CoreUObject/Private/UObject/MetaData.cpp
+++ b/Engine/Source/Runtime/CoreUObject/Private/UObject/MetaData.cpp
@@ -196,20 +196,26 @@ void UDEPRECATED_MetaData::Serialize(FArchive& Ar)
 		for (TMap< FWeakObjectPtr, TMap<FName, FString> >::TIterator ObjectIt(ObjectMetaDataMap); ObjectIt; ++ObjectIt)
 		{
 			TMap<FName, FString>& CurrentMap = ObjectIt.Value();
+
+			TArray<TPair<FName, FString>, TInlineAllocator<32>> RemappedKeys;
+
 			for (TMap<FName, FString>::TIterator PairIt(CurrentMap); PairIt; ++PairIt)
 			{
-				const FName OldKey = PairIt.Key();
+				const FName OldKey = PairIt->Key;
 				const FName NewKey = KeyRedirectMap.FindRef(OldKey);
 				if (NewKey != NAME_None)
 				{
-					const FString Value = PairIt.Value();
-
+					RemappedKeys.Add( { NewKey, MoveTemp(PairIt->Value) });
 					PairIt.RemoveCurrent();
-					CurrentMap.Add(NewKey, Value);
 
 					UE_LOG(LogMetaData, Verbose, TEXT("Remapping old metadata key '%s' to new key '%s' on object '%s'."), *OldKey.ToString(), *NewKey.ToString(), *ObjectIt.Key().Get()->GetPathName());
 				}
 			}
+
+			for (const TPair<FName, FString>& RemappedKey : RemappedKeys)
+			{
+				CurrentMap.Emplace(RemappedKey.Key, RemappedKey.Value);
+			}
 		}
 
 		for (TMap<FName, FString>::TIterator PairIt(RootMetaDataMap); PairIt; ++PairIt)
diff --git a/Engine/Source/Runtime/Engine/Private/AudioDevice.cpp b/Engine/Source/Runtime/Engine/Private/AudioDevice.cpp
--- a/Engine/Source/Runtime/Engine/Private/AudioDevice.cpp
+++ b/Engine/Source/Runtime/Engine/Private/AudioDevice.cpp
@@ -3258,9 +3258,11 @@ void FAudioDevice::UpdateSoundClassProperties(float DeltaTime)
 	// Remove SoundMix modifications and propagate the properties down the hierarchy
 	ParseSoundClasses(DeltaTime);
 
-	for (decltype(SoundMixModifiers)::TIterator It(SoundMixModifiers); It; ++It)
+	TArray<TObjectPtr<USoundMix>, TInlineAllocator<32>> SoundMixesToClear;
+
+	for (TPair<TObjectPtr<USoundMix>, FSoundMixState>& SoundMix :SoundMixModifiers)
 	{
-		FSoundMixState* SoundMixState = &(It.Value());
+		FSoundMixState* SoundMixState = &SoundMix.Value;
 
 		// Initial delay before mix is applied
 		const double AudioTime = GetAudioClock();
@@ -3294,7 +3296,7 @@ void FAudioDevice::UpdateSoundClassProperties(float DeltaTime)
 			{
 				// Start fading EQ at same time
 				SoundMixState->CurrentState = ESoundMixState::FadingOut;
-				TryClearingEQSoundMix(It.Key());
+				TryClearingEQSoundMix(SoundMix.Key);
 			}
 		}
 		else
@@ -3304,13 +3306,18 @@ void FAudioDevice::UpdateSoundClassProperties(float DeltaTime)
 			SoundMixState->CurrentState = ESoundMixState::AwaitingRemoval;
 		}
 
-		ApplyClassAdjusters(It.Key(), SoundMixState->InterpValue, DeltaTime);
+		ApplyClassAdjusters(SoundMix.Key, SoundMixState->InterpValue, DeltaTime);
 
 		if (SoundMixState->CurrentState == ESoundMixState::AwaitingRemoval && SoundMixState->PassiveRefCount == 0)
 		{
-			ClearSoundMix(It.Key());
+			SoundMixesToClear.Add(SoundMix.Key);
 		}
 	}
+
+	for (USoundMix* SoundMix : SoundMixesToClear)
+	{
+		ClearSoundMix(SoundMix);
+	}
 }
 
 void FAudioDevice::VirtualizeInactiveLoops()
@@ -4097,10 +4104,14 @@ void FAudioDevice::ClearSoundMixModifiers()
 		return;
 	}
 
+	// Work on local data since ClearSoundMixModifier might modify SoundMixModifiers
+	TArray<TObjectPtr<USoundMix>, TInlineAllocator<32>> Keys;
+	SoundMixModifiers.GenerateKeyArray(Keys);
+
 	// Clear all sound mix modifiers
-	for (decltype(SoundMixModifiers)::TIterator It(SoundMixModifiers); It; ++It)
+	for (USoundMix* Key : Keys)
 	{
-		ClearSoundMixModifier(It.Key());
+		ClearSoundMixModifier(Key);
 	}
 }
 
diff --git a/Engine/Source/Runtime/Engine/Private/CurveTable.cpp b/Engine/Source/Runtime/Engine/Private/CurveTable.cpp
--- a/Engine/Source/Runtime/Engine/Private/CurveTable.cpp
+++ b/Engine/Source/Runtime/Engine/Private/CurveTable.cpp
@@ -619,15 +619,11 @@ void UCurveTable::RenameRow(FName& CurveName, FName& NewCurveName)
 
 void UCurveTable::DeleteRow(FName& CurveName)
 {
-	if (RowMap.Contains(CurveName))
+	FRealCurve* CurveToDelete = nullptr;
+	if (RowMap.RemoveAndCopyValue(CurveName, CurveToDelete))
 	{
-		FRealCurve** Curve = RowMap.Find(CurveName);
-		RowMap.Remove(CurveName);
-		if (Curve != nullptr && *Curve != nullptr)
-		{
-			delete *Curve;
-		}
-	}	
+		delete CurveToDelete;
+	}
 }
 
 /** */
diff --git a/Engine/Source/Runtime/Engine/Private/LatentActionManager.cpp b/Engine/Source/Runtime/Engine/Private/LatentActionManager.cpp
--- a/Engine/Source/Runtime/Engine/Private/LatentActionManager.cpp
+++ b/Engine/Source/Runtime/Engine/Private/LatentActionManager.cpp
@@ -283,26 +283,31 @@ void FLatentActionManager::ProcessLatentActions(UObject* InObject, float DeltaTi
 
 void FLatentActionManager::TickLatentActionForObject(float DeltaTime, FActionList& ObjectActionList, UObject* InObject)
 {
-	typedef TPair<int32, FPendingLatentAction*> FActionListPair;
-	TArray<FActionListPair, TInlineAllocator<4>> ItemsToRemove;
-	
+	// Operator on a copy of the action list since it might be modified while processing it
+	using FActionElementArray = TArray<TPair<int32, class FPendingLatentAction*>, TInlineAllocator<32>>;
+	FActionElementArray ItemsToRemove;
+	for (const TPair<int32, class FPendingLatentAction*>& Element : ObjectActionList)
+	{
+		ItemsToRemove.Add(Element);
+	}
+
 	FLatentResponse Response(DeltaTime);
-	for (TMultiMap<int32, FPendingLatentAction*>::TConstIterator It(ObjectActionList); It; ++It)
+	for (FActionElementArray::TIterator It = ItemsToRemove.CreateIterator(); It; ++It)
 	{
-		FPendingLatentAction* Action = It.Value();
+		FPendingLatentAction* Action = It->Value;
 
 		Response.bRemoveAction = false;
 
 		Action->UpdateOperation(Response);
 
-		if (Response.bRemoveAction)
+		if (!Response.bRemoveAction)
 		{
-			ItemsToRemove.Emplace(It.Key(), Action);
+			It.RemoveCurrent();
 		}
 	}
 
 	// Remove any items that were deleted
-	for (const FActionListPair& ItemPair : ItemsToRemove)
+	for (const TPair<int32, class FPendingLatentAction*>& ItemPair : ItemsToRemove)
 	{
 		const int32 ItemIndex = ItemPair.Key;
 		FPendingLatentAction* DyingAction = ItemPair.Value;
diff --git a/Engine/Source/Runtime/Engine/Private/NetConnection.cpp b/Engine/Source/Runtime/Engine/Private/NetConnection.cpp
--- a/Engine/Source/Runtime/Engine/Private/NetConnection.cpp
+++ b/Engine/Source/Runtime/Engine/Private/NetConnection.cpp
@@ -1938,6 +1938,7 @@ void UNetConnection::UpdateLevelVisibilityInternal(const FUpdateLevelVisibilityL
 
 			if (Channel->Actor && Channel->Actor->GetLevel() && Channel->Actor->GetLevel()->GetOutermost()->GetFName() == LevelVisibility.PackageName)
 			{
+				It.RemoveCurrent(); // Remove safely here since Channel->Close will remove unsafely
 				Channel->Close(EChannelCloseReason::LevelUnloaded);
 			}
 		}
