[Backout] - CL44389834
#rnx
Original CL Desc
-----------------------------------------------------------------
[Core] Fix bad use of remove during iteration

#tests Using the editor, running PIE and game client, making builds

diff --git a/Engine/Plugins/Cameras/CameraShakePreviewer/Source/CameraShakePreviewer/Private/CameraShakePreviewerModule.cpp b/Engine/Plugins/Cameras/CameraShakePreviewer/Source/CameraShakePreviewer/Private/CameraShakePreviewerModule.cpp
--- a/Engine/Plugins/Cameras/CameraShakePreviewer/Source/CameraShakePreviewer/Private/CameraShakePreviewerModule.cpp
+++ b/Engine/Plugins/Cameras/CameraShakePreviewer/Source/CameraShakePreviewer/Private/CameraShakePreviewerModule.cpp
@@ -199,7 +199,7 @@ void FCameraShakePreviewerModule::OnLevelViewportClientListChanged()
 	{
 		if (!LevelViewportClients.Contains(It.Key()))
 		{
-			It.RemoveCurrent();
+			ViewportInfos.Remove(It.Key());
 		}
 	}
 	// Add recently added viewports that we don't know about yet.
diff --git a/Engine/Plugins/Media/PixelStreaming/Source/PixelStreaming/Private/AudioSink.cpp b/Engine/Plugins/Media/PixelStreaming/Source/PixelStreaming/Private/AudioSink.cpp
--- a/Engine/Plugins/Media/PixelStreaming/Source/PixelStreaming/Private/AudioSink.cpp
+++ b/Engine/Plugins/Media/PixelStreaming/Source/PixelStreaming/Private/AudioSink.cpp
@@ -7,9 +7,14 @@ namespace UE::PixelStreaming
 {
 	FAudioSink::~FAudioSink()
 	{
-		for (IPixelStreamingAudioConsumer* AudioConsumer : AudioConsumers)
+		for (auto Iter = AudioConsumers.CreateIterator(); Iter; ++Iter)
 		{
-			AudioConsumer->OnConsumerRemoved();
+			IPixelStreamingAudioConsumer* AudioConsumer = Iter.ElementIt->Value;
+			Iter.RemoveCurrent();
+			if (AudioConsumer != nullptr)
+			{
+				AudioConsumer->OnConsumerRemoved();
+			}
 		}
 	}
 
@@ -47,8 +52,9 @@ namespace UE::PixelStreaming
 
 	void FAudioSink::RemoveAudioConsumer(IPixelStreamingAudioConsumer* AudioConsumer)
 	{
-		if (AudioConsumers.Remove(AudioConsumer) > 0)
+		if (AudioConsumers.Contains(AudioConsumer))
 		{
+			AudioConsumers.Remove(AudioConsumer);
 			AudioConsumer->OnConsumerRemoved();
 		}
 	}
diff --git a/Engine/Plugins/Online/OnlineSubsystem/Source/Private/OnlineSubsystemModule.cpp b/Engine/Plugins/Online/OnlineSubsystem/Source/Private/OnlineSubsystemModule.cpp
--- a/Engine/Plugins/Online/OnlineSubsystem/Source/Private/OnlineSubsystemModule.cpp
+++ b/Engine/Plugins/Online/OnlineSubsystem/Source/Private/OnlineSubsystemModule.cpp
@@ -265,8 +265,6 @@ void FOnlineSubsystemModule::ShutdownOnlineSubsystem()
 		// Unloading the module will do proper cleanup
 		FName ModuleName = GetOnlineModuleName(It.Key().ToString(), ModuleRedirects);
 
-		It.RemoveCurrent(); // Remove safely from map otherwise UnloadModule will do it and kill our iteration
-
 		const bool bIsShutdown = true;
 		ModuleManager.UnloadModule(ModuleName, bIsShutdown);
 	} 
diff --git a/Engine/Plugins/Runtime/ComputeFramework/Source/ComputeFramework/Private/ComputeKernelShaderCompilationManager.cpp b/Engine/Plugins/Runtime/ComputeFramework/Source/ComputeFramework/Private/ComputeKernelShaderCompilationManager.cpp
--- a/Engine/Plugins/Runtime/ComputeFramework/Source/ComputeFramework/Private/ComputeKernelShaderCompilationManager.cpp
+++ b/Engine/Plugins/Runtime/ComputeFramework/Source/ComputeFramework/Private/ComputeKernelShaderCompilationManager.cpp
@@ -96,7 +96,7 @@ void FComputeKernelShaderCompilationManager::ProcessAsyncResults()
 		if (Results.FinishedJobs.Num() == Results.NumJobsQueued)
 		{
 			PendingFinalizeComputeKernelShaderMaps.Add(It.Key(), FComputeKernelShaderMapCompileResults(Results));
-			It.RemoveCurrent();
+			ComputeKernelShaderMapJobs.Remove(It.Key());
 		}
 	}
 
@@ -363,7 +363,7 @@ void FComputeKernelShaderCompilationManager::ProcessCompiledComputeKernelShaderM
 
 				// Cleanup shader jobs and compile tracking structures
 				ResultArray.Empty();
-				ProcessIt.RemoveCurrent();
+				CompiledShaderMaps.Remove(ShaderMap->GetCompilingId());
 			}
 
 			if (TimeBudget < 0)
diff --git a/Engine/Source/Developer/DirectoryWatcher/Private/Windows/DirectoryWatcherWindows.cpp b/Engine/Source/Developer/DirectoryWatcher/Private/Windows/DirectoryWatcherWindows.cpp
--- a/Engine/Source/Developer/DirectoryWatcher/Private/Windows/DirectoryWatcherWindows.cpp
+++ b/Engine/Source/Developer/DirectoryWatcher/Private/Windows/DirectoryWatcherWindows.cpp
@@ -81,24 +81,23 @@ bool FDirectoryWatcherWindows::RegisterDirectoryChangedCallback_Handle( const FS
 
 bool FDirectoryWatcherWindows::UnregisterDirectoryChangedCallback_Handle( const FString& Directory, FDelegateHandle InHandle )
 {
-	for (TMap<FDirectoryWithFlags, FDirectoryWatchRequestWindows*>::TIterator RequestIt(RequestMap); RequestIt; ++RequestIt)
+	for (const auto& RequestPair : RequestMap)
 	{
-		if (RequestIt->Key.Key == Directory)
+		if (RequestPair.Key.Key == Directory)
 		{
 			// There should be no NULL entries in the map
-			FDirectoryWatchRequestWindows* Watch = RequestIt->Value;
-			check(Watch);
+			check(RequestPair.Value);
 
-			if (Watch->RemoveDelegate(InHandle))
+			if (RequestPair.Value->RemoveDelegate(InHandle))
 			{
-				if (!Watch->HasDelegates())
+				if (!RequestPair.Value->HasDelegates())
 				{
 					// Remove from the active map and add to the pending delete list
-					RequestIt.RemoveCurrent();
-					RequestsPendingDelete.AddUnique(Watch);
+					RequestMap.Remove(RequestPair.Key);
+					RequestsPendingDelete.AddUnique(RequestPair.Value);
 
 					// Signal to end the watch which will mark this request for deletion
-					Watch->EndWatchRequest();
+					RequestPair.Value->EndWatchRequest();
 				}
 
 				return true;
diff --git a/Engine/Source/Developer/MessageLog/Private/Model/MessageLogModel.cpp b/Engine/Source/Developer/MessageLog/Private/Model/MessageLogModel.cpp
--- a/Engine/Source/Developer/MessageLog/Private/Model/MessageLogModel.cpp
+++ b/Engine/Source/Developer/MessageLog/Private/Model/MessageLogModel.cpp
@@ -4,10 +4,9 @@
 
 FMessageLogModel::~FMessageLogModel()
 {
-	if (!NameToModelMap.IsEmpty())
+	for( auto It = NameToModelMap.CreateIterator(); It; ++It )
 	{
-		NameToModelMap.Empty();
-		Notify();
+		UnregisterLogListingModel( It.Key() );
 	}
 }
 
diff --git a/Engine/Source/Editor/PropertyEditor/Private/DetailLayoutBuilderImpl.cpp b/Engine/Source/Editor/PropertyEditor/Private/DetailLayoutBuilderImpl.cpp
--- a/Engine/Source/Editor/PropertyEditor/Private/DetailLayoutBuilderImpl.cpp
+++ b/Engine/Source/Editor/PropertyEditor/Private/DetailLayoutBuilderImpl.cpp
@@ -897,9 +897,9 @@ FCustomPropertyTypeLayoutMap FDetailLayoutBuilderImpl::GetInstancedPropertyTypeL
 
 void FDetailLayoutBuilderImpl::RefreshNodeVisbility()
 {
-	for(FDetailTreeNode* Node : TickableNodes)
+	for(TSet<FDetailTreeNode*>::TIterator It = TickableNodes.CreateIterator(); It; ++It)
 	{
-		Node->RefreshVisibility();
+		(*It)->RefreshVisibility();
 	}
 }
 
@@ -1057,12 +1057,12 @@ void FDetailLayoutBuilderImpl::SetCustomProperty( const TSharedPtr<FPropertyNode
 	PropertyNode->SetNodeFlags( EPropertyNodeFlags::IsCustomized, true );
 }
 
-void FDetailLayoutBuilderImpl::Tick(float DeltaTime)
+void FDetailLayoutBuilderImpl::Tick( float DeltaTime )
 {
-	TGuardValue Guard(bIsIteratingTickableNodes, true);
-
-	auto TickNode = [DeltaTime](FDetailTreeNode* Node)
+	for( auto It = TickableNodes.CreateIterator(); It; ++It )
 	{
+		FDetailTreeNode* Node = *It;
+
 		// Skip ticking tree nodes which point to destroyed property nodes.
 		// This can happen when because the update order is this:
 		//	- update property nodes, calling DestroyTree(), and creating new nodes
@@ -1070,60 +1070,23 @@ void FDetailLayoutBuilderImpl::Tick(float DeltaTime)
 		//  - tick layout builders, which includes the stale builders
 		//  - refresh tree view, which finally gets rid of the stale builders
 		TSharedPtr<FPropertyNode> PropertyNode = Node->GetPropertyNode();
-		if (!(PropertyNode.IsValid() && PropertyNode->IsDestroyed()))
+		if (PropertyNode.IsValid() && PropertyNode->IsDestroyed())
 		{
-			Node->Tick(DeltaTime);
+			continue;
 		}
-	};
-
-	for( FDetailTreeNode* Node : TickableNodes)
-	{
-		TickNode(Node);
-	}
-
-	while (!TickableNodesToAdd.IsEmpty())
-	{
-		// Tick the new nodes, but keep them in a local set so we don't modify TickableNodesToAdd while ticking them
-		TSet<FDetailTreeNode*> NodesToAdd = MoveTemp(TickableNodesToAdd);
-		TickableNodesToAdd.Reset();
-
-		for( FDetailTreeNode* Node : NodesToAdd)
-		{
-			TickNode(Node);
-		}
-
-		TickableNodes.Append(MoveTemp(NodesToAdd));
-	}
-
-	for( FDetailTreeNode* Node : TickableNodesToRemove )
-	{
-		TickableNodes.Remove(Node);
+		
+		Node->Tick( DeltaTime );
 	}
-	TickableNodesToRemove.Reset();
 }
 
 void FDetailLayoutBuilderImpl::AddTickableNode( FDetailTreeNode& TickableNode )
 {
-	if (bIsIteratingTickableNodes)
-	{
-		TickableNodesToAdd.Add(&TickableNode);
-	}
-	else
-	{
-		TickableNodes.Add(&TickableNode);
-	}
+	TickableNodes.Add( &TickableNode );
 }
 
 void FDetailLayoutBuilderImpl::RemoveTickableNode( FDetailTreeNode& TickableNode )
 {
-	if (bIsIteratingTickableNodes)
-	{
-		TickableNodesToRemove.Add(&TickableNode);
-	}
-	else
-	{
-		TickableNodes.Remove(&TickableNode);
-	}
+	TickableNodes.Remove( &TickableNode );
 }
 
 void FDetailLayoutBuilderImpl::SaveExpansionState( const FString& NodePath, bool bIsExpanded )
diff --git a/Engine/Source/Editor/PropertyEditor/Private/DetailLayoutBuilderImpl.h b/Engine/Source/Editor/PropertyEditor/Private/DetailLayoutBuilderImpl.h
--- a/Engine/Source/Editor/PropertyEditor/Private/DetailLayoutBuilderImpl.h
+++ b/Engine/Source/Editor/PropertyEditor/Private/DetailLayoutBuilderImpl.h
@@ -272,10 +272,6 @@ private:
 	TSet<FString> PropertyGenerationAllowListPaths;
 	/** Nodes that require ticking */
 	TSet<FDetailTreeNode*> TickableNodes;
-	/** Nodes to add to being ticked, only used while iterating */
-	TSet<FDetailTreeNode*> TickableNodesToAdd;
-	/** Nodes to remove from being ticked, only used while iterating */
-	TSet<FDetailTreeNode*> TickableNodesToRemove;
 	/** Current filter applied to the view */
 	FDetailFilter CurrentFilter;
 	/** All RootTreeNodes */
@@ -294,8 +290,6 @@ private:
 	UStruct* CurrentCustomizationClass;
 	/** True if the layout is for an external root property node and not in the main set of objects the details panel is observing */
 	bool bLayoutForExternalRoot;
-	/** True if the layout is for an external root property node and not in the main set of objects the details panel is observing */
-	bool bIsIteratingTickableNodes = false;
 	/** A delegate which is called whenever a node owned by this layout builder has it's visibility forcibly changed. */
 	FSimpleMulticastDelegate OnNodeVisibilityChanged;
 
diff --git a/Engine/Source/Editor/UnrealEd/Private/LevelEditorViewport.cpp b/Engine/Source/Editor/UnrealEd/Private/LevelEditorViewport.cpp
--- a/Engine/Source/Editor/UnrealEd/Private/LevelEditorViewport.cpp
+++ b/Engine/Source/Editor/UnrealEd/Private/LevelEditorViewport.cpp
@@ -4798,7 +4798,7 @@ void FLevelEditorViewportClient::UpdateHoveredObjects( const TSet<FViewportHover
 				{
 					// Remove hover effect from object that no longer needs it
 					RemoveHoverEffect( OldHoverTarget );
-					It.RemoveCurrent();
+					HoveredObjects.Remove( OldHoverTarget );
 
 					bAnyHoverChanges = true;
 				}
diff --git a/Engine/Source/Editor/UnrealEd/Private/Subsystems/AssetEditorSubsystem.cpp b/Engine/Source/Editor/UnrealEd/Private/Subsystems/AssetEditorSubsystem.cpp
--- a/Engine/Source/Editor/UnrealEd/Private/Subsystems/AssetEditorSubsystem.cpp
+++ b/Engine/Source/Editor/UnrealEd/Private/Subsystems/AssetEditorSubsystem.cpp
@@ -459,10 +459,9 @@ void UAssetEditorSubsystem::NotifyEditorClosed(IAssetEditorInstance* InInstance)
 bool UAssetEditorSubsystem::CloseAllAssetEditors()
 {
 	bool bAllEditorsClosed = true;
-	TArray<IAssetEditorInstance*> OpenedEditorInstances;
-	OpenedEditors.GetKeys(OpenedEditorInstances);
-	for (IAssetEditorInstance* Editor : OpenedEditorInstances) // operate on copy since closing will remove from map
+	for (TMultiMap<IAssetEditorInstance*, FAssetEntry>::TIterator It(OpenedEditors); It; ++It)
 	{
+		IAssetEditorInstance* Editor = It.Key();
 		if (Editor != nullptr)
 		{
 			if (!Editor->CloseWindow(EAssetEditorCloseReason::CloseAllAssetEditors))
diff --git a/Engine/Source/Runtime/Core/Public/Containers/LruCache.h b/Engine/Source/Runtime/Core/Public/Containers/LruCache.h
--- a/Engine/Source/Runtime/Core/Public/Containers/LruCache.h
+++ b/Engine/Source/Runtime/Core/Public/Containers/LruCache.h
@@ -587,11 +587,10 @@ public:
 					MostRecent = Entry->LessRecent;
 				}
 
-				It.RemoveCurrent(); // remove before delete so hashkey doesn't change
-
 				Entry->Unlink();
 				delete Entry;
 
+				It.RemoveCurrent();
 				++NumRemoved;
 			}
 		}
diff --git a/Engine/Source/Runtime/CoreUObject/Private/UObject/MetaData.cpp b/Engine/Source/Runtime/CoreUObject/Private/UObject/MetaData.cpp
--- a/Engine/Source/Runtime/CoreUObject/Private/UObject/MetaData.cpp
+++ b/Engine/Source/Runtime/CoreUObject/Private/UObject/MetaData.cpp
@@ -196,26 +196,20 @@ void UDEPRECATED_MetaData::Serialize(FArchive& Ar)
 		for (TMap< FWeakObjectPtr, TMap<FName, FString> >::TIterator ObjectIt(ObjectMetaDataMap); ObjectIt; ++ObjectIt)
 		{
 			TMap<FName, FString>& CurrentMap = ObjectIt.Value();
-
-			TArray<TPair<FName, FString>, TInlineAllocator<32>> RemappedKeys;
-
 			for (TMap<FName, FString>::TIterator PairIt(CurrentMap); PairIt; ++PairIt)
 			{
-				const FName OldKey = PairIt->Key;
+				const FName OldKey = PairIt.Key();
 				const FName NewKey = KeyRedirectMap.FindRef(OldKey);
 				if (NewKey != NAME_None)
 				{
-					RemappedKeys.Add( { NewKey, MoveTemp(PairIt->Value) });
+					const FString Value = PairIt.Value();
+
 					PairIt.RemoveCurrent();
+					CurrentMap.Add(NewKey, Value);
 
 					UE_LOG(LogMetaData, Verbose, TEXT("Remapping old metadata key '%s' to new key '%s' on object '%s'."), *OldKey.ToString(), *NewKey.ToString(), *ObjectIt.Key().Get()->GetPathName());
 				}
 			}
-
-			for (const TPair<FName, FString>& RemappedKey : RemappedKeys)
-			{
-				CurrentMap.Emplace(RemappedKey.Key, RemappedKey.Value);
-			}
 		}
 
 		for (TMap<FName, FString>::TIterator PairIt(RootMetaDataMap); PairIt; ++PairIt)
diff --git a/Engine/Source/Runtime/Engine/Private/AudioDevice.cpp b/Engine/Source/Runtime/Engine/Private/AudioDevice.cpp
--- a/Engine/Source/Runtime/Engine/Private/AudioDevice.cpp
+++ b/Engine/Source/Runtime/Engine/Private/AudioDevice.cpp
@@ -3258,11 +3258,9 @@ void FAudioDevice::UpdateSoundClassProperties(float DeltaTime)
 	// Remove SoundMix modifications and propagate the properties down the hierarchy
 	ParseSoundClasses(DeltaTime);
 
-	TArray<TObjectPtr<USoundMix>, TInlineAllocator<32>> SoundMixesToClear;
-
-	for (TPair<TObjectPtr<USoundMix>, FSoundMixState>& SoundMix :SoundMixModifiers)
+	for (decltype(SoundMixModifiers)::TIterator It(SoundMixModifiers); It; ++It)
 	{
-		FSoundMixState* SoundMixState = &SoundMix.Value;
+		FSoundMixState* SoundMixState = &(It.Value());
 
 		// Initial delay before mix is applied
 		const double AudioTime = GetAudioClock();
@@ -3296,7 +3294,7 @@ void FAudioDevice::UpdateSoundClassProperties(float DeltaTime)
 			{
 				// Start fading EQ at same time
 				SoundMixState->CurrentState = ESoundMixState::FadingOut;
-				TryClearingEQSoundMix(SoundMix.Key);
+				TryClearingEQSoundMix(It.Key());
 			}
 		}
 		else
@@ -3306,18 +3304,13 @@ void FAudioDevice::UpdateSoundClassProperties(float DeltaTime)
 			SoundMixState->CurrentState = ESoundMixState::AwaitingRemoval;
 		}
 
-		ApplyClassAdjusters(SoundMix.Key, SoundMixState->InterpValue, DeltaTime);
+		ApplyClassAdjusters(It.Key(), SoundMixState->InterpValue, DeltaTime);
 
 		if (SoundMixState->CurrentState == ESoundMixState::AwaitingRemoval && SoundMixState->PassiveRefCount == 0)
 		{
-			SoundMixesToClear.Add(SoundMix.Key);
+			ClearSoundMix(It.Key());
 		}
 	}
-
-	for (USoundMix* SoundMix : SoundMixesToClear)
-	{
-		ClearSoundMix(SoundMix);
-	}
 }
 
 void FAudioDevice::VirtualizeInactiveLoops()
@@ -4104,14 +4097,10 @@ void FAudioDevice::ClearSoundMixModifiers()
 		return;
 	}
 
-	// Work on local data since ClearSoundMixModifier might modify SoundMixModifiers
-	TArray<TObjectPtr<USoundMix>, TInlineAllocator<32>> Keys;
-	SoundMixModifiers.GenerateKeyArray(Keys);
-
 	// Clear all sound mix modifiers
-	for (USoundMix* Key : Keys)
+	for (decltype(SoundMixModifiers)::TIterator It(SoundMixModifiers); It; ++It)
 	{
-		ClearSoundMixModifier(Key);
+		ClearSoundMixModifier(It.Key());
 	}
 }
 
diff --git a/Engine/Source/Runtime/Engine/Private/CurveTable.cpp b/Engine/Source/Runtime/Engine/Private/CurveTable.cpp
--- a/Engine/Source/Runtime/Engine/Private/CurveTable.cpp
+++ b/Engine/Source/Runtime/Engine/Private/CurveTable.cpp
@@ -619,11 +619,15 @@ void UCurveTable::RenameRow(FName& CurveName, FName& NewCurveName)
 
 void UCurveTable::DeleteRow(FName& CurveName)
 {
-	FRealCurve* CurveToDelete = nullptr;
-	if (RowMap.RemoveAndCopyValue(CurveName, CurveToDelete))
+	if (RowMap.Contains(CurveName))
 	{
-		delete CurveToDelete;
-	}
+		FRealCurve** Curve = RowMap.Find(CurveName);
+		RowMap.Remove(CurveName);
+		if (Curve != nullptr && *Curve != nullptr)
+		{
+			delete *Curve;
+		}
+	}	
 }
 
 /** */
diff --git a/Engine/Source/Runtime/Engine/Private/LatentActionManager.cpp b/Engine/Source/Runtime/Engine/Private/LatentActionManager.cpp
--- a/Engine/Source/Runtime/Engine/Private/LatentActionManager.cpp
+++ b/Engine/Source/Runtime/Engine/Private/LatentActionManager.cpp
@@ -283,31 +283,26 @@ void FLatentActionManager::ProcessLatentActions(UObject* InObject, float DeltaTi
 
 void FLatentActionManager::TickLatentActionForObject(float DeltaTime, FActionList& ObjectActionList, UObject* InObject)
 {
-	// Operator on a copy of the action list since it might be modified while processing it
-	using FActionElementArray = TArray<TPair<int32, class FPendingLatentAction*>, TInlineAllocator<32>>;
-	FActionElementArray ItemsToRemove;
-	for (const TPair<int32, class FPendingLatentAction*>& Element : ObjectActionList)
-	{
-		ItemsToRemove.Add(Element);
-	}
-
+	typedef TPair<int32, FPendingLatentAction*> FActionListPair;
+	TArray<FActionListPair, TInlineAllocator<4>> ItemsToRemove;
+	
 	FLatentResponse Response(DeltaTime);
-	for (FActionElementArray::TIterator It = ItemsToRemove.CreateIterator(); It; ++It)
+	for (TMultiMap<int32, FPendingLatentAction*>::TConstIterator It(ObjectActionList); It; ++It)
 	{
-		FPendingLatentAction* Action = It->Value;
+		FPendingLatentAction* Action = It.Value();
 
 		Response.bRemoveAction = false;
 
 		Action->UpdateOperation(Response);
 
-		if (!Response.bRemoveAction)
+		if (Response.bRemoveAction)
 		{
-			It.RemoveCurrent();
+			ItemsToRemove.Emplace(It.Key(), Action);
 		}
 	}
 
 	// Remove any items that were deleted
-	for (const TPair<int32, class FPendingLatentAction*>& ItemPair : ItemsToRemove)
+	for (const FActionListPair& ItemPair : ItemsToRemove)
 	{
 		const int32 ItemIndex = ItemPair.Key;
 		FPendingLatentAction* DyingAction = ItemPair.Value;
diff --git a/Engine/Source/Runtime/Engine/Private/NetConnection.cpp b/Engine/Source/Runtime/Engine/Private/NetConnection.cpp
--- a/Engine/Source/Runtime/Engine/Private/NetConnection.cpp
+++ b/Engine/Source/Runtime/Engine/Private/NetConnection.cpp
@@ -1938,7 +1938,6 @@ void UNetConnection::UpdateLevelVisibilityInternal(const FUpdateLevelVisibilityL
 
 			if (Channel->Actor && Channel->Actor->GetLevel() && Channel->Actor->GetLevel()->GetOutermost()->GetFName() == LevelVisibility.PackageName)
 			{
-				It.RemoveCurrent(); // Remove safely here since Channel->Close will remove unsafely
 				Channel->Close(EChannelCloseReason::LevelUnloaded);
 			}
 		}
