PR #13422: FPropertyChangeListener crash fixes

diff --git a/Engine/Source/Editor/PropertyEditor/Private/PropertyChangeListener.cpp b/Engine/Source/Editor/PropertyEditor/Private/PropertyChangeListener.cpp
--- a/Engine/Source/Editor/PropertyEditor/Private/PropertyChangeListener.cpp
+++ b/Engine/Source/Editor/PropertyEditor/Private/PropertyChangeListener.cpp
@@ -63,18 +63,19 @@ public:
 	 */
 	void CacheValue()
 	{
-		FPropertyValueAddresses PropertyValueAddresses = GetPropertyValueAddresses();
 		Data.Reset();
 
-		FPropertyNode& PropertyNodeRef = *PropertyNode.Pin();
-		FProperty* Property = PropertyNodeRef.GetProperty();
+		TSharedPtr<FPropertyNode> PropertyNodePtr = PropertyNode.Pin();
+		if(FProperty* Property = PropertyNodePtr ? PropertyNodePtr->GetProperty() : nullptr)
 		{
+			FPropertyValueAddresses PropertyValueAddresses = GetPropertyValueAddresses();
+
 			// Not supported yet
 			check( !Property->IsA( FArrayProperty::StaticClass() ) );
 			check( !Property->IsA( FSetProperty::StaticClass() ) );
 			check( !Property->IsA( FMapProperty::StaticClass() ) );
 
-			if( PropertyNodeRef.GetArrayIndex() == INDEX_NONE && Property->ArrayDim > 1 )
+			if(PropertyNodePtr->GetArrayIndex() == INDEX_NONE && Property->ArrayDim > 1 )
 			{
 				// Check static arrays
 				Data.AddZeroed( Property->ArrayDim * Property->GetElementSize() );
@@ -97,64 +98,68 @@ public:
 	 */
 	bool ScanForChanges( bool bRecacheNewValues )
 	{
-		FPropertyNode& PropertyNodeRef = *PropertyNode.Pin();
-		FProperty* Property = PropertyNodeRef.GetProperty();
-		FPropertyValueAddresses PropertyValueAddresses = GetPropertyValueAddresses();
-
-		bool bPropertyValid = true;
-		bool bChanged = false;
-
-		if ( PropertyValueAddresses.BaseAddress != nullptr && PropertyValueAddresses.Address != nullptr )
+		TSharedPtr<FPropertyNode> PropertyNodePtr = PropertyNode.Pin();
+		if (FProperty* Property = PropertyNodePtr ? PropertyNodePtr->GetProperty() : nullptr)
 		{
-			FArrayProperty* OuterArrayProperty = Property->GetOwner<FArrayProperty>();
-			FSetProperty* OuterSetProperty = Property->GetOwner<FSetProperty>();
-			FMapProperty* OuterMapProperty = Property->GetOwner<FMapProperty>();
+			FPropertyValueAddresses PropertyValueAddresses = GetPropertyValueAddresses();
+
+			bool bPropertyValid = true;
+			bool bChanged = false;
 
-			if ( OuterArrayProperty != NULL )
+			if (PropertyValueAddresses.BaseAddress != nullptr && PropertyValueAddresses.Address != nullptr)
 			{
-				// make sure we're not trying to compare against an element that doesn't exist
-				FScriptArrayHelper ArrayHelper(OuterArrayProperty, PropertyValueAddresses.BaseAddress);
-				if ( PropertyNodeRef.GetArrayIndex() >= ArrayHelper.Num() )
+				FArrayProperty* OuterArrayProperty = Property->GetOwner<FArrayProperty>();
+				FSetProperty* OuterSetProperty = Property->GetOwner<FSetProperty>();
+				FMapProperty* OuterMapProperty = Property->GetOwner<FMapProperty>();
+
+				if (OuterArrayProperty != NULL)
+				{
+					// make sure we're not trying to compare against an element that doesn't exist
+					FScriptArrayHelper ArrayHelper(OuterArrayProperty, PropertyValueAddresses.BaseAddress);
+					if (PropertyNodePtr->GetArrayIndex() >= ArrayHelper.Num())
+					{
+						bPropertyValid = false;
+					}
+				}
+				else if (OuterSetProperty != NULL)
 				{
-					bPropertyValid = false;
+					FScriptSetHelper SetHelper(OuterSetProperty, PropertyValueAddresses.BaseAddress);
+
+					bool bIsValidIndex = PropertyNodePtr->GetArrayIndex() >= 0 && PropertyNodePtr->GetArrayIndex() < SetHelper.Num();
+					if (!bIsValidIndex)
+					{
+						bPropertyValid = false;
+					}
 				}
-			}
-			else if ( OuterSetProperty != NULL )
-			{
-				FScriptSetHelper SetHelper(OuterSetProperty, PropertyValueAddresses.BaseAddress);
-				
-				bool bIsValidIndex = PropertyNodeRef.GetArrayIndex() >= 0 && PropertyNodeRef.GetArrayIndex() < SetHelper.Num();
-				if (!bIsValidIndex)
+				else if (OuterMapProperty != NULL)
 				{
-					bPropertyValid = false;
+					FScriptMapHelper MapHelper(OuterMapProperty, PropertyValueAddresses.BaseAddress);
+
+					bool bIsValidIndex = PropertyNodePtr->GetArrayIndex() >= 0 && PropertyNodePtr->GetArrayIndex() < MapHelper.Num();
+					if (!bIsValidIndex)
+					{
+						bPropertyValid = false;
+					}
 				}
 			}
-			else if ( OuterMapProperty != NULL )
+			else
 			{
-				FScriptMapHelper MapHelper(OuterMapProperty, PropertyValueAddresses.BaseAddress);
+				bPropertyValid = false;
+			}
 
-				bool bIsValidIndex = PropertyNodeRef.GetArrayIndex() >= 0 && PropertyNodeRef.GetArrayIndex() < MapHelper.Num();
-				if (!bIsValidIndex)
+			if (bPropertyValid)
+			{
+				bChanged = !Property->Identical(PropertyValueAddresses.Address, Data.GetData());
+				if (bRecacheNewValues)
 				{
-					bPropertyValid = false;
+					CacheValue();
 				}
 			}
-		}
-		else
-		{
-			bPropertyValid = false;
-		}
 
-		if( bPropertyValid )
-		{
-			bChanged = !Property->Identical( PropertyValueAddresses.Address, Data.GetData() );
-			if( bRecacheNewValues )
-			{
-				CacheValue();
-			}
+			return bChanged;
 		}
 
-		return bChanged;
+		return false;
 	}
 
 private:
@@ -163,12 +168,12 @@ private:
 		/**
 		* The address of the property's value.
 		*/
-		uint8* Address;
+		uint8* Address = nullptr;
 		/**
 		* The base address of this property's value.  i.e. for dynamic arrays, the location of the FScriptArray which
 		* contains the array property's value
 		*/
-		uint8* BaseAddress;
+		uint8* BaseAddress = nullptr;
 	};
 
 	/**
@@ -177,23 +182,26 @@ private:
 	 */
 	FPropertyValueAddresses GetPropertyValueAddresses()
 	{
-		FPropertyNode& PropertyNodeRef = *PropertyNode.Pin();
-		FProperty* Property = PropertyNodeRef.GetProperty();
+		TSharedPtr<FPropertyNode> PropertyNodePtr = PropertyNode.Pin();
+		if (FProperty* Property = PropertyNodePtr ? PropertyNodePtr->GetProperty() : nullptr)
+		{
+			FPropertyNode* ParentNode = PropertyNodePtr->GetParentNode();
+			FArrayProperty* OuterArrayProp = Property->GetOwner<FArrayProperty>();
+			FSetProperty* OuterSetProp = Property->GetOwner<FSetProperty>();
+			FMapProperty* OuterMapProp = Property->GetOwner<FMapProperty>();
 
-		FPropertyNode* ParentNode = PropertyNodeRef.GetParentNode();
-		FArrayProperty* OuterArrayProp = Property->GetOwner<FArrayProperty>();
-		FSetProperty* OuterSetProp = Property->GetOwner<FSetProperty>();
-		FMapProperty* OuterMapProp = Property->GetOwner<FMapProperty>();
+			FPropertyValueAddresses ValueAddresses;
 
-		FPropertyValueAddresses ValueAddresses;
+			ValueAddresses.BaseAddress = (OuterArrayProp == NULL && OuterSetProp == NULL && OuterMapProp == NULL)
+				? PropertyNodePtr->GetValueBaseAddressFromObject(OwnerObject)
+				: (ParentNode ? ParentNode->GetValueBaseAddressFromObject(OwnerObject) : nullptr);
 
-		ValueAddresses.BaseAddress = (OuterArrayProp == NULL && OuterSetProp == NULL && OuterMapProp == NULL)
-			? PropertyNodeRef.GetValueBaseAddressFromObject(OwnerObject)
-			: ParentNode->GetValueBaseAddressFromObject(OwnerObject);
+			ValueAddresses.Address = PropertyNodePtr->GetValueAddressFromObject(OwnerObject);
 
-		ValueAddresses.Address = PropertyNodeRef.GetValueAddressFromObject(OwnerObject);
+			return ValueAddresses;
+		}
 
-		return ValueAddresses;
+		return FPropertyValueAddresses();
 	}
 
 private:
@@ -207,7 +215,7 @@ private:
 	TWeakObjectPtr<UObject> CachedTopLevelObject;
 
 	/** The address of the owning object */
-	UObject*	OwnerObject;
+	UObject*	OwnerObject = nullptr;
 };
 
 
