PR #13327: Struct Type Picker: Add option to exclude base struct from tree/list

diff --git a/Engine/Source/Editor/PropertyEditor/Private/UserInterface/PropertyEditor/SPropertyEditorStruct.cpp b/Engine/Source/Editor/PropertyEditor/Private/UserInterface/PropertyEditor/SPropertyEditorStruct.cpp
--- a/Engine/Source/Editor/PropertyEditor/Private/UserInterface/PropertyEditor/SPropertyEditorStruct.cpp
+++ b/Engine/Source/Editor/PropertyEditor/Private/UserInterface/PropertyEditor/SPropertyEditorStruct.cpp
@@ -17,6 +17,7 @@ class FPropertyEditorStructFilter : public IStructViewerFilter
 public:
 	/** The meta struct for the property that classes must be a child-of. */
 	const UScriptStruct* MetaStruct = nullptr;
+	bool bExcludeBaseStruct = false;
 
 	// TODO: Have a flag controlling whether we allow UserDefinedStructs, even when a MetaClass is set (as they cannot support inheritance, but may still be allowed (eg, data tables))?
 
@@ -28,6 +29,14 @@ public:
 			return MetaStruct == nullptr;
 		}
 
+		if (bExcludeBaseStruct && MetaStruct)
+		{
+			if (InStruct == MetaStruct)
+			{
+				return false;
+			}
+		}
+
 		// Query the native struct to see if it has the correct parent type (if any)
 		return !MetaStruct || InStruct->IsChildOf(MetaStruct);
 	}
@@ -98,6 +107,8 @@ void SPropertyEditorStruct::Construct(const FArguments& InArgs, const TSharedPtr
 		bShowViewOptions = Property->GetOwnerProperty()->HasMetaData(TEXT("HideViewOptions")) ? false : true;
 		bShowTree = Property->GetOwnerProperty()->HasMetaData(TEXT("ShowTreeView"));
 		bShowDisplayNames = Property->GetOwnerProperty()->HasMetaData(TEXT("ShowDisplayNames"));
+		static const FName NAME_ExcludeBaseStruct = "ExcludeBaseStruct";
+		bExcludeBaseStruct = Property->GetOwnerProperty()->HasMetaData(NAME_ExcludeBaseStruct);
 	}
 	else
 	{
@@ -109,6 +120,7 @@ void SPropertyEditorStruct::Construct(const FArguments& InArgs, const TSharedPtr
 		bShowViewOptions = InArgs._ShowViewOptions;
 		bShowTree = InArgs._ShowTree;
 		bShowDisplayNames = InArgs._ShowDisplayNames;
+		bExcludeBaseStruct = InArgs._bExcludeBaseStruct;
 
 		SelectedStruct = InArgs._SelectedStruct;
 		OnSetStruct = InArgs._OnSetStruct;
@@ -188,6 +200,7 @@ TSharedRef<SWidget> SPropertyEditorStruct::GenerateStructPicker()
 	TSharedRef<FPropertyEditorStructFilter> StructFilter = MakeShared<FPropertyEditorStructFilter>();
 	Options.StructFilter = StructFilter;
 	StructFilter->MetaStruct = MetaStruct;
+	StructFilter->bExcludeBaseStruct = bExcludeBaseStruct;
 	Options.NameTypeToDisplay = (bShowDisplayNames ? EStructViewerNameTypeToDisplay::DisplayName : EStructViewerNameTypeToDisplay::StructName);
 	Options.DisplayMode = bShowTree ? EStructViewerDisplayMode::TreeView : EStructViewerDisplayMode::ListView;
 	Options.bAllowViewOptions = bShowViewOptions;
diff --git a/Engine/Source/Editor/PropertyEditor/Private/UserInterface/PropertyEditor/SPropertyEditorStruct.h b/Engine/Source/Editor/PropertyEditor/Private/UserInterface/PropertyEditor/SPropertyEditorStruct.h
--- a/Engine/Source/Editor/PropertyEditor/Private/UserInterface/PropertyEditor/SPropertyEditorStruct.h
+++ b/Engine/Source/Editor/PropertyEditor/Private/UserInterface/PropertyEditor/SPropertyEditorStruct.h
@@ -27,6 +27,7 @@ public:
 		, _MetaStruct(nullptr)
 		, _AllowNone(true)
 		, _ShowDisplayNames(false)
+		, _bExcludeBaseStruct(false)
 		{}
 		SLATE_ARGUMENT(FSlateFontInfo, Font)
 		/** Arguments used when constructing this outside of a PropertyEditor (PropertyEditor == null), ignored otherwise */
@@ -43,6 +44,8 @@ public:
 		SLATE_ARGUMENT(bool, ShowTree)
 		/** Should we prettify struct names on the struct picker? (ie show their display name) */
 		SLATE_ARGUMENT(bool, ShowDisplayNames)
+		/** Should we hide the specified MetaStruct from the selection list and only display its children? */
+		SLATE_ARGUMENT(bool, bExcludeBaseStruct)
 		/** Delegate used to set the currently selected struct (required if PropertyEditor == null) */
 		SLATE_EVENT(FOnSetStruct, OnSetStruct)
 	SLATE_END_ARGS()
@@ -99,6 +102,8 @@ private:
 	bool bShowTree;
 	/** Should we prettify struct names on the struct picker? (ie show their display name) */
 	bool bShowDisplayNames;
+	/** Should we hide the specified MetaStruct from the selection list and only display its children? */
+	bool bExcludeBaseStruct;
 
 	/** Attribute used to get the currently selected struct (required if PropertyEditor == null) */
 	TAttribute<const UScriptStruct*> SelectedStruct;
