Asset Picker Size Scaling as Editor Preference

Adds the Editor Setting "Scale Asset Picker Widget Size" to allow the size of asset pickers to be scaled up or down.

PR #12909: Asset Picker Size Scaling as Editor Preference

diff --git a/Engine/Source/Editor/PropertyEditor/Private/UserInterface/PropertyEditor/SPropertyMenuAssetPicker.cpp b/Engine/Source/Editor/PropertyEditor/Private/UserInterface/PropertyEditor/SPropertyMenuAssetPicker.cpp
--- a/Engine/Source/Editor/PropertyEditor/Private/UserInterface/PropertyEditor/SPropertyMenuAssetPicker.cpp
+++ b/Engine/Source/Editor/PropertyEditor/Private/UserInterface/PropertyEditor/SPropertyMenuAssetPicker.cpp
@@ -19,6 +19,7 @@
 #include "UserInterface/PropertyEditor/SPropertyEditorAsset.h"
 #include "UserInterface/PropertyEditor/PropertyEditorAssetConstants.h"
 #include "Widgets/Layout/SBox.h"
+#include "Settings/EditorStyleSettings.h"
 
 #define LOCTEXT_NAMESPACE "PropertyEditor"
 
@@ -169,10 +170,16 @@ void SPropertyMenuAssetPicker::Construct( const FArguments& InArgs )
 
 		AssetPickerWidget = ContentBrowserModule.Get().CreateAssetPicker(AssetPickerConfig);
 
+		// Adjust widget size by Editor Preferences AssetPickerSizeScale
+		const UEditorStyleSettings* EditorSettings = GetDefault<UEditorStyleSettings>();
+		const float AssetPickerSizeScale = EditorSettings->AssetPickerSizeScale;
+		float WidgetWidth = static_cast<float>(PropertyEditorAssetConstants::ContentBrowserWindowSize.X) * AssetPickerSizeScale;
+		float WidgetHeight = static_cast<float>(PropertyEditorAssetConstants::ContentBrowserWindowSize.Y) * AssetPickerSizeScale;
+
 		TSharedRef<SWidget> MenuContent =
 			SNew(SBox)
-			.WidthOverride(static_cast<float>(PropertyEditorAssetConstants::ContentBrowserWindowSize.X))
-			.HeightOverride(static_cast<float>(PropertyEditorAssetConstants::ContentBrowserWindowSize.Y))
+			.WidthOverride(WidgetWidth)
+			.HeightOverride(WidgetHeight)
 			[
 				AssetPickerWidget.ToSharedRef()
 			];
diff --git a/Engine/Source/Editor/UnrealEd/Classes/Settings/EditorStyleSettings.h b/Engine/Source/Editor/UnrealEd/Classes/Settings/EditorStyleSettings.h
--- a/Engine/Source/Editor/UnrealEd/Classes/Settings/EditorStyleSettings.h
+++ b/Engine/Source/Editor/UnrealEd/Classes/Settings/EditorStyleSettings.h
@@ -186,6 +186,10 @@ public:
 	UPROPERTY(EditAnywhere, config, Category=Text, meta = (ConfigRestartRequired = true))
 	bool bEnableMiddleEllipsis = true;
 
+	/** Scale the size of the global asset picker */
+	UPROPERTY(EditAnywhere, config, Category = UserInterface, meta = (DisplayName = "Scale Asset Picker Widget Size", ClampMin = "0.75", ClampMax = "3.0", UIMin = "0.75", UIMax = "3.0"))
+	float AssetPickerSizeScale = 1.0f;
+
 public:
 
 	/**
