Fixed GetOwnerPackages() for instanced structs inside non-complex properties

FStructurePropertyNode::GetOwnerPackages() fails at walking up the property nodes for a specific case, and may trigger an ensure. When stumbling upon an instanced struct, GetOwnerPackages() walks up by getting the direct parent. However, this parent is not necessarily a complex parent (can be a non-object/structure property type). To make sure the walk happens correctly, the case should be handled the same way GetInstancesNum() does (get complex parent of parent).

PR #13588: Fixed GetOwnerPackages() for instanced structs inside non-complex properties

diff --git a/Engine/Source/Editor/PropertyEditor/Private/StructurePropertyNode.h b/Engine/Source/Editor/PropertyEditor/Private/StructurePropertyNode.h
--- a/Engine/Source/Editor/PropertyEditor/Private/StructurePropertyNode.h
+++ b/Engine/Source/Editor/PropertyEditor/Private/StructurePropertyNode.h
@@ -116,11 +116,12 @@ public:
 				{
 					if (StructNode->StructProvider && StructNode->StructProvider->IsPropertyIndirection())
 					{
-						// Skip this and keep walking up to the next parent. This is assumed to be a struct that is pointed
+						// Skip this and keep walking up to the next complex parent. This is assumed to be a struct that is pointed
 						// to be a parent indirection.
 						// Note: in InstancedStruct case, will re-enter GetOwnerPackages during EnumerateInstances if calling GetInstances here.
 						// We want to avoid that otherwise it can cause poor performance as recursion branches into multiple recursions on each level.
 						Parent = ComplexParent->GetParentNode();
+						Parent = Parent ? Parent->FindComplexParent() : nullptr;
 						if (Parent == nullptr)
 						{
 							// If no owning property, then there is no package.
