Slate - Add ability to SDetailsView to reset to 'All' section when opened.
- Added a new option in the DetailsViewArgs to be able to clear the last selected section that is persisted for named detail views on first use ever.
- This would be particularly useful for detail views in modal windows where the Section filters will be persistent only for the lifetime of the modal window.

#rnx

diff --git a/Engine/Source/Editor/PropertyEditor/Private/SDetailsView.cpp b/Engine/Source/Editor/PropertyEditor/Private/SDetailsView.cpp
--- a/Engine/Source/Editor/PropertyEditor/Private/SDetailsView.cpp
+++ b/Engine/Source/Editor/PropertyEditor/Private/SDetailsView.cpp
@@ -1414,8 +1414,7 @@ void SDetailsView::RebuildSectionSelector()
 	for (const TSharedPtr<FComplexPropertyNode>& RootPropertyNode : RootPropertyNodes)
 	{
 		const UStruct* RootBaseStruct = RootPropertyNode->GetBaseStructure();
-		const FDetailsViewConfig* ViewConfig = GetConstViewConfig();
-		if (ViewConfig != nullptr)
+		if (const FDetailsViewConfig* ViewConfig = GetConstViewConfig())
 		{
 			const FDetailsSectionSelection* SectionSelection = ViewConfig->SelectedSections.Find(RootBaseStruct->GetFName());
 			if (SectionSelection != nullptr)
@@ -1545,6 +1544,32 @@ TMap<FName, SDetailsView::FSection> SDetailsView::GetAllSections() const
 	return MoveTemp(AllSections);
 }
 
+bool SDetailsView::ResetToDefaultSection()
+{
+	if (RootPropertyNodes.IsEmpty())
+	{
+		return false;
+	}
+
+	bool bResetViewConfig = false;
+	CurrentFilter.VisibleSections.Reset();
+	for (const TSharedPtr<FComplexPropertyNode>& RootPropertyNode : RootPropertyNodes)
+	{
+		if (const UStruct* RootBaseStruct = RootPropertyNode->GetBaseStructure())
+		{
+			if (FDetailsViewConfig* MutableViewConfig = GetMutableViewConfig())
+			{
+				FDetailsSectionSelection& SectionSelection = MutableViewConfig->SelectedSections.FindOrAdd(RootBaseStruct->GetFName());
+				SectionSelection.SectionNames = CurrentFilter.VisibleSections;
+				SaveViewConfig();
+				bResetViewConfig |= true;
+			}
+		}
+	}
+
+	return bResetViewConfig;
+}
+
 const FSlateBrush* SDetailsView::GetViewOptionsBadgeIcon() const
 {
 	// Badge the icon if any view option that narrows down the results is checked
diff --git a/Engine/Source/Editor/PropertyEditor/Private/SDetailsView.h b/Engine/Source/Editor/PropertyEditor/Private/SDetailsView.h
--- a/Engine/Source/Editor/PropertyEditor/Private/SDetailsView.h
+++ b/Engine/Source/Editor/PropertyEditor/Private/SDetailsView.h
@@ -132,6 +132,13 @@ public:
 	*/
 	virtual bool IsDefaultStyle() const override;
 
+	/**
+	 * Resets the selected sections to the default 'All' section.
+	 *
+	 * @return whether the operation succeeded.
+	 */
+	virtual bool ResetToDefaultSection() override;
+
 private:
 	/** Minimal struct to hold section information. Note that the Name is the key, so is not present here. */
 	struct FSection
diff --git a/Engine/Source/Editor/PropertyEditor/Private/SDetailsViewBase.h b/Engine/Source/Editor/PropertyEditor/Private/SDetailsViewBase.h
--- a/Engine/Source/Editor/PropertyEditor/Private/SDetailsViewBase.h
+++ b/Engine/Source/Editor/PropertyEditor/Private/SDetailsViewBase.h
@@ -206,6 +206,14 @@ public:
 	virtual TSharedRef<SWidget> GetScrollWidget() override;
 	// End of IScrollableWidget interface
 
+	/**
+	 * Resets the selected sections to the default 'All' section.
+	 * Base class returns false as sections selection is not application to all detail views.
+	 * 
+	 * @return whether the operation succeeded.
+	 */
+	virtual bool ResetToDefaultSection() override { return false; }
+
 protected:
 	/**
 	 * Called when a color property is changed from a color picker
diff --git a/Engine/Source/Editor/PropertyEditor/Public/IDetailsView.h b/Engine/Source/Editor/PropertyEditor/Public/IDetailsView.h
--- a/Engine/Source/Editor/PropertyEditor/Public/IDetailsView.h
+++ b/Engine/Source/Editor/PropertyEditor/Public/IDetailsView.h
@@ -329,4 +329,12 @@ public:
 	virtual FVector2D GetScrollDistanceRemaining() = 0;
 	virtual TSharedRef<SWidget> GetScrollWidget() = 0;
 	// End of IScrollableWidget interface
+
+	/**
+	 * Resets the selected sections to the default 'All' section.
+	 *
+	 * @return whether the operation succeeded.
+	 */
+	virtual bool ResetToDefaultSection() = 0;
+
 };
