[Property Editor Asset] allow browse to when there is a value set

Browse to is readonly. There are cases the reference is set elsewhere other than from the content picker, and we want to allow browsing to.

diff --git a/Engine/Source/Editor/PropertyEditor/Private/UserInterface/PropertyEditor/SPropertyEditorAsset.cpp b/Engine/Source/Editor/PropertyEditor/Private/UserInterface/PropertyEditor/SPropertyEditorAsset.cpp
--- a/Engine/Source/Editor/PropertyEditor/Private/UserInterface/PropertyEditor/SPropertyEditorAsset.cpp
+++ b/Engine/Source/Editor/PropertyEditor/Private/UserInterface/PropertyEditor/SPropertyEditorAsset.cpp
@@ -808,6 +808,16 @@ void SPropertyEditorAsset::Construct(const FArguments& InArgs, const TSharedPtr<
 		TAttribute<const FSlateBrush*> IconDelegate = !InArgs._OnBrowseOverride.IsBound() ?
 			TAttribute<const FSlateBrush*>::CreateSP(this, &SPropertyEditorAsset::GetOnBrowseIcon) : TAttribute<const FSlateBrush*>();
 
+		// Browse to is readonly. There are cases the reference is set elsewhere other than from the content picker, and we want to allow browsing to.
+		TAttribute<bool> IsBrowseEnabledAttribute = TAttribute<bool>::CreateSPLambda(this,
+				[this]()
+				{
+					FObjectOrAssetData Value;
+					GetValue(Value, FObjectOrAssetData::EAssetDataOptions::None);
+
+					return Value.IsValid();
+				});
+
 		ButtonBox->AddSlot()
 		.Padding( 2.0f, 0.0f )
 		.AutoWidth()
@@ -816,7 +826,7 @@ void SPropertyEditorAsset::Construct(const FArguments& InArgs, const TSharedPtr<
 			PropertyCustomizationHelpers::MakeBrowseButton(
 				OnBrowseDelegate,
 				TAttribute<FText>( this, &SPropertyEditorAsset::GetOnBrowseToolTip ),
-				IsEnabledAttribute,
+				IsBrowseEnabledAttribute,
 				bIsActor,
 				IconDelegate
 				)
