Fixed an issue where pasting text in a text property would not change the localization key of the text

diff --git a/Engine/Source/Editor/PropertyEditor/Private/PropertyHandleImpl.cpp b/Engine/Source/Editor/PropertyEditor/Private/PropertyHandleImpl.cpp
--- a/Engine/Source/Editor/PropertyEditor/Private/PropertyHandleImpl.cpp
+++ b/Engine/Source/Editor/PropertyEditor/Private/PropertyHandleImpl.cpp
@@ -5877,8 +5877,9 @@ FPropertyAccess::Result FPropertyHandleText::ImportFormattedTextString(const FSt
 				if (TextProperty && NewTextValue.IsInitializedFromString())
 				{
 					UPackage* OuterPackage = OuterPackages.IsValidIndex(ObjectIndex) ? OuterPackages[ObjectIndex] : nullptr;
+					UObject* Object = ObjectToModify.Object;
 
-					auto DeterministicTextKeyGenerator = [PropertyPath, OuterPackage, TextProperty]()
+					auto DeterministicTextKeyGenerator = [PropertyPath, OuterPackage, Object, TextProperty]()
 					{
 						auto GetNameKeyHash = [](const FName Name)
 						{
@@ -5896,8 +5897,9 @@ FPropertyAccess::Result FPropertyHandleText::ImportFormattedTextString(const FSt
 							return TextKeyUtil::HashString(Builder.ToString(), Builder.Len());
 						};
 
-						// Build a (hopefully) unique and deterministic key from a combination of the outer package and text property info
-						FGuid KeyGuid(GetObjectKeyHash(OuterPackage), TextKeyUtil::HashString(PropertyPath), GetNameKeyHash(TextProperty->GetOwnerStruct()->GetFName()), GetNameKeyHash(TextProperty->GetFName()));
+						// Build a unique and deterministic key from a combination of the outer package, object and text property info.
+						// We can use both the Object or Outerpackage because in some cases (like for structs), we might not have an object but, if we do, we need to use that to create a unique key.
+						FGuid KeyGuid(GetObjectKeyHash(Object != nullptr ? Object : OuterPackage), TextKeyUtil::HashString(PropertyPath), GetNameKeyHash(TextProperty->GetOwnerStruct()->GetFName()), GetNameKeyHash(TextProperty->GetFName()));
 #if USE_STABLE_LOCALIZATION_KEYS
 						if (OuterPackage)
 						{
diff --git a/Engine/Source/Runtime/CoreUObject/Private/Internationalization/TextPackageNamespaceUtil.cpp b/Engine/Source/Runtime/CoreUObject/Private/Internationalization/TextPackageNamespaceUtil.cpp
--- a/Engine/Source/Runtime/CoreUObject/Private/Internationalization/TextPackageNamespaceUtil.cpp
+++ b/Engine/Source/Runtime/CoreUObject/Private/Internationalization/TextPackageNamespaceUtil.cpp
@@ -262,6 +262,20 @@ void TextNamespaceUtil::GetTextIdForEdit(UPackage* InPackage, const ETextEditAct
 #endif // USE_STABLE_LOCALIZATION_KEYS
 	{
 		OutStableKey = InTextKeyGenerator();
+
+#if USE_STABLE_LOCALIZATION_KEYS
+		// Make sure the new key didn't produce a conflict (as may sometimes happen when using deterministic text key generators).
+		// If it has, generate a new key using the random text key generator instead.
+		{
+			int32 ReferenceCount = 0;
+			FTextReferenceCollector Collector(InPackage, FTextReferenceCollector::EComparisonMode::MismatchSource, OutStableNamespace, OutStableKey, InTextSource, ReferenceCount);
+
+			if (ReferenceCount != 0)
+			{
+				OutStableKey = GenerateRandomTextKey();
+			}
+		}
+#endif // USE_STABLE_LOCALIZATION_KEYS
 	}
 }
 
