[Property Asset Editor] enabled BrowseTo when either the picker is enabled or the value is set elsewhere.

diff --git a/Engine/Source/Editor/PropertyEditor/Private/UserInterface/PropertyEditor/SPropertyEditorAsset.cpp b/Engine/Source/Editor/PropertyEditor/Private/UserInterface/PropertyEditor/SPropertyEditorAsset.cpp
--- a/Engine/Source/Editor/PropertyEditor/Private/UserInterface/PropertyEditor/SPropertyEditorAsset.cpp
+++ b/Engine/Source/Editor/PropertyEditor/Private/UserInterface/PropertyEditor/SPropertyEditorAsset.cpp
@@ -809,9 +809,15 @@ void SPropertyEditorAsset::Construct(const FArguments& InArgs, const TSharedPtr<
 			TAttribute<const FSlateBrush*>::CreateSP(this, &SPropertyEditorAsset::GetOnBrowseIcon) : TAttribute<const FSlateBrush*>();
 
 		// Browse to is readonly. There are cases the reference is set elsewhere other than from the content picker, and we want to allow browsing to.
+		bool bOldIsEnabledAttribute = IsEnabledAttribute.Get();
 		TAttribute<bool> IsBrowseEnabledAttribute = TAttribute<bool>::CreateSPLambda(this,
-				[this]()
+				[bOldIsEnabledAttribute, this]()
 				{
+					if (bOldIsEnabledAttribute)
+					{
+						return true;
+					}
+
 					FObjectOrAssetData Value;
 					GetValue(Value, FObjectOrAssetData::EAssetDataOptions::None);
 
