Allow struct type picking when class type picking is enabled

Data Table row editing sets bShowObjects to false so that Data Table assets aren't allowed to reference non-asset objects (such as actors). It tries to allow asset and class picking (bShowAssets and bShowClasses), but had missed struct type picking (as that manifests as a FObjectPropertyBase with a PropertyClass of UScriptStruct) due to it taking the default bShowObjects filter path.

This change also allows struct picking when bShowClasses is true, essentially making bShowClasses behave as allowing any type picking (structs or classes).
Note: I didn't add a new bShowStructs flag as that may be confused with real struct properties (FStructProperty).

diff --git a/Engine/Source/Editor/PropertyEditor/Private/PropertyEditorModule.cpp b/Engine/Source/Editor/PropertyEditor/Private/PropertyEditorModule.cpp
--- a/Engine/Source/Editor/PropertyEditor/Private/PropertyEditorModule.cpp
+++ b/Engine/Source/Editor/PropertyEditor/Private/PropertyEditorModule.cpp
@@ -1149,6 +1149,11 @@ TSharedRef<class IStructureDetailsView> FPropertyEditorModule::CreateStructureDe
 			const auto ObjectProperty = CastField<FObjectPropertyBase>(PropertyToTest);
 			if( ObjectProperty )
 			{
+				if (InStructureDetailsViewArgs.bShowClasses && ObjectProperty->PropertyClass && ObjectProperty->PropertyClass->IsChildOf(UScriptStruct::StaticClass()))
+				{
+					return true;
+				}
+
 				if( InStructureDetailsViewArgs.bShowAssets )
 				{
 					// Is this an "asset" property?
