trivial rename

diff --git a/Engine/Source/Editor/PropertyEditor/Private/UserInterface/PropertyEditor/SPropertyEditorAsset.cpp b/Engine/Source/Editor/PropertyEditor/Private/UserInterface/PropertyEditor/SPropertyEditorAsset.cpp
--- a/Engine/Source/Editor/PropertyEditor/Private/UserInterface/PropertyEditor/SPropertyEditorAsset.cpp
+++ b/Engine/Source/Editor/PropertyEditor/Private/UserInterface/PropertyEditor/SPropertyEditorAsset.cpp
@@ -809,11 +809,11 @@ void SPropertyEditorAsset::Construct(const FArguments& InArgs, const TSharedPtr<
 			TAttribute<const FSlateBrush*>::CreateSP(this, &SPropertyEditorAsset::GetOnBrowseIcon) : TAttribute<const FSlateBrush*>();
 
 		// Browse to is readonly. There are cases the reference is set elsewhere other than from the content picker, and we want to allow browsing to.
-		bool bOldIsEnabledAttribute = IsEnabledAttribute.Get();
+		bool bIsPickerEnabled = IsEnabledAttribute.Get();
 		TAttribute<bool> IsBrowseEnabledAttribute = TAttribute<bool>::CreateSPLambda(this,
-				[bOldIsEnabledAttribute, this]()
+				[bIsPickerEnabled, this]()
 				{
-					if (bOldIsEnabledAttribute)
+					if (bIsPickerEnabled)
 					{
 						return true;
 					}
