Avoid expensive ExportText calls when drawing array properties.
Takes editor window draw in a licensee project from ~43ms to ~35ms.

diff --git a/Engine/Source/Editor/PropertyEditor/Private/PropertyEditorHelpers.cpp b/Engine/Source/Editor/PropertyEditor/Private/PropertyEditorHelpers.cpp
--- a/Engine/Source/Editor/PropertyEditor/Private/PropertyEditorHelpers.cpp
+++ b/Engine/Source/Editor/PropertyEditor/Private/PropertyEditorHelpers.cpp
@@ -1052,8 +1052,8 @@ namespace PropertyEditorHelpers
 				FArrayProperty* ArrayProperty = CastField<FArrayProperty>(Property);
 				if (ArrayProperty)
 				{
-					FString ArrayString;
-					FPropertyAccess::Result GetValResult = PropertyEditor->GetPropertyHandle()->GetValueAsDisplayString(ArrayString);
+					void* ValueAddress;
+					FPropertyAccess::Result GetValResult = PropertyEditor->GetPropertyHandle()->GetValueData(ValueAddress);
 					return IsPropertyButtonEnabled(WeakPropertyNode) && GetValResult == FPropertyAccess::Success;
 				}
 			}
diff --git a/Engine/Source/Editor/PropertyEditor/Private/UserInterface/PropertyEditor/SPropertyEditorArray.cpp b/Engine/Source/Editor/PropertyEditor/Private/UserInterface/PropertyEditor/SPropertyEditorArray.cpp
--- a/Engine/Source/Editor/PropertyEditor/Private/UserInterface/PropertyEditor/SPropertyEditorArray.cpp
+++ b/Engine/Source/Editor/PropertyEditor/Private/UserInterface/PropertyEditor/SPropertyEditorArray.cpp
@@ -59,8 +59,8 @@ void SPropertyEditorArray::GetDesiredWidth(float& OutMinDesiredWidth, float& Out
 
 FText SPropertyEditorArray::GetArrayTextValue() const
 {
-	FString ArrayString;
-	FPropertyAccess::Result GetValResult = PropertyEditor->GetPropertyHandle()->GetValueAsDisplayString(ArrayString);
+	void* ValueAddress;
+	FPropertyAccess::Result GetValResult = PropertyEditor->GetPropertyHandle()->GetValueData(ValueAddress);
 
 	if (GetValResult == FPropertyAccess::MultipleValues)
 	{
