Improve construction performance of SEditableTextBox & SPropertyEditorAsset.
Takes rebuilding a specific detail panel from ~950ms -> ~650ms in a licensee project.

diff --git a/Engine/Source/Editor/PropertyEditor/Private/UserInterface/PropertyEditor/SPropertyEditorAsset.cpp b/Engine/Source/Editor/PropertyEditor/Private/UserInterface/PropertyEditor/SPropertyEditorAsset.cpp
--- a/Engine/Source/Editor/PropertyEditor/Private/UserInterface/PropertyEditor/SPropertyEditorAsset.cpp
+++ b/Engine/Source/Editor/PropertyEditor/Private/UserInterface/PropertyEditor/SPropertyEditorAsset.cpp
@@ -369,6 +369,7 @@ void SPropertyEditorAsset::Construct(const FArguments& InArgs, const TSharedPtr<
 	OnShouldFilterActor = InArgs._OnShouldFilterActor;
 	ObjectPath = InArgs._ObjectPath;
 	bDisplayUseSelected = InArgs._DisplayUseSelected;
+	NewAssetFactoriesOverride = InArgs._NewAssetFactories;
 
 	// Override this as we stole the value to use as OnIsEnabled for the inner widgets
 	SetEnabled(true);
@@ -502,19 +503,6 @@ void SPropertyEditorAsset::Construct(const FArguments& InArgs, const TSharedPtr<
 
 	bIsActor = ObjectClass->IsChildOf(AActor::StaticClass());
 
-	if (bAllowCreate)
-	{
-		if (InArgs._NewAssetFactories.IsSet())
-		{
-			NewAssetFactories = InArgs._NewAssetFactories.GetValue();
-		}
-		// If there are more allowed classes than just UObject 
-		else if (AllowedClassFilters.Num() > 1 || !AllowedClassFilters.Contains(UObject::StaticClass()))
-		{
-			NewAssetFactories = PropertyCustomizationHelpers::GetNewAssetFactoriesForClasses(AllowedClassFilters, DisallowedClassFilters);
-		}
-	}
-	
 	TSharedPtr<SHorizontalBox> ValueContentBox = nullptr;
 	ChildSlot
 	[
@@ -1096,6 +1084,20 @@ TSharedRef<SWidget> SPropertyEditorAsset::OnGetMenuContent()
 	}
 	else
 	{
+		TArray<UFactory*> NewAssetFactories;
+		if (bAllowCreate)
+		{
+			if (NewAssetFactoriesOverride.IsSet())
+			{
+				NewAssetFactories = NewAssetFactoriesOverride.GetValue();
+			}
+			// If there are more allowed classes than just UObject
+			else if (AllowedClassFilters.Num() > 1 || !AllowedClassFilters.Contains(UObject::StaticClass()))
+			{
+				NewAssetFactories = PropertyCustomizationHelpers::GetNewAssetFactoriesForClasses(AllowedClassFilters, DisallowedClassFilters);
+			}
+		}
+
 		return PropertyCustomizationHelpers::MakeAssetPickerWithMenu(Value.AssetData,
 																	 bAllowClear,
 																	 AllowedClassFilters,
diff --git a/Engine/Source/Editor/PropertyEditor/Private/UserInterface/PropertyEditor/SPropertyEditorAsset.h b/Engine/Source/Editor/PropertyEditor/Private/UserInterface/PropertyEditor/SPropertyEditorAsset.h
--- a/Engine/Source/Editor/PropertyEditor/Private/UserInterface/PropertyEditor/SPropertyEditorAsset.h
+++ b/Engine/Source/Editor/PropertyEditor/Private/UserInterface/PropertyEditor/SPropertyEditorAsset.h
@@ -416,7 +416,7 @@ private:
 	TArray<const UClass*> DisallowedClassFilters;
 
 	/** A list of the factories we can use to create new assets */
-	TArray<UFactory*> NewAssetFactories;
+	TOptional<TArray<UFactory*>> NewAssetFactoriesOverride;
 
 	/** Tags (and eventually values) that can NOT be used with this property */
 	TSharedPtr<FAssetDataTagMap> DisallowedAssetDataTags;
diff --git a/Engine/Source/Runtime/Slate/Private/Widgets/Input/SEditableTextBox.cpp b/Engine/Source/Runtime/Slate/Private/Widgets/Input/SEditableTextBox.cpp
--- a/Engine/Source/Runtime/Slate/Private/Widgets/Input/SEditableTextBox.cpp
+++ b/Engine/Source/Runtime/Slate/Private/Widgets/Input/SEditableTextBox.cpp
@@ -102,15 +102,6 @@ void SEditableTextBox::Construct( const FArguments& InArgs )
 			ErrorReporting->AsWidget()
 		];
 	}
-	else
-	{
-		// this also creates a default widget
-		// if we don't create the widget in Construct() 
-		// it will get created in OnEditableTextChanged()
-		// create it now so that the default size of the textbox
-		// won't grow after user use it once
-		SetError(FText::GetEmpty());
-	}
 }
 
 void SEditableTextBox::SetStyle(const FEditableTextBoxStyle* InStyle)
@@ -159,6 +150,16 @@ void SEditableTextBox::SetError( const FText& InError )
 void SEditableTextBox::SetError( const FString& InError )
 {
 	const bool bHaveError = !InError.IsEmpty();
+	if (!bHaveError)
+	{
+		if (ErrorReporting.IsValid())
+		{
+			ErrorReporting->AsWidget()->SetVisibility(EVisibility::Collapsed);
+			ErrorReporting->SetError(FString());
+		}
+
+		return;
+	}
 
 	if ( !ErrorReporting.IsValid() )
 	{
@@ -173,6 +174,7 @@ void SEditableTextBox::SetError( const FString& InError )
 		ErrorReporting = ErrorTextWidget;
 	}
 
+	ErrorReporting->AsWidget()->SetVisibility(EVisibility::Visible);
 	ErrorReporting->SetError( InError );
 }
 
