Add GetActorFilter UPROPERTY meta support to object picker (mirroring GetAssetFilter).

diff --git a/Engine/Source/Editor/PropertyEditor/Private/UserInterface/PropertyEditor/SPropertyEditorAsset.cpp b/Engine/Source/Editor/PropertyEditor/Private/UserInterface/PropertyEditor/SPropertyEditorAsset.cpp
--- a/Engine/Source/Editor/PropertyEditor/Private/UserInterface/PropertyEditor/SPropertyEditorAsset.cpp
+++ b/Engine/Source/Editor/PropertyEditor/Private/UserInterface/PropertyEditor/SPropertyEditorAsset.cpp
@@ -406,6 +406,20 @@ void SPropertyEditorAsset::Construct(const FArguments& InArgs, const TSharedPtr<
 	bIsSoftObjectPath = CastField<FSoftObjectProperty>(Property) != nullptr;
 	
 	InitializeAssetDataTags(Property);
+	InitializeClassFilters(Property);
+
+	// Make the ObjectClass more specific if we only have one class filter
+	// eg. if ObjectClass was set to Actor, but AllowedClasses="PointLight", we can limit it to PointLight immediately
+	if (AllowedClassFilters.Num() == 1 && DisallowedClassFilters.Num() == 0)
+	{
+		ObjectClass = const_cast<UClass*>(AllowedClassFilters[0]);
+	}
+	else
+	{
+		ObjectClass = UClass::FindCommonBase(ConstCastClassArray(AllowedClassFilters));
+	}
+
+	bIsActor = ObjectClass->IsChildOf(AActor::StaticClass());
 
 	auto AppendOnShouldFilterAssetCallback = [this](FOnShouldFilterAsset&& OnShouldFilterAssetCallback)
 	{
@@ -425,17 +439,37 @@ void SPropertyEditorAsset::Construct(const FArguments& InArgs, const TSharedPtr<
 
 	OnShouldFilterAsset = InArgs._OnShouldFilterAsset;
 	
+	auto AppendOnShouldFilterActorCallback = [this](FOnShouldFilterActor&& OnShouldFilterActorCallback)
+	{
+		check(OnShouldFilterActorCallback.IsBound());
+		if (OnShouldFilterActor.IsBound())
+		{
+			OnShouldFilterActor.BindLambda([BaseOnShouldFilterActor = OnShouldFilterActor, OnShouldFilterActorCallback = MoveTemp(OnShouldFilterActorCallback)](const AActor* InActor)
+			{
+				return BaseOnShouldFilterActor.Execute(InActor) || OnShouldFilterActorCallback.Execute(InActor);
+			});
+		}
+		else
+		{
+			OnShouldFilterActor = MoveTemp(OnShouldFilterActorCallback);
+		}
+	};
+
 	if (DisallowedAssetDataTags.IsValid() || RequiredAssetDataTags.IsValid())
 	{
 		// re-route the filter delegate to our own if we have our own asset data tags filter :
 		AppendOnShouldFilterAssetCallback(FOnShouldFilterAsset::CreateRaw(this, &SPropertyEditorAsset::IsAssetFiltered));
 	}
 
-	if (Property && Property->GetOwnerProperty()->HasMetaData("GetAssetFilter"))
+	static const FName GetAssetFilterMetaName = "GetAssetFilter";
+	static const FName GetActorFilterMetaName = "GetActorFilter";
+
+	const FName GetFilterMetaName = bIsActor ? GetActorFilterMetaName : GetAssetFilterMetaName;
+	if (Property && Property->GetOwnerProperty()->HasMetaData(GetFilterMetaName))
 	{
 		// Add MetaData asset filter
-		const FString GetAssetFilterFunctionName = Property->GetOwnerProperty()->GetMetaData("GetAssetFilter");
-		if (!GetAssetFilterFunctionName.IsEmpty())
+		const FString GetFilterFunctionName = Property->GetOwnerProperty()->GetMetaData(GetFilterMetaName);
+		if (!GetFilterFunctionName.IsEmpty())
 		{
 			TArray<UObject*> ObjectList;
 			if (PropertyEditor.IsValid())
@@ -452,8 +486,8 @@ void SPropertyEditorAsset::Construct(const FArguments& InArgs, const TSharedPtr<
 				{
 					continue;
 				}
-				const UFunction* GetAssetFilterFunction = Object->FindFunction(*GetAssetFilterFunctionName);
-				if (ensureMsgf(GetAssetFilterFunction, TEXT("Could not find UFunction %s on %s"), *GetAssetFilterFunctionName, *GetNameSafe(Object)))
+				const UFunction* GetFilterFunction = Object->FindFunction(*GetFilterFunctionName);
+				if (ensureMsgf(GetFilterFunction, TEXT("Could not find UFunction %s on %s"), *GetFilterFunctionName, *GetNameSafe(Object)))
 				{
 					if (UActorComponent* ActorComponent = Cast<UActorComponent>(Object) )
 					{
@@ -462,7 +496,29 @@ void SPropertyEditorAsset::Construct(const FArguments& InArgs, const TSharedPtr<
 						FSoftComponentReference ComponentReference;
 						ComponentReference.OtherActor = ActorComponent->GetOwner();
 						ComponentReference.PathToComponent = ActorComponent->GetPathName(ActorComponent->GetOwner());
-						AppendOnShouldFilterAssetCallback(FOnShouldFilterAsset::CreateLambda([ComponentSoftPtr, ComponentReference, FunctionName =  GetAssetFilterFunction->GetFName()](const FAssetData& AssetData) mutable {
+
+						if (bIsActor)
+						{
+							AppendOnShouldFilterActorCallback(FOnShouldFilterActor::CreateLambda([ComponentSoftPtr, ComponentReference, FunctionName =  GetFilterFunction->GetFName()](const AActor* InActor) mutable {
+								if (!ComponentSoftPtr.IsValid())
+								{
+									ComponentSoftPtr = ComponentReference.GetComponent(nullptr);
+								}
+								if (UActorComponent* Component = ComponentSoftPtr.Get())
+								{
+									FOnShouldFilterActor UFunctionDelegate = FOnShouldFilterActor::CreateUFunction(Component, FunctionName);
+									if (UFunctionDelegate.IsBound())
+									{
+										return UFunctionDelegate.Execute(InActor);
+									}
+								}
+								return false;
+							}));
+
+							continue;
+						}
+
+						AppendOnShouldFilterAssetCallback(FOnShouldFilterAsset::CreateLambda([ComponentSoftPtr, ComponentReference, FunctionName =  GetFilterFunction->GetFName()](const FAssetData& AssetData) mutable {
 							if (!ComponentSoftPtr.IsValid())
 							{
 								ComponentSoftPtr = ComponentReference.GetComponent(nullptr);
@@ -479,30 +535,19 @@ void SPropertyEditorAsset::Construct(const FArguments& InArgs, const TSharedPtr<
 						}));
 						
 					}
+					else if (bIsActor)
+					{
+						AppendOnShouldFilterActorCallback(FOnShouldFilterActor::CreateUFunction(Object, GetFilterFunction->GetFName()));
+					}
 					else
 					{
-						AppendOnShouldFilterAssetCallback(FOnShouldFilterAsset::CreateUFunction(Object, GetAssetFilterFunction->GetFName()));
+						AppendOnShouldFilterAssetCallback(FOnShouldFilterAsset::CreateUFunction(Object, GetFilterFunction->GetFName()));
 					}
 				}
 			}
 		}
 	}
 
-	InitializeClassFilters(Property);
-
-	// Make the ObjectClass more specific if we only have one class filter
-	// eg. if ObjectClass was set to Actor, but AllowedClasses="PointLight", we can limit it to PointLight immediately
-	if (AllowedClassFilters.Num() == 1 && DisallowedClassFilters.Num() == 0)
-	{
-		ObjectClass = const_cast<UClass*>(AllowedClassFilters[0]);
-	}
-	else
-	{
-		ObjectClass = UClass::FindCommonBase(ConstCastClassArray(AllowedClassFilters));
-	}
-
-	bIsActor = ObjectClass->IsChildOf(AActor::StaticClass());
-
 	TSharedPtr<SHorizontalBox> ValueContentBox = nullptr;
 	ChildSlot
 	[
