Asset Picker: Add a AssetFactoryContextName meta tag with which it's possible to provide a context name to be passed to the factory when a new asset is created via the picker.

diff --git a/Engine/Source/Editor/PropertyEditor/Private/UserInterface/PropertyEditor/SPropertyMenuAssetPicker.cpp b/Engine/Source/Editor/PropertyEditor/Private/UserInterface/PropertyEditor/SPropertyMenuAssetPicker.cpp
--- a/Engine/Source/Editor/PropertyEditor/Private/UserInterface/PropertyEditor/SPropertyMenuAssetPicker.cpp
+++ b/Engine/Source/Editor/PropertyEditor/Private/UserInterface/PropertyEditor/SPropertyMenuAssetPicker.cpp
@@ -369,6 +369,13 @@ void SPropertyMenuAssetPicker::OnCreateNewAssetSelected(TWeakObjectPtr<UFactory>
 	TStrongObjectPtr<UFactory> FactoryStrong(FactoryInstance);
 	FAssetToolsModule& AssetToolsModule = FAssetToolsModule::GetModule();
 
+	static const FName ContextNameMetaTag = "AssetFactoryContextName";
+	FName ContextName;
+	if (PropertyHandle && PropertyHandle->HasMetaData(ContextNameMetaTag))
+	{
+		ContextName = FName(PropertyHandle->GetMetaData(ContextNameMetaTag));
+	}
+
 	if (EnumHasAnyFlags(FactoryInstance->GetSupportedWorkflows(), EFactoryCreateWorkflow::Asynchronous))
 	{
 		AssetToolsModule.Get().CreateAssetWithDialogAsync(
@@ -379,12 +386,14 @@ void SPropertyMenuAssetPicker::OnCreateNewAssetSelected(TWeakObjectPtr<UFactory>
 				[this](UFactory* Factory, UObject* InObject)
 				{
 					SetValue(InObject);
-				})
+				}),
+			FAssetCreateCancelled(),
+			ContextName
 		);
 	}
 	else if (EnumHasAnyFlags(FactoryInstance->GetSupportedWorkflows(), EFactoryCreateWorkflow::Default))
 	{
-		UObject* NewAsset = AssetToolsModule.Get().CreateAssetWithDialog(FactoryInstance->GetSupportedClass(), FactoryInstance);
+		UObject* NewAsset = AssetToolsModule.Get().CreateAssetWithDialog(FactoryInstance->GetSupportedClass(), FactoryInstance, ContextName);
 		if (NewAsset != nullptr)
 		{
 			SetValue(NewAsset);
