Slate: Add custom menu options to details view.

#rnx

diff --git a/Engine/Source/Editor/PropertyEditor/Private/SDetailsView.cpp b/Engine/Source/Editor/PropertyEditor/Private/SDetailsView.cpp
--- a/Engine/Source/Editor/PropertyEditor/Private/SDetailsView.cpp
+++ b/Engine/Source/Editor/PropertyEditor/Private/SDetailsView.cpp
@@ -114,7 +114,9 @@ void SDetailsView::Construct(const FArguments& InArgs, const FDetailsViewArgs& I
 	const TSharedRef<SScrollBar> ExternalScrollbar = InDetailsViewArgs.ExternalScrollbar ? InDetailsViewArgs.ExternalScrollbar.ToSharedRef() : SNew(SScrollBar);
 	ExternalScrollbar->SetVisibility( TAttribute<EVisibility>( this, &SDetailsView::GetScrollBarVisibility ) );
 
-	FMenuBuilder DetailViewOptions( true, nullptr);
+	FMenuBuilder DetailViewOptions(true, nullptr, InDetailsViewArgs.OptionsExtender);
+
+	DetailViewOptions.BeginSection("DetailsViewShowOptions");
 
 	if (DetailsViewArgs.bShowModifiedPropertiesOption)
 	{
@@ -253,6 +255,9 @@ void SDetailsView::Construct(const FArguments& InArgs, const FDetailsViewArgs& I
 		EUserInterfaceActionType::Check
 	);		
 
+	DetailViewOptions.EndSection();
+	DetailViewOptions.BeginSection("DetailsViewCategories");
+
 	DetailViewOptions.AddMenuEntry(
 		LOCTEXT("CollapseAll", "Collapse All Categories"),
 		LOCTEXT("CollapseAll_ToolTip", "Collapses all root level categories"),
@@ -267,6 +272,8 @@ void SDetailsView::Construct(const FArguments& InArgs, const FDetailsViewArgs& I
 
 	TSharedRef<SVerticalBox> VerticalBox = SNew(SVerticalBox);
 		
+	DetailViewOptions.EndSection();
+
 	// Create the name area which does not change when selection changes
 	SAssignNew(NameArea, SDetailNameArea, &SelectedObjects)
 		// the name area is only for actors
diff --git a/Engine/Source/Editor/PropertyEditor/Public/DetailsViewArgs.h b/Engine/Source/Editor/PropertyEditor/Public/DetailsViewArgs.h
--- a/Engine/Source/Editor/PropertyEditor/Public/DetailsViewArgs.h
+++ b/Engine/Source/Editor/PropertyEditor/Public/DetailsViewArgs.h
@@ -8,6 +8,7 @@
 
 class FDetailsNameWidgetOverrideCustomization;
 class FDetailsViewObjectFilter;
+class FExtender;
 class FNotifyHook;
 class FTabManager;
 class FUICommandList;
@@ -61,6 +62,8 @@ struct FDetailsViewArgs
 	TSharedPtr<FDetailsViewStyleKey> StyleKey;
 	/** A Customization that can be applied to all name widgets */
 	TSharedPtr<FDetailsNameWidgetOverrideCustomization> DetailsNameWidgetOverrideCustomization;
+	/** If bShowOptions menu is available, adds support for user defined menu options*/
+	TSharedPtr<FExtender> OptionsExtender;
 	/** Identifier for this details view; NAME_None if this view is anonymous */
 	FName ViewIdentifier;
 	/** Notify hook to call when properties are changed */
