Fixed FChildActorComponentDetails to not add non-class types to the FPropertyRestriction for ChildActorClass.

Additionally, moved the FindOrLoadClass function into PropertyEditorHelpers.h. Also made it a bit noisier if the given asset path doesn't actually correspond to a class type.

#rnx

diff --git a/Engine/Source/Editor/Kismet/Private/BlueprintDetailsCustomization.cpp b/Engine/Source/Editor/Kismet/Private/BlueprintDetailsCustomization.cpp
--- a/Engine/Source/Editor/Kismet/Private/BlueprintDetailsCustomization.cpp
+++ b/Engine/Source/Editor/Kismet/Private/BlueprintDetailsCustomization.cpp
@@ -8101,7 +8101,6 @@ void FChildActorComponentDetails::CustomizeDetails(IDetailLayoutBuilder& DetailB
 				static FText RestrictReason = LOCTEXT("NoSelfChildActors", "Cannot append a child-actor of this blueprint type (could cause infinite recursion).");
 				TSharedPtr<FPropertyRestriction> ClassRestriction = MakeShareable(new FPropertyRestriction(RestrictReason));
 
-				ClassRestriction->AddDisabledValue(Blueprint->GetPathName());
 				if (Blueprint->GeneratedClass)
 				{
 					ClassRestriction->AddDisabledValue(Blueprint->GeneratedClass->GetPathName());
diff --git a/Engine/Source/Editor/PropertyEditor/Private/PropertyEditorHelpers.cpp b/Engine/Source/Editor/PropertyEditor/Private/PropertyEditorHelpers.cpp
--- a/Engine/Source/Editor/PropertyEditor/Private/PropertyEditorHelpers.cpp
+++ b/Engine/Source/Editor/PropertyEditor/Private/PropertyEditorHelpers.cpp
@@ -6,6 +6,7 @@
 #include "IDocumentation.h"
 
 #include "PropertyHandleImpl.h"
+#include "PropertyNode.h"
 
 #include "UserInterface/PropertyEditor/PropertyEditorConstants.h"
 #include "UserInterface/PropertyEditor/SPropertyEditor.h"
@@ -1448,6 +1449,27 @@ namespace PropertyEditorHelpers
 
 		return NAME_None;
 	}
+
+	const UClass* FindOrLoadClass(const FString& ClassPath)
+	{
+		if (const UObject* Object = StaticLoadObject(nullptr, nullptr, ClassPath))
+		{
+			if (const UClass* Class = Cast<UClass>(Object))
+			{
+				return Class;
+			}
+			else
+			{
+				UE_LOG(LogPropertyNode, Error, TEXT("Found object '%s' instead of a class."), *Object->GetFullName());
+			}
+		}
+		else
+		{
+			UE_LOG(LogPropertyNode, Error, TEXT("Unable to load class '%s'."), *ClassPath);
+		}
+
+		return nullptr;
+	}
 }
 
 #undef LOCTEXT_NAMESPACE
diff --git a/Engine/Source/Editor/PropertyEditor/Private/PropertyEditorHelpers.h b/Engine/Source/Editor/PropertyEditor/Private/PropertyEditorHelpers.h
--- a/Engine/Source/Editor/PropertyEditor/Private/PropertyEditorHelpers.h
+++ b/Engine/Source/Editor/PropertyEditor/Private/PropertyEditorHelpers.h
@@ -343,5 +343,13 @@ namespace PropertyEditorHelpers
 	* Returns NAME_None if the property doesn't support, or doesn't have, options.
 	*/
 	FName GetPropertyOptionsMetaDataKey(const FProperty* Property);
+
+	/**
+	 * Attempts to find the class with the given object path.
+	 * @param ClassPath  The asset path of the class to find or load (eg: /Game/Blueprints/BP_Example.BP_Example_C)
+	 * 
+	 * @return nullptr if the class could not be found, or if the path doesn't correspond to a class object
+	 */
+	const UClass* FindOrLoadClass(const FString& ClassPath);
 }
 
diff --git a/Engine/Source/Editor/PropertyEditor/Private/UserInterface/PropertyEditor/SPropertyEditorClass.cpp b/Engine/Source/Editor/PropertyEditor/Private/UserInterface/PropertyEditor/SPropertyEditorClass.cpp
--- a/Engine/Source/Editor/PropertyEditor/Private/UserInterface/PropertyEditor/SPropertyEditorClass.cpp
+++ b/Engine/Source/Editor/PropertyEditor/Private/UserInterface/PropertyEditor/SPropertyEditorClass.cpp
@@ -13,6 +13,7 @@
 #include "UObject/UObjectIterator.h"
 #include "UObject/VerseClassProperty.h"
 #include "VerseClassConstraintFilters.h"
+#include "PropertyEditorHelpers.h"
 #include "PropertyNode.h"
 #include "PropertyRestriction.h"
 
@@ -88,23 +89,6 @@ private:
 	}
 };
 
-namespace UE::PropertyEditor::Class::Private
-{
-
-static UClass* FindOrLoadClass(const FString& ClassName)
-{
-	UClass* Class = UClass::TryFindTypeSlow<UClass>(ClassName, EFindFirstObjectOptions::EnsureIfAmbiguous);
-
-	if (!Class)
-	{
-		Class = LoadObject<UClass>(nullptr, *ClassName);
-	}
-
-	return Class;
-}
-
-} // namespace UE::PropertyEditor::Class::Private
-
 void SPropertyEditorClass::GetDesiredWidth(float& OutMinDesiredWidth, float& OutMaxDesiredWidth)
 {
 	OutMinDesiredWidth = 200.0f;
@@ -189,14 +173,12 @@ void SPropertyEditorClass::Construct(const FArguments& InArgs, const TSharedPtr<
 		}
 		PropertyEditorUtils::GetAllowedAndDisallowedClasses(ObjectList, *Property, AllowedClassFilters, DisallowedClassFilters, false);
 
-		using namespace UE::PropertyEditor::Class::Private;
-
 		// Filter based on restrictions
 		for (const TSharedRef<const FPropertyRestriction>& ClassRestriction : PropertyNode->GetRestrictions())
 		{
 			for (TArray<FString>::TConstIterator Iter = ClassRestriction.Get().GetHiddenValuesIterator(); Iter; ++Iter)
 			{
-				if (const UClass* HiddenClass = FindOrLoadClass(*Iter))
+				if (const UClass* HiddenClass = PropertyEditorHelpers::FindOrLoadClass(*Iter))
 				{
 					DisallowedClassFilters.Add(HiddenClass);
 				}
@@ -204,7 +186,7 @@ void SPropertyEditorClass::Construct(const FArguments& InArgs, const TSharedPtr<
 
 			for (TArray<FString>::TConstIterator Iter = ClassRestriction.Get().GetDisabledValuesIterator(); Iter; ++Iter)
 			{
-				if (const UClass* DisabledClass = FindOrLoadClass(*Iter))
+				if (const UClass* DisabledClass = PropertyEditorHelpers::FindOrLoadClass(*Iter))
 				{
 					DisallowedClassFilters.Add(DisabledClass);
 				}
@@ -376,7 +358,7 @@ TSharedRef<SWidget> SPropertyEditorClass::GenerateClassPicker()
 
 void SPropertyEditorClass::OnClassPicked(UClass* InClass)
 {
-	if(!InClass)
+	if (!InClass)
 	{
 		SendToObjects(TEXT("None"));
 	}
@@ -390,14 +372,14 @@ void SPropertyEditorClass::OnClassPicked(UClass* InClass)
 
 void SPropertyEditorClass::SendToObjects(const FString& NewValue)
 {
-	if(PropertyEditor.IsValid())
+	if (PropertyEditor.IsValid())
 	{
 		const TSharedRef<IPropertyHandle> PropertyHandle = PropertyEditor->GetPropertyHandle();
 		PropertyHandle->SetValueFromFormattedString(NewValue);
 	}
 	else if (!NewValue.IsEmpty() && NewValue != TEXT("None"))
 	{
-		const UClass* NewClass = UE::PropertyEditor::Class::Private::FindOrLoadClass(NewValue);
+		const UClass* NewClass = PropertyEditorHelpers::FindOrLoadClass(NewValue);
 		OnSetClass.Execute(NewClass);
 	}
 	else
diff --git a/Engine/Source/Editor/PropertyEditor/Private/UserInterface/PropertyEditor/SPropertyEditorEditInline.cpp b/Engine/Source/Editor/PropertyEditor/Private/UserInterface/PropertyEditor/SPropertyEditorEditInline.cpp
--- a/Engine/Source/Editor/PropertyEditor/Private/UserInterface/PropertyEditor/SPropertyEditorEditInline.cpp
+++ b/Engine/Source/Editor/PropertyEditor/Private/UserInterface/PropertyEditor/SPropertyEditorEditInline.cpp
@@ -13,6 +13,7 @@
 #include "Styling/SlateIconFinder.h"
 #include "UObject/ConstructorHelpers.h"
 #include "Editor.h"
+#include "PropertyEditorHelpers.h"
 #include "PropertyEditorUtils.h"
 
 class FPropertyEditorInlineClassFilter : public IClassViewerFilter
@@ -127,23 +128,6 @@ private:
 	}
 };
 
-namespace UE::PropertyEditor::EditInline::Private
-{
-
-static UClass* FindOrLoadClass(const FString& ClassName)
-{
-	UClass* Class = UClass::TryFindTypeSlow<UClass>(ClassName, EFindFirstObjectOptions::EnsureIfAmbiguous);
-
-	if (!Class)
-	{
-		Class = LoadObject<UClass>(nullptr, *ClassName);
-	}
-
-	return Class;
-}
-
-} // namespace UE::PropertyEditor::EditInline::Private
-
 void SPropertyEditorEditInline::Construct( const FArguments& InArgs, const TSharedRef< class FPropertyEditor >& InPropertyEditor )
 {
 	PropertyEditor = InPropertyEditor;
@@ -302,8 +286,6 @@ TSharedRef<SWidget> SPropertyEditorEditInline::GenerateClassPicker(TSharedRef<IP
 	TArray<const UClass*> DisallowedClassFilters;
 	PropertyEditorUtils::GetAllowedAndDisallowedClasses(ObjectList, *Property, AllowedClassFilters, DisallowedClassFilters, false);
 
-	using namespace UE::PropertyEditor::EditInline::Private;
-
 	TSharedRef<FPropertyHandleBase> PropertyHandleImpl = StaticCastSharedRef<FPropertyHandleBase>(PropertyHandle);
 	TSharedPtr<FPropertyNode> PropertyNode = PropertyHandleImpl->GetPropertyNode();
 
@@ -312,7 +294,7 @@ TSharedRef<SWidget> SPropertyEditorEditInline::GenerateClassPicker(TSharedRef<IP
 	{
 		for (TArray<FString>::TConstIterator Iter = ClassRestriction.Get().GetHiddenValuesIterator(); Iter; ++Iter)
 		{
-			if (const UClass* HiddenClass = FindOrLoadClass(*Iter))
+			if (const UClass* HiddenClass = PropertyEditorHelpers::FindOrLoadClass(*Iter))
 			{
 				DisallowedClassFilters.Add(HiddenClass);
 			}
@@ -320,7 +302,7 @@ TSharedRef<SWidget> SPropertyEditorEditInline::GenerateClassPicker(TSharedRef<IP
 
 		for (TArray<FString>::TConstIterator Iter = ClassRestriction.Get().GetDisabledValuesIterator(); Iter; ++Iter)
 		{
-			if (const UClass* DisabledClass = FindOrLoadClass(*Iter))
+			if (const UClass* DisabledClass = PropertyEditorHelpers::FindOrLoadClass(*Iter))
 			{
 				DisallowedClassFilters.Add(DisabledClass);
 			}
