Added FEnumValuePermission list, which is used to restrict which values can be used when editing enum properties.

diff --git a/Engine/Source/Editor/PropertyEditor/Private/PropertyNode.cpp b/Engine/Source/Editor/PropertyEditor/Private/PropertyNode.cpp
--- a/Engine/Source/Editor/PropertyEditor/Private/PropertyNode.cpp
+++ b/Engine/Source/Editor/PropertyEditor/Private/PropertyNode.cpp
@@ -19,6 +19,7 @@
 #include "PropertyTextUtilities.h"
 #include "PropertyEditorEditConstPolicy.h"
 #include "PropertyEditorArchetypePolicy.h"
+#include "PropertyPermissionList.h"
 #include "StringPrefixTree.h"
 #include "StructurePropertyNode.h"
 
@@ -396,6 +397,37 @@ void FPropertyNode::InitNode(const FPropertyNodeInitParams& InitParams)
 		bRequiresValidation |= HasNodeFlags(EPropertyNodeFlags::SupportsDynamicInstancing);
 		
 		SetNodeFlags( EPropertyNodeFlags::RequiresValidation, bRequiresValidation );
+
+		UEnum* Enum = nullptr;
+		if (const FByteProperty* ByteProperty = CastField<FByteProperty>(MyProperty))
+		{
+			if (ByteProperty->IsEnum())
+			{
+				Enum = ByteProperty->Enum;
+			}
+		}
+		else if (const FEnumProperty* EnumProperty = CastField<FEnumProperty>(MyProperty))
+		{
+			Enum = EnumProperty->GetEnum();
+		}
+
+		if (Enum && FEnumValuePermissionList::Get().HasFiltering(Enum))
+		{
+			TSharedRef<FPropertyRestriction> Restriction = MakeShared<FPropertyRestriction>(FText::FromString(TEXT("Not available")));
+
+			for (int32 Index = 0; Index < Enum->NumEnums(); ++Index)
+			{
+				const FString NameString = Enum->GetNameStringByIndex(Index);
+
+				if (!FEnumValuePermissionList::Get().DoesEnumValuePassFilter(Enum, FName(NameString)))
+				{
+					Restriction->AddHiddenValue(NameString);
+					Restriction->AddDisabledValue(NameString);
+				}
+			}
+
+			AddRestriction(MoveTemp(Restriction));
+		}
 	}
 
 	InitExpansionFlags();
diff --git a/Engine/Source/Editor/PropertyEditor/Private/PropertyPermissionList.cpp b/Engine/Source/Editor/PropertyEditor/Private/PropertyPermissionList.cpp
--- a/Engine/Source/Editor/PropertyEditor/Private/PropertyPermissionList.cpp
+++ b/Engine/Source/Editor/PropertyEditor/Private/PropertyPermissionList.cpp
@@ -463,4 +463,53 @@ FHiddenPropertyPermissionList& FHiddenPropertyPermissionList::Get()
 {
 	static FHiddenPropertyPermissionList PermissionList;
 	return PermissionList;
-}
\ No newline at end of file
+}
+
+
+void FEnumValuePermissionList::AddPermissionList(TSoftObjectPtr<const UEnum> Enum, const FNamePermissionList& InPermissionList)
+{
+	FNamePermissionList& Entry = PermissionList.FindOrAdd(Enum);
+	Entry = InPermissionList;
+}
+
+void FEnumValuePermissionList::RemovePermissionList(TSoftObjectPtr<const UEnum> Enum)
+{
+	PermissionList.Remove(Enum);
+}
+
+void FEnumValuePermissionList::ClearPermissionList()
+{
+	PermissionList.Reset();
+}
+
+bool FEnumValuePermissionList::DoesEnumValuePassFilter(const UEnum* Enum, FName ValueName) const
+{
+	if (!bEnablePermissionList)
+	{
+		return true;
+	}
+
+	const FNamePermissionList* Entry = PermissionList.Find(Enum);
+	if (Entry)
+	{
+		return Entry->PassesFilter(ValueName);
+	}
+	return true;
+}
+
+bool FEnumValuePermissionList::HasFiltering(const UEnum* Enum) const
+{
+	if (!bEnablePermissionList)
+	{
+		return false;
+	}
+
+	const FNamePermissionList* Entry = PermissionList.Find(Enum);
+	return Entry != nullptr;
+}
+
+FEnumValuePermissionList& FEnumValuePermissionList::Get()
+{
+	static FEnumValuePermissionList Instance;
+	return Instance;
+}
diff --git a/Engine/Source/Editor/PropertyEditor/Public/PropertyPermissionList.h b/Engine/Source/Editor/PropertyEditor/Public/PropertyPermissionList.h
--- a/Engine/Source/Editor/PropertyEditor/Public/PropertyPermissionList.h
+++ b/Engine/Source/Editor/PropertyEditor/Public/PropertyPermissionList.h
@@ -187,4 +187,36 @@ private:
 	~FHiddenPropertyPermissionList() = default;
 };
 
+class FEnumValuePermissionList
+{
+public:
+	/** Add a set of rules for a specific UEnum to determine which values are visible in all details panels */
+	UE_API void AddPermissionList(TSoftObjectPtr<const UEnum> Enum, const FNamePermissionList& PermissionList);
+	/** Remove a set of rules for a specific UEnum to determine which values are visible in all details panels */
+	UE_API void RemovePermissionList(TSoftObjectPtr<const UEnum> Enum);
+	/** Remove all rules */
+	UE_API void ClearPermissionList();
+
+	/** Check whether the specified value is permitted for the provided UEnum */
+	UE_API bool DoesEnumValuePassFilter(const UEnum* Enum, FName ValueName) const;
+
+	/** Checks if the provided enum has any entries in the permission list and returns true in that case. */
+	UE_API bool HasFiltering(const UEnum* Enum) const;
+
+	/** When false, DoesEnumValuePassFilter will always return true. */
+	bool IsEnabled() const { return bEnablePermissionList; }
+	/** Turn on or off the enum PermissionList. When off, DoesEnumValuePassFilter will always return true. */
+	void SetEnabled(bool bEnable) { bEnablePermissionList = bEnable; }
+
+	/** Gets a read-only copy of the original, un-flattened PermissionList. */
+	const TMap<TSoftObjectPtr<const UEnum>, FNamePermissionList>& GetRawPermissionList() const { return PermissionList; }
+
+	static UE_API FEnumValuePermissionList& Get();
+
+private:
+	TMap<TSoftObjectPtr<const UEnum>, FNamePermissionList> PermissionList;
+
+	bool bEnablePermissionList = false;
+};
+
 #undef UE_API
