Bypass crash when toggling edit condition property inside an instanced struct details.

diff --git a/Engine/Source/Editor/PropertyEditor/Private/PropertyNode.cpp b/Engine/Source/Editor/PropertyEditor/Private/PropertyNode.cpp
--- a/Engine/Source/Editor/PropertyEditor/Private/PropertyNode.cpp
+++ b/Engine/Source/Editor/PropertyEditor/Private/PropertyNode.cpp
@@ -1647,26 +1647,31 @@ void FPropertyNode::ToggleEditConditionState()
 				
 				for (int32 InstanceIndex = 0; InstanceIndex < ArchetypeInstances.Num(); ++InstanceIndex)
 				{
-					uint8* ArchetypeBaseOffset;
+					uint8* ArchetypeBaseOffset = nullptr;
 					if (!bIsSparse)
 					{
 						ArchetypeBaseOffset = MyParentNode->GetValueAddressFromObject(ArchetypeInstances[InstanceIndex]);
 					}
 					else
 					{
-						uint8* ConditionedPropertyAddress = GetValueAddressFromObject(ArchetypeInstances[InstanceIndex]);
-						const int32 ConditionedPropertyOffset = GetProperty()->GetOffset_ForInternal();
-						uint8* OwningStructStartAddress = ConditionedPropertyAddress - ConditionedPropertyOffset;
-						ArchetypeBaseOffset = OwningStructStartAddress;
+						if(uint8* ConditionedPropertyAddress = GetValueAddressFromObject(ArchetypeInstances[InstanceIndex]))
+						{
+							const int32 ConditionedPropertyOffset = GetProperty()->GetOffset_ForInternal();
+							uint8* OwningStructStartAddress = ConditionedPropertyAddress - ConditionedPropertyOffset;
+							ArchetypeBaseOffset = OwningStructStartAddress;
+						}
 					}
 					
-					uint8* ArchetypeValueAddr = EditConditionProperty->ContainerPtrToValuePtr<uint8>(ArchetypeBaseOffset);
-
-					// Only propagate if the current value on the instance matches the previous value on the template.
-					const bool CurValue = EditConditionProperty->GetPropertyValue(ArchetypeValueAddr);
-					if (OldValue == CurValue)
+					if(ArchetypeBaseOffset != nullptr)
 					{
-						EditConditionProperty->SetPropertyValue(ArchetypeValueAddr, !OldValue);
+						uint8* ArchetypeValueAddr = EditConditionProperty->ContainerPtrToValuePtr<uint8>(ArchetypeBaseOffset);
+
+						// Only propagate if the current value on the instance matches the previous value on the template.
+						const bool CurValue = EditConditionProperty->GetPropertyValue(ArchetypeValueAddr);
+						if (OldValue == CurValue)
+						{
+							EditConditionProperty->SetPropertyValue(ArchetypeValueAddr, !OldValue);
+						}
 					}
 				}
 			}
