DetailsView: Export SetRootExpansionStates through the public interface.

diff --git a/Engine/Source/Editor/PropertyEditor/Private/SDetailsViewBase.h b/Engine/Source/Editor/PropertyEditor/Private/SDetailsViewBase.h
--- a/Engine/Source/Editor/PropertyEditor/Private/SDetailsViewBase.h
+++ b/Engine/Source/Editor/PropertyEditor/Private/SDetailsViewBase.h
@@ -178,7 +178,7 @@ public:
 	 * @param bExpand			The new expansion state
 	 * @param bRecurse			Whether or not to apply the expansion change to any children
 	 */
-	void SetRootExpansionStates(const bool bExpand, const bool bRecurse);
+	virtual void SetRootExpansionStates(const bool bExpand, const bool bRecurse) override;
 
 	/**
 	 * Adds an action to execute next tick
diff --git a/Engine/Source/Editor/PropertyEditor/Public/IDetailsView.h b/Engine/Source/Editor/PropertyEditor/Public/IDetailsView.h
--- a/Engine/Source/Editor/PropertyEditor/Public/IDetailsView.h
+++ b/Engine/Source/Editor/PropertyEditor/Public/IDetailsView.h
@@ -236,6 +236,14 @@ public:
 	 * Forces all advanced property sections to be in expanded state:
 	 */
 	virtual void ShowAllAdvancedProperties() = 0;
+
+	/**
+	 * Sets the expansion state all root nodes and optionally all of their children
+	 *
+	 * @param bExpand			The new expansion state
+	 * @param bRecurse			Whether or not to apply the expansion change to any children
+	 */
+	virtual void SetRootExpansionStates(const bool bExpand, const bool bRecurse) = 0;
 	
 	/**
 	 * Refreshes the visibility of root objects in this details view. 
