Add a CanItemChange event to array properties so that it's possible to conditionally back out of delete / duplicate / etc actions.

diff --git a/Engine/Source/Editor/PropertyEditor/Private/PropertyHandleImpl.cpp b/Engine/Source/Editor/PropertyEditor/Private/PropertyHandleImpl.cpp
--- a/Engine/Source/Editor/PropertyEditor/Private/PropertyHandleImpl.cpp
+++ b/Engine/Source/Editor/PropertyEditor/Private/PropertyHandleImpl.cpp
@@ -798,6 +798,54 @@ void FPropertyValueImpl::SetOnPropertyResetToDefault(const FSimpleDelegate& InOn
 	}
 }
 
+void FPropertyValueImpl::SetOnCanItemChange(const FCanPropertyItemChange& InDelegate)
+{
+	if (PropertyNode.IsValid())
+	{
+		PropertyNode.Pin()->OnCanItemChangeEvent() = InDelegate;
+	}
+}
+
+bool FPropertyValueImpl::CanItemChange(EPropertyChangeType::Type Type, int32 Index)
+{
+	if (const TSharedPtr<FPropertyNode> Node = PropertyNode.Pin())
+	{
+		const FCanPropertyItemChange& Event = Node->OnCanItemChangeEvent();
+		if (Event.IsBound())
+		{
+			FCanPropertyItemChangeRetType Result = Event.Execute(Type, Index);
+			if (Result.HasError())
+			{
+				const FPropertyChangeBlocked& ErrorEvent = Node->OnChangeBlockedEvent();
+				ErrorEvent.Broadcast(Node->GetProperty(), Result.GetError());
+				return false;
+			}
+
+			return Result.GetValue();
+		}
+	}
+
+	return true;
+}
+
+FDelegateHandle FPropertyValueImpl::SetOnChangeBlocked(const FPropertyChangeBlocked::FDelegate& InDelegate)
+{
+	if (PropertyNode.IsValid())
+	{
+		return PropertyNode.Pin()->OnChangeBlockedEvent().Add(InDelegate);
+	}
+
+	return { };
+}
+
+void FPropertyValueImpl::UnregisterOnChangeBlocked(FDelegateHandle Handle)
+{
+	if (PropertyNode.IsValid())
+	{
+		PropertyNode.Pin()->OnChangeBlockedEvent().Remove(Handle);
+	}
+}
+
 FDelegateHandle FPropertyValueImpl::SetOnRebuildChildren( const FSimpleDelegate& InOnRebuildChildren )
 {
 	if( PropertyNode.IsValid() )
@@ -5359,9 +5407,12 @@ FPropertyHandleItemAddResult FPropertyHandleArray::AddItem()
 	int32 Index = INDEX_NONE;
 	if( IsEditable() )
 	{
-		Index = Implementation->AddChild();
-		Result = FPropertyAccess::Success;
-		return FPropertyHandleItemAddResult(Index, Result);
+		if (Implementation->CanItemChange(EPropertyChangeType::ArrayAdd, Index))
+		{
+			Index = Implementation->AddChild();
+			Result = FPropertyAccess::Success;
+			return FPropertyHandleItemAddResult(Index, Result);
+		}
 	}
 	return FPropertyHandleItemAddResult();
 }
@@ -5371,8 +5422,11 @@ FPropertyAccess::Result FPropertyHandleArray::EmptyArray()
 	FPropertyAccess::Result Result = FPropertyAccess::Fail;
 	if( IsEditable() )
 	{
-		Implementation->ClearChildren();
-		Result = FPropertyAccess::Success;
+		if (Implementation->CanItemChange(EPropertyChangeType::ArrayClear, INDEX_NONE))
+		{
+			Implementation->ClearChildren();
+			Result = FPropertyAccess::Success;
+		}
 	}
 
 	return Result;
@@ -5383,8 +5437,11 @@ FPropertyAccess::Result FPropertyHandleArray::Insert( int32 Index )
 	FPropertyAccess::Result Result = FPropertyAccess::Fail;
 	if( IsEditable() && Index < Implementation->GetNumChildren()  )
 	{
-		Implementation->InsertChild( Index );
-		Result = FPropertyAccess::Success;
+		if (Implementation->CanItemChange(EPropertyChangeType::ArrayAdd, Index))
+		{
+			Implementation->InsertChild( Index );
+			Result = FPropertyAccess::Success;
+		}
 	}
 
 	return Result;
@@ -5395,8 +5452,11 @@ FPropertyAccess::Result FPropertyHandleArray::DuplicateItem( int32 Index )
 	FPropertyAccess::Result Result = FPropertyAccess::Fail;
 	if( IsEditable() && Index < Implementation->GetNumChildren() )
 	{
-		Implementation->DuplicateChild( Index );
-		Result = FPropertyAccess::Success;
+		if (Implementation->CanItemChange(EPropertyChangeType::Duplicate, Index))
+		{
+			Implementation->DuplicateChild( Index );
+			Result = FPropertyAccess::Success;
+		}
 	}
 
 	return Result;
@@ -5407,8 +5467,11 @@ FPropertyAccess::Result FPropertyHandleArray::DeleteItem( int32 Index )
 	FPropertyAccess::Result Result = FPropertyAccess::Fail;
 	if( IsEditable() && Index < Implementation->GetNumChildren() )
 	{
-		Implementation->DeleteChild( Index );
-		Result = FPropertyAccess::Success;
+		if (Implementation->CanItemChange(EPropertyChangeType::ArrayRemove, Index))
+		{
+			Implementation->DeleteChild( Index );
+			Result = FPropertyAccess::Success;
+		}
 	}
 
 	return Result;
@@ -5419,8 +5482,12 @@ FPropertyAccess::Result FPropertyHandleArray::SwapItems(int32 FirstIndex, int32
 	FPropertyAccess::Result Result = FPropertyAccess::Fail;
 	if (IsEditable() && FirstIndex >= 0 && SecondIndex >= 0 && FirstIndex < Implementation->GetNumChildren() && SecondIndex < Implementation->GetNumChildren())
 	{
-		Implementation->SwapChildren(FirstIndex, SecondIndex);
-		Result = FPropertyAccess::Success;
+		if (Implementation->CanItemChange(EPropertyChangeType::ArrayMove, FirstIndex) &&
+			Implementation->CanItemChange(EPropertyChangeType::ArrayMove, SecondIndex))
+		{
+			Implementation->SwapChildren(FirstIndex, SecondIndex);
+			Result = FPropertyAccess::Success;
+		}
 	}
 
 	return Result;
@@ -5442,6 +5509,21 @@ void FPropertyHandleArray::UnregisterOnNumElementsChanged(FDelegateHandle Handle
 	Implementation->UnregisterOnRebuildChildren(Handle);
 }
 
+void FPropertyHandleArray::SetOnCanItemChange(const FCanPropertyItemChange& InDelegate)
+{
+	Implementation->SetOnCanItemChange(InDelegate);
+}
+
+FDelegateHandle FPropertyHandleArray::SetOnChangeBlocked(const FPropertyChangeBlocked::FDelegate& InDelegate)
+{
+	return Implementation->SetOnChangeBlocked(InDelegate);
+}
+
+void FPropertyHandleArray::UnregisterOnChangeBlocked(FDelegateHandle Handle)
+{
+	return Implementation->UnregisterOnChangeBlocked(Handle);
+}
+
 TSharedPtr<IPropertyHandleArray> FPropertyHandleArray::AsArray()
 {
 	return SharedThis(this);
diff --git a/Engine/Source/Editor/PropertyEditor/Private/PropertyHandleImpl.h b/Engine/Source/Editor/PropertyEditor/Private/PropertyHandleImpl.h
--- a/Engine/Source/Editor/PropertyEditor/Private/PropertyHandleImpl.h
+++ b/Engine/Source/Editor/PropertyEditor/Private/PropertyHandleImpl.h
@@ -171,6 +171,22 @@ public:
 	 */
 	void SetOnPropertyResetToDefault(const FSimpleDelegate& InOnPropertyResetToDefault);
 
+	/**
+	 * Sets a delegate to call before an item in a container is being mutated
+	 */
+	void SetOnCanItemChange(const FCanPropertyItemChange& InDelegate);
+	bool CanItemChange(EPropertyChangeType::Type Type, int32 Index);
+
+	/**
+	 * Sets a delegate to call when a property change has been disallowed
+	 */
+	FDelegateHandle SetOnChangeBlocked(const FPropertyChangeBlocked::FDelegate& InDelegate);
+
+	/**
+	 * Unregisters a delegate that's bound to OnChangeBlocked
+	 */
+	void UnregisterOnChangeBlocked(FDelegateHandle Handle);
+
 	/**
 	 * Sets a delegate to call when children of the property node must be rebuilt
 	 */
@@ -731,6 +747,9 @@ public:
 	virtual FPropertyAccess::Result GetNumElements( uint32& OutNumItems ) const override;
 	virtual FDelegateHandle SetOnNumElementsChanged( const FSimpleDelegate& InOnNumElementsChanged ) override;
 	virtual void UnregisterOnNumElementsChanged(FDelegateHandle Handle) override;
+	virtual void SetOnCanItemChange(const FCanPropertyItemChange& InDelegate) override;
+	virtual FDelegateHandle SetOnChangeBlocked(const FPropertyChangeBlocked::FDelegate& InDelegate) override;
+	virtual void UnregisterOnChangeBlocked(FDelegateHandle Handle) override;
 	virtual TSharedPtr<IPropertyHandleArray> AsArray() override;
 	virtual TSharedRef<IPropertyHandle> GetElement( int32 Index ) const override;
 	virtual FPropertyAccess::Result MoveElementTo(int32 OriginalIndex, int32 NewIndex) override;
diff --git a/Engine/Source/Editor/PropertyEditor/Private/PropertyNode.h b/Engine/Source/Editor/PropertyEditor/Private/PropertyNode.h
--- a/Engine/Source/Editor/PropertyEditor/Private/PropertyNode.h
+++ b/Engine/Source/Editor/PropertyEditor/Private/PropertyNode.h
@@ -790,6 +790,15 @@ public:
 	DECLARE_EVENT(FPropertyNode, FPropertyResetToDefaultEvent);
 	FPropertyResetToDefaultEvent& OnPropertyResetToDefault() { return PropertyResetToDefaultEvent; }
 
+	/**
+	 * Broadcasts before an item in a container is being mutated, giving it a chance to back out.
+	 * If an error occurs, it's sent by ChangeBlockedEvent.
+	 */
+	FCanPropertyItemChange& OnCanItemChangeEvent() { return CanItemChangeEvent; }
+
+	/** Broadcasts before an item in a container is being mutated, giving it a chance to back out */
+	FPropertyChangeBlocked& OnChangeBlockedEvent() { return ChangeBlockedEvent; }
+
 	/**
 	 * Marks window's seem due to filtering flags
 	 * @param InFilterStrings	- List of strings that must be in the property name in order to display
@@ -1265,6 +1274,12 @@ protected:
 	/** Called when the property is reset to default */
 	FPropertyResetToDefaultEvent PropertyResetToDefaultEvent;
 
+	/** Called before a container element is changed */
+	FCanPropertyItemChange CanItemChangeEvent;
+
+	/** Called when an attempted property change has been disallowed */
+	FPropertyChangeBlocked ChangeBlockedEvent;
+
 	/** The property being displayed/edited. */
 	TWeakFieldPtr<FProperty> Property;
 
diff --git a/Engine/Source/Editor/PropertyEditor/Public/PropertyEditorDelegates.h b/Engine/Source/Editor/PropertyEditor/Public/PropertyEditorDelegates.h
--- a/Engine/Source/Editor/PropertyEditor/Public/PropertyEditorDelegates.h
+++ b/Engine/Source/Editor/PropertyEditor/Public/PropertyEditorDelegates.h
@@ -11,6 +11,7 @@
 #include "Internationalization/Text.h"
 #include "Misc/Attribute.h"
 #include "Templates/SharedPointer.h"
+#include "Templates/ValueOrError.h"
 #include "Textures/SlateIcon.h"
 #include "UObject/WeakObjectPtr.h"
 #include "UObject/WeakObjectPtrTemplates.h"
@@ -98,6 +99,13 @@ DECLARE_DELEGATE_RetVal_TwoParams(TSharedRef< SWidget >, FConstructExternalColum
 /** Delegate called to see if a property editing is enabled */
 DECLARE_DELEGATE_RetVal(bool, FIsPropertyEditingEnabled);
 
+/** Delegate to be called before a container item changes */
+using FCanPropertyItemChangeRetType = TValueOrError<bool, FText>;
+DECLARE_DELEGATE_RetVal_TwoParams(FCanPropertyItemChangeRetType, FCanPropertyItemChange, EPropertyChangeType::Type, int32 /* Index */);
+
+/** Delegate to be called when a change is blocked on a property */
+DECLARE_MULTICAST_DELEGATE_TwoParams(FPropertyChangeBlocked, const FProperty*, FText /* ErrorMessage */);
+
 /**
  * A delegate that returns the name of a type that is to be used to look up the layout callback instead of the
  * property's type. Return None to not provide an override.
diff --git a/Engine/Source/Editor/PropertyEditor/Public/PropertyHandle.h b/Engine/Source/Editor/PropertyEditor/Public/PropertyHandle.h
--- a/Engine/Source/Editor/PropertyEditor/Public/PropertyHandle.h
+++ b/Engine/Source/Editor/PropertyEditor/Public/PropertyHandle.h
@@ -863,6 +863,20 @@ public:
 	*/
 	virtual FPropertyAccess::Result MoveElementTo(int32 OriginalIndex, int32 NewIndex) = 0;
 
+	/**
+	 * Sets a delegate to call before an item in a container is being mutated
+	 */
+	virtual void SetOnCanItemChange(const FCanPropertyItemChange& InDelegate) = 0;
+
+	/**
+	 * Sets a delegate to call when a property change has been disallowed
+	 */
+	virtual FDelegateHandle SetOnChangeBlocked(const FPropertyChangeBlocked::FDelegate& InDelegate) = 0;
+
+	/**
+	 * Unregisters a delegate that is bound to OnChangeBlocked
+	 */
+	virtual void UnregisterOnChangeBlocked(FDelegateHandle Handle) = 0;
 
 	/**
 	 * Sets a delegate to call when the number of elements changes
