[PropertyEditor] Change some utilities to hold widgets via weak pointer

#tests Using the editor, running PIE and game client, making builds

diff --git a/Engine/Source/Editor/PropertyEditor/Private/DetailsViewPropertyGenerationUtilities.cpp b/Engine/Source/Editor/PropertyEditor/Private/DetailsViewPropertyGenerationUtilities.cpp
new file mode 100644
--- /dev/null
+++ b/Engine/Source/Editor/PropertyEditor/Private/DetailsViewPropertyGenerationUtilities.cpp
@@ -0,0 +1,29 @@
+// Copyright Epic Games, Inc. All Rights Reserved.
+
+#pragma once
+
+#include "DetailsViewPropertyGenerationUtilities.h"
+
+FDetailsViewPropertyGenerationUtilities::FDetailsViewPropertyGenerationUtilities(SDetailsViewBase& InDetailsView)
+	: DetailsViewPtr(StaticCastWeakPtr<SDetailsViewBase>(InDetailsView.AsWeak()))
+{
+}
+
+const FCustomPropertyTypeLayoutMap& FDetailsViewPropertyGenerationUtilities::GetInstancedPropertyTypeLayoutMap() const
+{
+	if (TSharedPtr<SDetailsViewBase> DetailsView = DetailsViewPtr.Pin())
+	{
+		return DetailsView->GetCustomPropertyTypeLayoutMap();
+	}
+
+	static FCustomPropertyTypeLayoutMap Empty;
+	return Empty;
+}
+
+void FDetailsViewPropertyGenerationUtilities::RebuildTreeNodes()
+{
+	if (TSharedPtr<SDetailsViewBase> DetailsView = DetailsViewPtr.Pin())
+	{
+		DetailsView->RerunCurrentFilter();
+	}
+}
\ No newline at end of file
diff --git a/Engine/Source/Editor/PropertyEditor/Private/DetailsViewPropertyGenerationUtilities.h b/Engine/Source/Editor/PropertyEditor/Private/DetailsViewPropertyGenerationUtilities.h
--- a/Engine/Source/Editor/PropertyEditor/Private/DetailsViewPropertyGenerationUtilities.h
+++ b/Engine/Source/Editor/PropertyEditor/Private/DetailsViewPropertyGenerationUtilities.h
@@ -9,21 +9,11 @@
 class FDetailsViewPropertyGenerationUtilities : public IPropertyGenerationUtilities
 {
 public:
-	FDetailsViewPropertyGenerationUtilities(SDetailsViewBase& InDetailsView)
-		: DetailsView(&InDetailsView)
-	{
-	}
-
-	virtual const FCustomPropertyTypeLayoutMap& GetInstancedPropertyTypeLayoutMap() const override
-	{
-		return DetailsView->GetCustomPropertyTypeLayoutMap();
-	}
-
-	virtual void RebuildTreeNodes() override
-	{
-		DetailsView->RerunCurrentFilter();
-	}
+	FDetailsViewPropertyGenerationUtilities(SDetailsViewBase& InDetailsView);
+	/** IPropertyGenerationUtilities interface */
+	virtual const FCustomPropertyTypeLayoutMap& GetInstancedPropertyTypeLayoutMap() const override;
+	virtual void RebuildTreeNodes() override;
 
 private:
-	SDetailsViewBase * DetailsView;
+	TWeakPtr<SDetailsViewBase> DetailsViewPtr;
 };
\ No newline at end of file
diff --git a/Engine/Source/Editor/PropertyEditor/Private/UserInterface/PropertyDetails/PropertyDetailsUtilities.cpp b/Engine/Source/Editor/PropertyEditor/Private/UserInterface/PropertyDetails/PropertyDetailsUtilities.cpp
--- a/Engine/Source/Editor/PropertyEditor/Private/UserInterface/PropertyDetails/PropertyDetailsUtilities.cpp
+++ b/Engine/Source/Editor/PropertyEditor/Private/UserInterface/PropertyDetails/PropertyDetailsUtilities.cpp
@@ -5,13 +5,17 @@
 #include "IDetailsViewPrivate.h"
 
 FPropertyDetailsUtilities::FPropertyDetailsUtilities(IDetailsViewPrivate& InDetailsView)
-	: DetailsView( InDetailsView )
+	: DetailsViewPtr(StaticCastWeakPtr<IDetailsViewPrivate>(InDetailsView.AsWeak()))
 {
 }
 
 class FNotifyHook* FPropertyDetailsUtilities::GetNotifyHook() const
 {
-	return DetailsView.GetNotifyHook();
+	if (TSharedPtr<IDetailsViewPrivate> DetailsView = DetailsViewPtr.Pin())
+	{
+		return DetailsView->GetNotifyHook();
+	}
+	return nullptr;
 }
 
 bool FPropertyDetailsUtilities::AreFavoritesEnabled() const
@@ -27,60 +31,106 @@ void FPropertyDetailsUtilities::ToggleFavorite( const TSharedRef< FPropertyEdito
 
 void FPropertyDetailsUtilities::CreateColorPickerWindow( const TSharedRef< FPropertyEditor >& PropertyEditor, bool bUseAlpha ) const
 {
-	DetailsView.CreateColorPickerWindow( PropertyEditor, bUseAlpha );
+	if (TSharedPtr<IDetailsViewPrivate> DetailsView = DetailsViewPtr.Pin())
+	{
+		DetailsView->CreateColorPickerWindow(PropertyEditor, bUseAlpha);
+	}
 }
 
 void FPropertyDetailsUtilities::EnqueueDeferredAction( FSimpleDelegate DeferredAction )
 {
-	DetailsView.EnqueueDeferredAction( DeferredAction );
+	if (TSharedPtr<IDetailsViewPrivate> DetailsView = DetailsViewPtr.Pin())
+	{
+		DetailsView->EnqueueDeferredAction(DeferredAction);
+	}
 }
 
 bool FPropertyDetailsUtilities::IsPropertyEditingEnabled() const
 {
-	return DetailsView.IsPropertyEditingEnabled();
+	if (TSharedPtr<IDetailsViewPrivate> DetailsView = DetailsViewPtr.Pin())
+	{
+		return DetailsView->IsPropertyEditingEnabled();
+	}
+	return false;
 }
 
 void FPropertyDetailsUtilities::ForceRefresh()
 {
-	DetailsView.ForceRefresh();
+	if (TSharedPtr<IDetailsViewPrivate> DetailsView = DetailsViewPtr.Pin())
+	{
+		DetailsView->ForceRefresh();
+	}
 }
 
 void FPropertyDetailsUtilities::RequestRefresh()
 {
-	DetailsView.RefreshTree();
+	if (TSharedPtr<IDetailsViewPrivate> DetailsView = DetailsViewPtr.Pin())
+	{
+		DetailsView->RefreshTree();
+	}
 }
 
 void FPropertyDetailsUtilities::RequestForceRefresh()
 {
-	DetailsView.RequestForceRefresh();
+	if (TSharedPtr<IDetailsViewPrivate> DetailsView = DetailsViewPtr.Pin())
+	{
+		DetailsView->RequestForceRefresh();
+	}
 }
 
 TSharedPtr<class FAssetThumbnailPool> FPropertyDetailsUtilities::GetThumbnailPool() const
 {
-	return DetailsView.GetThumbnailPool();
+	if (TSharedPtr<IDetailsViewPrivate> DetailsView = DetailsViewPtr.Pin())
+	{
+		return DetailsView->GetThumbnailPool();
+	}
+	return {};
 }
 
 const TArray<TSharedRef<class IClassViewerFilter>>& FPropertyDetailsUtilities::GetClassViewerFilters() const
 {
-	return DetailsView.GetClassViewerFilters();
+	if (TSharedPtr<IDetailsViewPrivate> DetailsView = DetailsViewPtr.Pin())
+	{
+		return DetailsView->GetClassViewerFilters();
+	}
+
+	static TArray<TSharedRef<class IClassViewerFilter>> Empty;
+	return Empty;
 }
 
 void FPropertyDetailsUtilities::NotifyFinishedChangingProperties(const FPropertyChangedEvent& PropertyChangedEvent)
 {
-	DetailsView.NotifyFinishedChangingProperties(PropertyChangedEvent);
+	if (TSharedPtr<IDetailsViewPrivate> DetailsView = DetailsViewPtr.Pin())
+	{
+		DetailsView->NotifyFinishedChangingProperties(PropertyChangedEvent);
+	}
 }
 
 bool FPropertyDetailsUtilities::DontUpdateValueWhileEditing() const
 {
-	return DetailsView.DontUpdateValueWhileEditing();
+	if (TSharedPtr<IDetailsViewPrivate> DetailsView = DetailsViewPtr.Pin())
+	{
+		return DetailsView->DontUpdateValueWhileEditing();
+	}
+	return true;
 }
 
 const TArray<TWeakObjectPtr<UObject>>& FPropertyDetailsUtilities::GetSelectedObjects() const
 {
-	return DetailsView.GetSelectedObjects();
+	if (TSharedPtr<IDetailsViewPrivate> DetailsView = DetailsViewPtr.Pin())
+	{
+		return DetailsView->GetSelectedObjects();
+	}
+	
+	static TArray<TWeakObjectPtr<UObject>> Empty;
+	return Empty;
 }
 
 bool FPropertyDetailsUtilities::HasClassDefaultObject() const
 {
-	return DetailsView.HasClassDefaultObject();
+	if (TSharedPtr<IDetailsViewPrivate> DetailsView = DetailsViewPtr.Pin())
+	{
+		return DetailsView->HasClassDefaultObject();
+	}
+	return false;
 }
diff --git a/Engine/Source/Editor/PropertyEditor/Private/UserInterface/PropertyDetails/PropertyDetailsUtilities.h b/Engine/Source/Editor/PropertyEditor/Private/UserInterface/PropertyDetails/PropertyDetailsUtilities.h
--- a/Engine/Source/Editor/PropertyEditor/Private/UserInterface/PropertyDetails/PropertyDetailsUtilities.h
+++ b/Engine/Source/Editor/PropertyEditor/Private/UserInterface/PropertyDetails/PropertyDetailsUtilities.h
@@ -31,5 +31,5 @@ public:
 	virtual const TArray<TSharedRef<class IClassViewerFilter>>& GetClassViewerFilters() const override;
 
 private:
-	IDetailsViewPrivate& DetailsView;
+	TWeakPtr<IDetailsViewPrivate> DetailsViewPtr;
 };
