Fix cis issue

diff --git a/Engine/Source/Editor/PropertyEditor/Private/DetailsViewPropertyGenerationUtilities.cpp b/Engine/Source/Editor/PropertyEditor/Private/DetailsViewPropertyGenerationUtilities.cpp
--- a/Engine/Source/Editor/PropertyEditor/Private/DetailsViewPropertyGenerationUtilities.cpp
+++ b/Engine/Source/Editor/PropertyEditor/Private/DetailsViewPropertyGenerationUtilities.cpp
@@ -1,7 +1,5 @@
 // Copyright Epic Games, Inc. All Rights Reserved.
 
-#pragma once
-
 #include "DetailsViewPropertyGenerationUtilities.h"
 
 FDetailsViewPropertyGenerationUtilities::FDetailsViewPropertyGenerationUtilities(SDetailsViewBase& InDetailsView)
