Updated the internals of FPropertyPermissionList and added unit tests

Updated the internal implementation of FPropertyPermissionList::GetCachedPermissionListForStruct to remove the confusing and hard to follow bInOutShouldAllowListAllProperties variable, in favor of using the existing EPropertyPermissionListRules Role value correctly. This results in slightly shorter code, which is much more explicit about its rules.

The new implementation has been verified to be bug-for-bug compatible with the old implementation, and some unit tests have been added to verify the various edge-cases that were found when testing the combined permission cache.

#rnx

diff --git a/Engine/Source/Editor/PropertyEditor/Private/PropertyPermissionList.cpp b/Engine/Source/Editor/PropertyEditor/Private/PropertyPermissionList.cpp
--- a/Engine/Source/Editor/PropertyEditor/Private/PropertyPermissionList.cpp
+++ b/Engine/Source/Editor/PropertyEditor/Private/PropertyPermissionList.cpp
@@ -283,19 +283,10 @@ bool FPropertyPermissionList::DoesDetailTreeNodePassFilter(const UStruct* Object
 
 const FNamePermissionList& FPropertyPermissionList::GetCachedPermissionListForStruct(const UStruct* Struct) const
 {
-	const FPropertyPermissionListEntry* CachedPermissionList = CachedPropertyPermissionList.Find(Struct);
-	if (CachedPermissionList)
-	{
-		return CachedPermissionList->PermissionList;
-	}
-
-	// Default value doesn't matter since it's a no-op until the first PermissionList is encountered, at which
-	// point the rules will re-assign the value.
-	bool bShouldPermissionListAllProperties = false;
-	return GetCachedPermissionListForStructHelper(Struct, bShouldPermissionListAllProperties);
+	return GetCachedPermissionListEntryForStructHelper(Struct).PermissionList;
 }
 
-const FNamePermissionList& FPropertyPermissionList::GetCachedPermissionListForStructHelper(const UStruct* Struct, bool& bInOutShouldAllowListAllProperties) const
+const FPropertyPermissionListEntry& FPropertyPermissionList::GetCachedPermissionListEntryForStructHelper(const UStruct* Struct) const
 {
 	check(Struct);
 
@@ -311,121 +302,117 @@ const FNamePermissionList& FPropertyPermissionList::GetCachedPermissionListForSt
 	}
 #endif
 
-	auto IsAllowListAllProps = [](const FPropertyPermissionListEntry* InEntry)
-		{
-			return InEntry ? (InEntry->PermissionList.GetAllowList().Num() == 0 || InEntry->Rules == EPropertyPermissionListRules::AllowListAllProperties) : true;
-		};
-
-	const FPropertyPermissionListEntry* Entry = RawPropertyPermissionList.Find(Struct);
-	// If an entry is set to AllowListAll, then treat it as if it has no manually-defined properties.
-	const bool bIsThisAllowListEmpty = IsAllowListAllProps(Entry);
-
-	// Normally this case would be caught in GetCachedPermissionListForStruct, but when being called recursively from a subclass
-	// we still need to update bInOutShouldAllowListAllProperties so that new PermissionLists cache properly.
-	FPropertyPermissionListEntry* PermissionListEntry = CachedPropertyPermissionList.Find(Struct);
-	if (PermissionListEntry)
+	// Is this struct already cached? If so, we can just return its result
+	if (const FPropertyPermissionListEntry* CachedPermissionList = CachedPropertyPermissionList.Find(Struct))
 	{
-		// Same check as below, but it specifically has to come after the recursive call so we have to duplicate the code here
-		if (Entry && Entry->Rules == EPropertyPermissionListRules::AllowListAllProperties)
-		{
-			bInOutShouldAllowListAllProperties = true;
-		}
-		// This is to properly support propagation of AllowListAllSubclassProperties for classes that do not have a raw entry
-		if (PermissionListEntry->Rules == EPropertyPermissionListRules::AllowListAllProperties)
-		{
-			bInOutShouldAllowListAllProperties = true;
-		}
+		return *CachedPermissionList;
 	}
-	else
-	{
-		FNamePermissionList NewPermissionList;
 
-		UStruct* SuperStruct = Struct->GetSuperStruct();
-		// Recursively fill the cache for all parent structs
-		if (SuperStruct)
-		{
-			NewPermissionList.Append(GetCachedPermissionListForStructHelper(SuperStruct, bInOutShouldAllowListAllProperties));
-		}
+	const FPropertyPermissionListEntry* Entry = RawPropertyPermissionList.Find(Struct);
 
-		// Append this struct's PermissionList on top of the parent's PermissionList
-		if (Entry)
+	FNamePermissionList NewPermissionList;
+
+	auto AllowAllFields =
+		[&NewPermissionList](const UStruct* InStruct, const EFieldIterationFlags InIterationFlags)
 		{
-			if (bIsThisAllowListEmpty)
+			for (TFieldIterator<FProperty> Property(InStruct, InIterationFlags); Property; ++Property)
 			{
-				// If the AllowList is empty, we only want to append the DenyLists
-				FNamePermissionList DuplicatePermissionList = Entry->PermissionList;
-				// Hack to get around the fact that there's no easy way to only clear an AllowList
-				TMap<FName, FPermissionListOwners>& AllowList = const_cast<TMap<FName, FPermissionListOwners>&>(DuplicatePermissionList.GetAllowList());
-				AllowList.Empty();
-				NewPermissionList.Append(DuplicatePermissionList);
+				NewPermissionList.AddAllowListItem(PropertyEditorPermissionList::PropertyPermissionListOwner, Property->GetFName());
 			}
-			else
+
+			if (PropertyEditorPermissionList::bSupportAllowAllFunctions)
 			{
-				// if the parent struct is explicitly allowing all properties and has an empty allow list, add all its prop to new list, before adding the current entry
-				if (bInOutShouldAllowListAllProperties && NewPermissionList.GetAllowList().Num() == 0 && SuperStruct)
+				for (TFieldIterator<UFunction> FunctionIter(InStruct, InIterationFlags); FunctionIter; ++FunctionIter)
 				{
-					for (TFieldIterator<FProperty> Property(SuperStruct, EFieldIteratorFlags::IncludeSuper, EFieldIteratorFlags::ExcludeDeprecated); Property; ++Property)
-					{
-						NewPermissionList.AddAllowListItem(PropertyEditorPermissionList::PropertyPermissionListOwner, Property->GetFName());
-					}
-
-					if (PropertyEditorPermissionList::bSupportAllowAllFunctions)
+					UFunction* TestFunction = *FunctionIter;
+					if ((TestFunction->ParmsSize == 0) && TestFunction->GetBoolMetaData(PropertyEditorPermissionList::NAME_CallInEditor))
 					{
-						for (TFieldIterator<UFunction> FunctionIter(SuperStruct, EFieldIteratorFlags::IncludeSuper, EFieldIteratorFlags::ExcludeDeprecated); FunctionIter; ++FunctionIter)
-						{
-							UFunction* TestFunction = *FunctionIter;
-							if ((TestFunction->ParmsSize == 0) && TestFunction->GetBoolMetaData(PropertyEditorPermissionList::NAME_CallInEditor))
-							{
-								NewPermissionList.AddAllowListItem(PropertyEditorPermissionList::PropertyPermissionListOwner, TestFunction->GetFName());
-							}
-						}
+						NewPermissionList.AddAllowListItem(PropertyEditorPermissionList::PropertyPermissionListOwner, TestFunction->GetFName());
 					}
 				}
-				NewPermissionList.Append(Entry->PermissionList);
 			}
+		};
 
-			bInOutShouldAllowListAllProperties = IsAllowListAllProps(Entry);
+	// Recursively fill the cache for all parent structs
+	UStruct* SuperStruct = Struct->GetSuperStruct();
+	EPropertyPermissionListRules SuperEntryRule = EPropertyPermissionListRules::UseExistingPermissionList;
+	if (SuperStruct)
+	{
+		const FPropertyPermissionListEntry& SuperPermissionList = GetCachedPermissionListEntryForStructHelper(SuperStruct);
+		SuperEntryRule = SuperPermissionList.Rules;
+		NewPermissionList.Append(SuperPermissionList.PermissionList);
+	}
+
+	// Resolve the permission list rule for this entry
+	EPropertyPermissionListRules EntryRule = EPropertyPermissionListRules::UseExistingPermissionList;
+	if (Entry)
+	{
+		// This causes an issue in the case where a struct should have no AllowList properties but wants use AllowListAllSubclassProperties
+		// In this case, simply add a dummy AllowList entry that (likely) won't ever collide with a real property name
+		if (Entry->PermissionList.GetAllowList().Num() == 0 || Entry->Rules == EPropertyPermissionListRules::AllowListAllProperties)
+		{
+			EntryRule = EPropertyPermissionListRules::AllowListAllProperties;
 		}
+		else if (Entry->Rules == EPropertyPermissionListRules::AllowListAllSubclassProperties)
+		{
+			EntryRule = EPropertyPermissionListRules::AllowListAllSubclassProperties;
+		}
+	}
+	else
+	{
+		// If we don't have a raw entry then we propagate the rule from our parent struct to avoid breaking AllowListAllProperties or AllowListAllSubclassProperties chains
+		EntryRule = SuperEntryRule;
+	}
 
-		// PermissionList all properties if the flag is set, the parent Struct has a PermissionList, and this Struct has no PermissionList
-		// If the parent Struct's PermissionList is empty then that already implies all properties are visible
-		if (bInOutShouldAllowListAllProperties && NewPermissionList.GetAllowList().Num() > 0)
+	// If this entry has explicit rules, append them on-top of any permission rules inherited from our parent(s)
+	if (Entry)
+	{
+		if (EntryRule == EPropertyPermissionListRules::AllowListAllProperties)
 		{
-			for (TFieldIterator<FProperty> Property(Struct, EFieldIteratorFlags::ExcludeSuper, EFieldIteratorFlags::ExcludeDeprecated); Property; ++Property)
+			// Allow all fields if requested
+			// If the allow list inherited from our parent(s) is empty then that already implies all fields are visible, so we can just keep the list as empty
+			if (NewPermissionList.GetAllowList().Num() > 0)
 			{
-				NewPermissionList.AddAllowListItem(PropertyEditorPermissionList::PropertyPermissionListOwner, Property->GetFName());
+				AllowAllFields(Struct, EFieldIterationFlags::None);
 			}
 
-			if (PropertyEditorPermissionList::bSupportAllowAllFunctions)
+			// If the AllowList is empty, we only want to append the DenyLists
+			FNamePermissionList DuplicatePermissionList = Entry->PermissionList;
+			// Hack to get around the fact that there's no easy way to only clear an AllowList
+			TMap<FName, FPermissionListOwners>& AllowList = const_cast<TMap<FName, FPermissionListOwners>&>(DuplicatePermissionList.GetAllowList());
+			AllowList.Empty();
+			NewPermissionList.Append(DuplicatePermissionList);
+		}
+		else
+		{
+			check(Entry->PermissionList.GetAllowList().Num() > 0);
+
+			// If the parent struct is explicitly allowing all properties but has an empty allow list, then we need to explicitly allow all of its fields before appending our permissions
+			if (SuperEntryRule == EPropertyPermissionListRules::AllowListAllProperties && NewPermissionList.GetAllowList().Num() == 0)
 			{
-				for (TFieldIterator<UFunction> FunctionIter(Struct, EFieldIteratorFlags::ExcludeSuper, EFieldIteratorFlags::ExcludeDeprecated); FunctionIter; ++FunctionIter)
-				{
-					UFunction* TestFunction = *FunctionIter;
-					if ((TestFunction->ParmsSize == 0) && TestFunction->GetBoolMetaData(PropertyEditorPermissionList::NAME_CallInEditor))
-					{
-						NewPermissionList.AddAllowListItem(PropertyEditorPermissionList::PropertyPermissionListOwner, TestFunction->GetFName());
-					}
-				}
+				check(SuperStruct);
+				AllowAllFields(SuperStruct, EFieldIterationFlags::IncludeSuper);
 			}
+			NewPermissionList.Append(Entry->PermissionList);
 		}
-
-		PermissionListEntry = &CachedPropertyPermissionList.Add(Struct, { MoveTemp(NewPermissionList),
-			// Note: The AdditionalOwnerNames are only used to prevent removal from RawPropertyPermissionList by calls to RemoveOwner()
-			TArray<FName>(),
-			// propagate the allow list all properties setting, this is to make AllowListAllSubclassProperties work properly 
-			// when a sub class was already previously cached from a base class that has AllowListAllSubclassProperties set on it and another class is then building its list off of the cached subclass list
-			bInOutShouldAllowListAllProperties ? EPropertyPermissionListRules::AllowListAllProperties : EPropertyPermissionListRules::UseExistingPermissionList });
 	}
-	
-	// If this Struct has no PermissionList, then the ShouldPermissionListAllProperties rule just forwards its current value on to the next subclass.
-	// This causes an issue in the case where a Struct should have no PermissionListed properties but wants to PermissionList all subclass properties.
-	// In this case, simply add a dummy entry to the Struct's PermissionList that (likely) won't ever collide with a real property name
-	if (!bIsThisAllowListEmpty)
+
+	// Did our super class ask for sub-classes to expose all their properties? If so, respect that request unless we define our own permission rules
+	EPropertyPermissionListRules EntryRuleToPropagate = EntryRule;
+	if ((!Entry || EntryRule == EPropertyPermissionListRules::AllowListAllProperties) && SuperEntryRule == EPropertyPermissionListRules::AllowListAllSubclassProperties)
 	{
-		bInOutShouldAllowListAllProperties = Entry && Entry->Rules == EPropertyPermissionListRules::AllowListAllSubclassProperties;
+		EntryRuleToPropagate = EPropertyPermissionListRules::AllowListAllSubclassProperties;
+		if (NewPermissionList.GetAllowList().Num() > 0)
+		{
+			AllowAllFields(Struct, EFieldIterationFlags::None);
+		}
 	}
 
-	return PermissionListEntry->PermissionList;
+	return CachedPropertyPermissionList.Add(Struct, {
+		MoveTemp(NewPermissionList),
+		TArray<FName>(), // Note: The AdditionalOwnerNames are only used to prevent removal from RawPropertyPermissionList by calls to RemoveOwner()
+		EntryRuleToPropagate
+		});
 }
 
 bool FPropertyPermissionList::HasSpecificList(const UStruct* ObjectStruct) const
diff --git a/Engine/Source/Editor/PropertyEditor/Private/Tests/PropertyPermissionTests.cpp b/Engine/Source/Editor/PropertyEditor/Private/Tests/PropertyPermissionTests.cpp
new file mode 100644
--- /dev/null
+++ b/Engine/Source/Editor/PropertyEditor/Private/Tests/PropertyPermissionTests.cpp
@@ -0,0 +1,237 @@
+// Copyright Epic Games, Inc. All Rights Reserved.
+
+#include "PropertyPermissionTests.h"
+
+#include "Misc/AutomationTest.h"
+#include "PropertyPermissionList.h"
+
+#include UE_INLINE_GENERATED_CPP_BY_NAME(PropertyPermissionTests)
+
+#if WITH_DEV_AUTOMATION_TESTS
+
+class FTestPropertyPermissionList : public FPropertyPermissionList
+{
+};
+
+IMPLEMENT_SIMPLE_AUTOMATION_TEST(FPropertyEditorTests_PropertyPermissions, "PropertyEditor.PropertyPermissions", EAutomationTestFlags::EditorContext | EAutomationTestFlags::EngineFilter)
+
+bool FPropertyEditorTests_PropertyPermissions::RunTest(const FString& Parameters)
+{
+	FTestPropertyPermissionList PermissionList;
+	PermissionList.SetEnabled(true);
+
+	auto AddPermissionList =
+		[&PermissionList](const UStruct* Struct, TConstArrayView<FName> AllowPropertyNames, TConstArrayView<FName> DenyPropertyNames, const EPropertyPermissionListRules Rules)
+		{
+			FNamePermissionList PropertyNameList;
+			for (const FName PropertyName : AllowPropertyNames)
+			{
+				PropertyNameList.AddAllowListItem(FName(), PropertyName);
+			}
+			for (const FName PropertyName : DenyPropertyNames)
+			{
+				PropertyNameList.AddDenyListItem(FName(), PropertyName);
+			}
+			PermissionList.AddPermissionList(Struct, PropertyNameList, Rules, {});
+		};
+
+	auto TestDoesPropertyPassFilter =
+		[this, &PermissionList](const UStruct* Struct, FName PropertyName, const bool bShouldPass, const TCHAR* What)
+		{
+			const bool bDidPass = PermissionList.DoesPropertyPassFilter(Struct, PropertyName);
+			TestEqual(What, bDidPass, bShouldPass);
+		};
+
+	{
+		// Allow COne and CTwo from struct C
+		// Properties from struct A or struct B for an instance of struct C should be denied
+		// Properties from struct A or struct B for an instance of struct A or B should be allowed
+		// Properties from struct C for an instance of struct C should be allowed
+		AddPermissionList(FPropertyEditorPermissionTestStructC::StaticStruct(), { "COne", "CTwo" }, {}, EPropertyPermissionListRules::UseExistingPermissionList);
+		TestDoesPropertyPassFilter(FPropertyEditorPermissionTestStructA::StaticStruct(), "AOne", true, TEXT("[Test1] AOne on struct A"));
+		TestDoesPropertyPassFilter(FPropertyEditorPermissionTestStructA::StaticStruct(), "ATwo", true, TEXT("[Test1] ATwo on struct A"));
+		TestDoesPropertyPassFilter(FPropertyEditorPermissionTestStructB::StaticStruct(), "AOne", true, TEXT("[Test1] AOne on struct B"));
+		TestDoesPropertyPassFilter(FPropertyEditorPermissionTestStructB::StaticStruct(), "ATwo", true, TEXT("[Test1] ATwo on struct B"));
+		TestDoesPropertyPassFilter(FPropertyEditorPermissionTestStructB::StaticStruct(), "BOne", true, TEXT("[Test1] BOne on struct B"));
+		TestDoesPropertyPassFilter(FPropertyEditorPermissionTestStructB::StaticStruct(), "BTwo", true, TEXT("[Test1] BTwo on struct B"));
+		TestDoesPropertyPassFilter(FPropertyEditorPermissionTestStructC::StaticStruct(), "AOne", false, TEXT("[Test1] AOne on struct C"));
+		TestDoesPropertyPassFilter(FPropertyEditorPermissionTestStructC::StaticStruct(), "ATwo", false, TEXT("[Test1] ATwo on struct C"));
+		TestDoesPropertyPassFilter(FPropertyEditorPermissionTestStructC::StaticStruct(), "BOne", false, TEXT("[Test1] BOne on struct C"));
+		TestDoesPropertyPassFilter(FPropertyEditorPermissionTestStructC::StaticStruct(), "BTwo", false, TEXT("[Test1] BTwo on struct C"));
+		TestDoesPropertyPassFilter(FPropertyEditorPermissionTestStructC::StaticStruct(), "COne", true, TEXT("[Test1] COne on struct C"));
+		TestDoesPropertyPassFilter(FPropertyEditorPermissionTestStructC::StaticStruct(), "CTwo", true, TEXT("[Test1] CTwo on struct C"));
+		
+		// Allow AOne property from struct A
+		// AOne for an instance of struct A, B, or C should be allowed
+		// ATwo for an instance of struct A, B, or C should be denied
+		// Properties from struct B for an instance of struct B or C should be denied
+		// Properties from struct C for an instance of struct C should be allowed (due to the previous rule)
+		AddPermissionList(FPropertyEditorPermissionTestStructA::StaticStruct(), { "AOne" }, {}, EPropertyPermissionListRules::UseExistingPermissionList);
+		TestDoesPropertyPassFilter(FPropertyEditorPermissionTestStructA::StaticStruct(), "AOne", true, TEXT("[Test2] AOne on struct A"));
+		TestDoesPropertyPassFilter(FPropertyEditorPermissionTestStructA::StaticStruct(), "ATwo", false, TEXT("[Test2] ATwo on struct A"));
+		TestDoesPropertyPassFilter(FPropertyEditorPermissionTestStructB::StaticStruct(), "AOne", true, TEXT("[Test2] AOne on struct B"));
+		TestDoesPropertyPassFilter(FPropertyEditorPermissionTestStructB::StaticStruct(), "ATwo", false, TEXT("[Test2] ATwo on struct B"));
+		TestDoesPropertyPassFilter(FPropertyEditorPermissionTestStructB::StaticStruct(), "BOne", false, TEXT("[Test2] BOne on struct B"));
+		TestDoesPropertyPassFilter(FPropertyEditorPermissionTestStructB::StaticStruct(), "BTwo", false, TEXT("[Test2] BTwo on struct B"));
+		TestDoesPropertyPassFilter(FPropertyEditorPermissionTestStructC::StaticStruct(), "AOne", true, TEXT("[Test2] AOne on struct C"));
+		TestDoesPropertyPassFilter(FPropertyEditorPermissionTestStructC::StaticStruct(), "ATwo", false, TEXT("[Test2] ATwo on struct C"));
+		TestDoesPropertyPassFilter(FPropertyEditorPermissionTestStructC::StaticStruct(), "BOne", false, TEXT("[Test2] BOne on struct C"));
+		TestDoesPropertyPassFilter(FPropertyEditorPermissionTestStructC::StaticStruct(), "BTwo", false, TEXT("[Test2] BTwo on struct C"));
+		TestDoesPropertyPassFilter(FPropertyEditorPermissionTestStructC::StaticStruct(), "COne", true, TEXT("[Test2] COne on struct C"));
+		TestDoesPropertyPassFilter(FPropertyEditorPermissionTestStructC::StaticStruct(), "CTwo", true, TEXT("[Test2] CTwo on struct C"));
+	}
+
+	PermissionList.ClearPermissionList();
+	
+	{
+		// Allow all properties from struct C
+		// Properties from struct A or struct B for an instance of struct A, B, or C should be allowed
+		// Properties from struct C for an instance of struct C should be allowed
+		AddPermissionList(FPropertyEditorPermissionTestStructC::StaticStruct(), {}, {}, EPropertyPermissionListRules::AllowListAllProperties);
+		TestDoesPropertyPassFilter(FPropertyEditorPermissionTestStructA::StaticStruct(), "AOne", true, TEXT("[Test3] AOne on struct A"));
+		TestDoesPropertyPassFilter(FPropertyEditorPermissionTestStructA::StaticStruct(), "ATwo", true, TEXT("[Test3] ATwo on struct A"));
+		TestDoesPropertyPassFilter(FPropertyEditorPermissionTestStructB::StaticStruct(), "AOne", true, TEXT("[Test3] AOne on struct B"));
+		TestDoesPropertyPassFilter(FPropertyEditorPermissionTestStructB::StaticStruct(), "ATwo", true, TEXT("[Test3] ATwo on struct B"));
+		TestDoesPropertyPassFilter(FPropertyEditorPermissionTestStructB::StaticStruct(), "BOne", true, TEXT("[Test3] BOne on struct B"));
+		TestDoesPropertyPassFilter(FPropertyEditorPermissionTestStructB::StaticStruct(), "BTwo", true, TEXT("[Test3] BTwo on struct B"));
+		TestDoesPropertyPassFilter(FPropertyEditorPermissionTestStructC::StaticStruct(), "AOne", true, TEXT("[Test3] AOne on struct C"));
+		TestDoesPropertyPassFilter(FPropertyEditorPermissionTestStructC::StaticStruct(), "ATwo", true, TEXT("[Test3] ATwo on struct C"));
+		TestDoesPropertyPassFilter(FPropertyEditorPermissionTestStructC::StaticStruct(), "BOne", true, TEXT("[Test3] BOne on struct C"));
+		TestDoesPropertyPassFilter(FPropertyEditorPermissionTestStructC::StaticStruct(), "BTwo", true, TEXT("[Test3] BTwo on struct C"));
+		TestDoesPropertyPassFilter(FPropertyEditorPermissionTestStructC::StaticStruct(), "COne", true, TEXT("[Test3] COne on struct C"));
+		TestDoesPropertyPassFilter(FPropertyEditorPermissionTestStructC::StaticStruct(), "CTwo", true, TEXT("[Test3] CTwo on struct C"));
+		
+		// Allow AOne property from struct A
+		// AOne for an instance of struct A, B, or C should be allowed
+		// ATwo for an instance of struct A, B, or C should be denied
+		// Properties from struct B for an instance of struct B or C should be denied
+		// Properties from struct C for an instance of struct C should be allowed (due to the previous rule)
+		AddPermissionList(FPropertyEditorPermissionTestStructA::StaticStruct(), { "AOne" }, {}, EPropertyPermissionListRules::UseExistingPermissionList);
+		TestDoesPropertyPassFilter(FPropertyEditorPermissionTestStructA::StaticStruct(), "AOne", true, TEXT("[Test4] AOne on struct A"));
+		TestDoesPropertyPassFilter(FPropertyEditorPermissionTestStructA::StaticStruct(), "ATwo", false, TEXT("[Test4] ATwo on struct A"));
+		TestDoesPropertyPassFilter(FPropertyEditorPermissionTestStructB::StaticStruct(), "AOne", true, TEXT("[Test4] AOne on struct B"));
+		TestDoesPropertyPassFilter(FPropertyEditorPermissionTestStructB::StaticStruct(), "ATwo", false, TEXT("[Test4] ATwo on struct B"));
+		TestDoesPropertyPassFilter(FPropertyEditorPermissionTestStructB::StaticStruct(), "BOne", false, TEXT("[Test4] BOne on struct B"));
+		TestDoesPropertyPassFilter(FPropertyEditorPermissionTestStructB::StaticStruct(), "BTwo", false, TEXT("[Test4] BTwo on struct B"));
+		TestDoesPropertyPassFilter(FPropertyEditorPermissionTestStructC::StaticStruct(), "AOne", true, TEXT("[Test4] AOne on struct C"));
+		TestDoesPropertyPassFilter(FPropertyEditorPermissionTestStructC::StaticStruct(), "ATwo", false, TEXT("[Test4] ATwo on struct C"));
+		TestDoesPropertyPassFilter(FPropertyEditorPermissionTestStructC::StaticStruct(), "BOne", false, TEXT("[Test4] BOne on struct C"));
+		TestDoesPropertyPassFilter(FPropertyEditorPermissionTestStructC::StaticStruct(), "BTwo", false, TEXT("[Test4] BTwo on struct C"));
+		TestDoesPropertyPassFilter(FPropertyEditorPermissionTestStructC::StaticStruct(), "COne", true, TEXT("[Test4] COne on struct C"));
+		TestDoesPropertyPassFilter(FPropertyEditorPermissionTestStructC::StaticStruct(), "CTwo", true, TEXT("[Test4] CTwo on struct C"));
+	}
+
+	PermissionList.ClearPermissionList();
+	
+	{
+		// Allow all sub-class properties from struct B
+		// Properties from struct A for an instance of struct A should be allowed
+		// Properties from struct A or B for an instance of struct B or C should be denied
+		// Properties from struct C for an instance of struct C should be allowed
+		AddPermissionList(FPropertyEditorPermissionTestStructB::StaticStruct(), { "Dummy" }, {}, EPropertyPermissionListRules::AllowListAllSubclassProperties);
+		TestDoesPropertyPassFilter(FPropertyEditorPermissionTestStructA::StaticStruct(), "AOne", true, TEXT("[Test5] AOne on struct A"));
+		TestDoesPropertyPassFilter(FPropertyEditorPermissionTestStructA::StaticStruct(), "ATwo", true, TEXT("[Test5] ATwo on struct A"));
+		TestDoesPropertyPassFilter(FPropertyEditorPermissionTestStructB::StaticStruct(), "AOne", false, TEXT("[Test5] AOne on struct B"));
+		TestDoesPropertyPassFilter(FPropertyEditorPermissionTestStructB::StaticStruct(), "ATwo", false, TEXT("[Test5] ATwo on struct B"));
+		TestDoesPropertyPassFilter(FPropertyEditorPermissionTestStructB::StaticStruct(), "BOne", false, TEXT("[Test5] BOne on struct B"));
+		TestDoesPropertyPassFilter(FPropertyEditorPermissionTestStructB::StaticStruct(), "BTwo", false, TEXT("[Test5] BTwo on struct B"));
+		TestDoesPropertyPassFilter(FPropertyEditorPermissionTestStructC::StaticStruct(), "AOne", false, TEXT("[Test5] AOne on struct C"));
+		TestDoesPropertyPassFilter(FPropertyEditorPermissionTestStructC::StaticStruct(), "ATwo", false, TEXT("[Test5] ATwo on struct C"));
+		TestDoesPropertyPassFilter(FPropertyEditorPermissionTestStructC::StaticStruct(), "BOne", false, TEXT("[Test5] BOne on struct C"));
+		TestDoesPropertyPassFilter(FPropertyEditorPermissionTestStructC::StaticStruct(), "BTwo", false, TEXT("[Test5] BTwo on struct C"));
+		TestDoesPropertyPassFilter(FPropertyEditorPermissionTestStructC::StaticStruct(), "COne", true, TEXT("[Test5] COne on struct C"));
+		TestDoesPropertyPassFilter(FPropertyEditorPermissionTestStructC::StaticStruct(), "CTwo", true, TEXT("[Test5] CTwo on struct C"));
+		
+		// Allow all properties from struct A
+		// Properties from struct A for an instance of struct A, B, or C should be allowed
+		// Properties from struct B for an instance of struct B or C should be denied
+		// Properties from struct C for an instance of struct C should be allowed
+		AddPermissionList(FPropertyEditorPermissionTestStructA::StaticStruct(), {}, {}, EPropertyPermissionListRules::AllowListAllProperties);
+		TestDoesPropertyPassFilter(FPropertyEditorPermissionTestStructA::StaticStruct(), "AOne", true, TEXT("[Test6] AOne on struct A"));
+		TestDoesPropertyPassFilter(FPropertyEditorPermissionTestStructA::StaticStruct(), "ATwo", true, TEXT("[Test6] ATwo on struct A"));
+		TestDoesPropertyPassFilter(FPropertyEditorPermissionTestStructB::StaticStruct(), "AOne", true, TEXT("[Test6] AOne on struct B"));
+		TestDoesPropertyPassFilter(FPropertyEditorPermissionTestStructB::StaticStruct(), "ATwo", true, TEXT("[Test6] ATwo on struct B"));
+		TestDoesPropertyPassFilter(FPropertyEditorPermissionTestStructB::StaticStruct(), "BOne", false, TEXT("[Test6] BOne on struct B"));
+		TestDoesPropertyPassFilter(FPropertyEditorPermissionTestStructB::StaticStruct(), "BTwo", false, TEXT("[Test6] BTwo on struct B"));
+		TestDoesPropertyPassFilter(FPropertyEditorPermissionTestStructC::StaticStruct(), "AOne", true, TEXT("[Test6] AOne on struct C"));
+		TestDoesPropertyPassFilter(FPropertyEditorPermissionTestStructC::StaticStruct(), "ATwo", true, TEXT("[Test6] ATwo on struct C"));
+		TestDoesPropertyPassFilter(FPropertyEditorPermissionTestStructC::StaticStruct(), "BOne", false, TEXT("[Test6] BOne on struct C"));
+		TestDoesPropertyPassFilter(FPropertyEditorPermissionTestStructC::StaticStruct(), "BTwo", false, TEXT("[Test6] BTwo on struct C"));
+		TestDoesPropertyPassFilter(FPropertyEditorPermissionTestStructC::StaticStruct(), "COne", true, TEXT("[Test6] COne on struct C"));
+		TestDoesPropertyPassFilter(FPropertyEditorPermissionTestStructC::StaticStruct(), "CTwo", true, TEXT("[Test6] CTwo on struct C"));
+	}
+
+	PermissionList.ClearPermissionList();
+
+	{
+		// Allow BOne property from struct B
+		// Properties from struct A for an instance of struct A should be allowed
+		// Properties from struct A for an instance of struct B or C should be denied
+		// BOne for an instance of struct B or C should be allowed
+		// BTwo for an instance of struct B or C should be denied
+		// Properties from struct C for an instance of struct C should be denied
+		AddPermissionList(FPropertyEditorPermissionTestStructB::StaticStruct(), { "BOne" }, {}, EPropertyPermissionListRules::UseExistingPermissionList);
+		TestDoesPropertyPassFilter(FPropertyEditorPermissionTestStructA::StaticStruct(), "AOne", true, TEXT("[Test7] AOne on struct A"));
+		TestDoesPropertyPassFilter(FPropertyEditorPermissionTestStructA::StaticStruct(), "ATwo", true, TEXT("[Test7] ATwo on struct A"));
+		TestDoesPropertyPassFilter(FPropertyEditorPermissionTestStructB::StaticStruct(), "AOne", false, TEXT("[Test7] AOne on struct B"));
+		TestDoesPropertyPassFilter(FPropertyEditorPermissionTestStructB::StaticStruct(), "ATwo", false, TEXT("[Test7] ATwo on struct B"));
+		TestDoesPropertyPassFilter(FPropertyEditorPermissionTestStructB::StaticStruct(), "BOne", true, TEXT("[Test7] BOne on struct B"));
+		TestDoesPropertyPassFilter(FPropertyEditorPermissionTestStructB::StaticStruct(), "BTwo", false, TEXT("[Test7] BTwo on struct B"));
+		TestDoesPropertyPassFilter(FPropertyEditorPermissionTestStructC::StaticStruct(), "AOne", false, TEXT("[Test7] AOne on struct C"));
+		TestDoesPropertyPassFilter(FPropertyEditorPermissionTestStructC::StaticStruct(), "ATwo", false, TEXT("[Test7] ATwo on struct C"));
+		TestDoesPropertyPassFilter(FPropertyEditorPermissionTestStructC::StaticStruct(), "BOne", true, TEXT("[Test7] BOne on struct C"));
+		TestDoesPropertyPassFilter(FPropertyEditorPermissionTestStructC::StaticStruct(), "BTwo", false, TEXT("[Test7] BTwo on struct C"));
+		TestDoesPropertyPassFilter(FPropertyEditorPermissionTestStructC::StaticStruct(), "COne", false, TEXT("[Test7] COne on struct C"));
+		TestDoesPropertyPassFilter(FPropertyEditorPermissionTestStructC::StaticStruct(), "CTwo", false, TEXT("[Test7] CTwo on struct C"));
+		
+		// Allow all properties from struct A
+		// Properties from struct A for an instance of struct A, B, or C should be allowed
+		// BOne for an instance of struct B or C should be allowed
+		// BTwo for an instance of struct B or C should be denied
+		// Properties from struct C for an instance of struct C should be denied
+		AddPermissionList(FPropertyEditorPermissionTestStructA::StaticStruct(), {}, {}, EPropertyPermissionListRules::AllowListAllProperties);
+		TestDoesPropertyPassFilter(FPropertyEditorPermissionTestStructA::StaticStruct(), "AOne", true, TEXT("[Test8] AOne on struct A"));
+		TestDoesPropertyPassFilter(FPropertyEditorPermissionTestStructA::StaticStruct(), "ATwo", true, TEXT("[Test8] ATwo on struct A"));
+		TestDoesPropertyPassFilter(FPropertyEditorPermissionTestStructB::StaticStruct(), "AOne", true, TEXT("[Test8] AOne on struct B"));
+		TestDoesPropertyPassFilter(FPropertyEditorPermissionTestStructB::StaticStruct(), "ATwo", true, TEXT("[Test8] ATwo on struct B"));
+		TestDoesPropertyPassFilter(FPropertyEditorPermissionTestStructB::StaticStruct(), "BOne", true, TEXT("[Test8] BOne on struct B"));
+		TestDoesPropertyPassFilter(FPropertyEditorPermissionTestStructB::StaticStruct(), "BTwo", false, TEXT("[Test8] BTwo on struct B"));
+		TestDoesPropertyPassFilter(FPropertyEditorPermissionTestStructC::StaticStruct(), "AOne", true, TEXT("[Test8] AOne on struct C"));
+		TestDoesPropertyPassFilter(FPropertyEditorPermissionTestStructC::StaticStruct(), "ATwo", true, TEXT("[Test8] ATwo on struct C"));
+		TestDoesPropertyPassFilter(FPropertyEditorPermissionTestStructC::StaticStruct(), "BOne", true, TEXT("[Test8] BOne on struct C"));
+		TestDoesPropertyPassFilter(FPropertyEditorPermissionTestStructC::StaticStruct(), "BTwo", false, TEXT("[Test8] BTwo on struct C"));
+		TestDoesPropertyPassFilter(FPropertyEditorPermissionTestStructC::StaticStruct(), "COne", false, TEXT("[Test8] COne on struct C"));
+		TestDoesPropertyPassFilter(FPropertyEditorPermissionTestStructC::StaticStruct(), "CTwo", false, TEXT("[Test8] CTwo on struct C"));
+	}
+	
+	PermissionList.ClearPermissionList();
+
+	{
+		// Allow all sub-class properties from struct A
+		// Deny BOne from struct B
+		// Properties from struct A for an instance of struct A, B or C should be denied
+		// BOne for an instance of struct B or C should be denied
+		// BTwo for an instance of struct B or C should be allowed
+		// Properties from struct C for an instance of struct C should be allowed
+		AddPermissionList(FPropertyEditorPermissionTestStructA::StaticStruct(), { "Dummy" }, {}, EPropertyPermissionListRules::AllowListAllSubclassProperties);
+		AddPermissionList(FPropertyEditorPermissionTestStructB::StaticStruct(), {}, { "BOne" }, EPropertyPermissionListRules::UseExistingPermissionList);
+		TestDoesPropertyPassFilter(FPropertyEditorPermissionTestStructA::StaticStruct(), "AOne", false, TEXT("[Test9] AOne on struct A"));
+		TestDoesPropertyPassFilter(FPropertyEditorPermissionTestStructA::StaticStruct(), "ATwo", false, TEXT("[Test9] ATwo on struct A"));
+		TestDoesPropertyPassFilter(FPropertyEditorPermissionTestStructB::StaticStruct(), "AOne", false, TEXT("[Test9] AOne on struct B"));
+		TestDoesPropertyPassFilter(FPropertyEditorPermissionTestStructB::StaticStruct(), "ATwo", false, TEXT("[Test9] ATwo on struct B"));
+		TestDoesPropertyPassFilter(FPropertyEditorPermissionTestStructB::StaticStruct(), "BOne", false, TEXT("[Test9] BOne on struct B"));
+		TestDoesPropertyPassFilter(FPropertyEditorPermissionTestStructB::StaticStruct(), "BTwo", true, TEXT("[Test9] BTwo on struct B"));
+		TestDoesPropertyPassFilter(FPropertyEditorPermissionTestStructC::StaticStruct(), "AOne", false, TEXT("[Test9] AOne on struct C"));
+		TestDoesPropertyPassFilter(FPropertyEditorPermissionTestStructC::StaticStruct(), "ATwo", false, TEXT("[Test9] ATwo on struct C"));
+		TestDoesPropertyPassFilter(FPropertyEditorPermissionTestStructC::StaticStruct(), "BOne", false, TEXT("[Test9] BOne on struct C"));
+		TestDoesPropertyPassFilter(FPropertyEditorPermissionTestStructC::StaticStruct(), "BTwo", true, TEXT("[Test9] BTwo on struct C"));
+		TestDoesPropertyPassFilter(FPropertyEditorPermissionTestStructC::StaticStruct(), "COne", true, TEXT("[Test9] COne on struct C"));
+		TestDoesPropertyPassFilter(FPropertyEditorPermissionTestStructC::StaticStruct(), "CTwo", true, TEXT("[Test9] CTwo on struct C"));
+	}
+
+	PermissionList.ClearPermissionList();
+	
+	return true;
+}
+
+#endif // WITH_DEV_AUTOMATION_TESTS
diff --git a/Engine/Source/Editor/PropertyEditor/Private/Tests/PropertyPermissionTests.h b/Engine/Source/Editor/PropertyEditor/Private/Tests/PropertyPermissionTests.h
new file mode 100644
--- /dev/null
+++ b/Engine/Source/Editor/PropertyEditor/Private/Tests/PropertyPermissionTests.h
@@ -0,0 +1,44 @@
+﻿// Copyright Epic Games, Inc. All Rights Reserved.
+
+#pragma once
+
+#include "PropertyPermissionTests.generated.h"
+
+USTRUCT()
+struct FPropertyEditorPermissionTestStructA
+{
+	GENERATED_BODY()
+	
+public:
+	UPROPERTY(EditAnywhere, Category="Properties")
+	int32 AOne = 0;
+	
+	UPROPERTY(EditAnywhere, Category="Properties")
+	int32 ATwo = 0;
+};
+
+USTRUCT()
+struct FPropertyEditorPermissionTestStructB : public FPropertyEditorPermissionTestStructA
+{
+	GENERATED_BODY()
+	
+public:
+	UPROPERTY(EditAnywhere, Category="Properties")
+	int32 BOne = 0;
+	
+	UPROPERTY(EditAnywhere, Category="Properties")
+	int32 BTwo = 0;
+};
+
+USTRUCT()
+struct FPropertyEditorPermissionTestStructC : public FPropertyEditorPermissionTestStructB
+{
+	GENERATED_BODY()
+	
+public:
+	UPROPERTY(EditAnywhere, Category="Properties")
+	int32 COne = 0;
+	
+	UPROPERTY(EditAnywhere, Category="Properties")
+	int32 CTwo = 0;
+};
diff --git a/Engine/Source/Editor/PropertyEditor/Public/PropertyPermissionList.h b/Engine/Source/Editor/PropertyEditor/Public/PropertyPermissionList.h
--- a/Engine/Source/Editor/PropertyEditor/Public/PropertyPermissionList.h
+++ b/Engine/Source/Editor/PropertyEditor/Public/PropertyPermissionList.h
@@ -116,7 +116,7 @@ public:
 	/** Check whether a property exists on the DenyList for a specific Struct - this will return false if the property is DenyListed on a parent Struct */
 	UE_API bool IsSpecificPropertyDenyListed(const UStruct* ObjectStruct, FName PropertyName) const;
 
-	/** Gets a read-only copy of the original, un-flattened PermissionList. */
+	/** Gets a read-only reference to the original, un-flattened PermissionList. */
 	const TMap<TSoftObjectPtr<const UStruct>, FPropertyPermissionListEntry>& GetRawPermissionList() const { return RawPropertyPermissionList; }
 
 	/** Clear CachedPropertyPermissionList to cause the PermissionListed property list to be regenerated next time it's queried */
@@ -156,8 +156,8 @@ private:
 	mutable TMap<TWeakObjectPtr<const UStruct>, FPropertyPermissionListEntry> CachedPropertyPermissionList;
 
 	/** Get or create the cached PermissionList for a specific UStruct */
-	UE_API const FNamePermissionList& GetCachedPermissionListForStruct(const UStruct* Struct) const;
-	UE_API const FNamePermissionList& GetCachedPermissionListForStructHelper(const UStruct* Struct, bool& bInOutShouldAllowListAllProperties) const;
+	const FNamePermissionList& GetCachedPermissionListForStruct(const UStruct* Struct) const;
+	const FPropertyPermissionListEntry& GetCachedPermissionListEntryForStructHelper(const UStruct* Struct) const;
 };
 
 class FPropertyEditorPermissionList : public FPropertyPermissionList
