SDetailSingleItemRow : when OnPasteFromText is called for a category don't include the edit condition in the call to IsEditConst

When pasting an entire category, an edit condition property might come after the property that uses it. Consider that the value on the clipboard enables the edit condition but that it may have been disabled prior to the paste. Therefore, the edit condition should not be included when checking if a pasted property IsEditConst.

diff --git a/Engine/Source/Editor/PropertyEditor/Private/SDetailSingleItemRow.cpp b/Engine/Source/Editor/PropertyEditor/Private/SDetailSingleItemRow.cpp
--- a/Engine/Source/Editor/PropertyEditor/Private/SDetailSingleItemRow.cpp
+++ b/Engine/Source/Editor/PropertyEditor/Private/SDetailSingleItemRow.cpp
@@ -385,6 +385,8 @@ void SDetailSingleItemRow::Construct( const FArguments& InArgs, FDetailLayoutCus
 						OnCopyToTextDelegate = GroupOnCopyToTextDelegate;
 					}
 				}
+
+				bool bIncludeEditCondition = true;
 				
 				if (WidgetRow.IsPasteFromTextBound())
 				{
@@ -393,6 +395,11 @@ void SDetailSingleItemRow::Construct( const FArguments& InArgs, FDetailLayoutCus
 				else if (Category.IsValid())
 				{
 					OnPasteFromTextDelegate = Category->OnPasteFromText();
+
+					// When pasting an entire category, an edit condition property might come after the property that uses it.
+					// Consider the case where the value on the clipboard enables the edit condition but that it may have been disabled prior to the paste.
+					// Therefore, the edit condition should not be included when checking if a pasted property IsEditConst.
+					bIncludeEditCondition = false;
 				}
 				// If still not set, but this is a group, initialize
 				else if (Group.IsValid())
@@ -407,7 +414,7 @@ void SDetailSingleItemRow::Construct( const FArguments& InArgs, FDetailLayoutCus
 				// also bind to OnCopyToText to match it
 				if (OnPasteFromTextDelegate.IsValid())
 				{
-					OnPasteFromTextDelegate->AddSP(this, &SDetailSingleItemRow::OnPasteFromText);
+					OnPasteFromTextDelegate->AddSP(this, &SDetailSingleItemRow::OnPasteFromText, bIncludeEditCondition);
 				}
 			}
 
@@ -1528,7 +1535,7 @@ TSharedRef<SWidget> SDetailSingleItemRow::GetNameWidget(TSharedRef<SWidget> Name
 	return NameWidget;
 }
 
-bool SDetailSingleItemRow::CanPasteFromText(const FString& InTag, const FString& InText) const
+bool SDetailSingleItemRow::CanPasteFromText(const FString& InTag, const FString& InText, const bool bIncludeEditCondition) const
 {
 	if (InText.IsEmpty())
 	{
@@ -1579,7 +1586,7 @@ bool SDetailSingleItemRow::CanPasteFromText(const FString& InTag, const FString&
 		{
 			if (const FPropertyEditor* PropertyEditor = PropertyRow->GetPropertyEditor().Get())
 			{
-				return !PropertyEditor->IsEditConst();
+				return !PropertyEditor->IsEditConst(bIncludeEditCondition);
 			}
 		}
 	}
@@ -1616,9 +1623,9 @@ void SDetailSingleItemRow::OnPasteProperty()
 	}
 }
 
-void SDetailSingleItemRow::OnPasteFromText(const FString& InTag, const FString& InText, const TOptional<FGuid>& InOperationId)
+void SDetailSingleItemRow::OnPasteFromText(const FString& InTag, const FString& InText, const TOptional<FGuid>& InOperationId, const bool bIncludeEditCondition)
 {
-	if (PasteFromText(InTag, InText))
+	if (PasteFromText(InTag, InText, bIncludeEditCondition))
 	{
 		TSharedPtr<FPropertyNode> PropertyNode = GetPropertyNode();
 		if (PropertyNode.IsValid())
@@ -1630,9 +1637,9 @@ void SDetailSingleItemRow::OnPasteFromText(const FString& InTag, const FString&
 	}
 }
 
-bool SDetailSingleItemRow::PasteFromText(const FString& InTag, const FString& InText)
+bool SDetailSingleItemRow::PasteFromText(const FString& InTag, const FString& InText, const bool bIncludeEditCondition)
 {
-	if (!CanPasteFromText(InTag, InText))
+	if (!CanPasteFromText(InTag, InText, bIncludeEditCondition))
 	{
 		return false;
 	}
diff --git a/Engine/Source/Editor/PropertyEditor/Private/SDetailSingleItemRow.h b/Engine/Source/Editor/PropertyEditor/Private/SDetailSingleItemRow.h
--- a/Engine/Source/Editor/PropertyEditor/Private/SDetailSingleItemRow.h
+++ b/Engine/Source/Editor/PropertyEditor/Private/SDetailSingleItemRow.h
@@ -102,17 +102,17 @@ private:
 	/**
 	 * @return True if the (optionally tagged) input contents can be pasted
 	 */
-	bool CanPasteFromText(const FString& InTag, const FString& InText) const;
+	bool CanPasteFromText(const FString& InTag, const FString& InText, const bool bIncludeEditCondition = true) const;
 
 	/**
 	 * Delegate handling pasting an optionally tagged text snippet
 	 */
-	void OnPasteFromText(const FString& InTag, const FString& InText, const TOptional<FGuid>& InOperationId);
+	void OnPasteFromText(const FString& InTag, const FString& InText, const TOptional<FGuid>& InOperationId, const bool bIncludeEditCondition);
 
 	/**
 	 * Handle pasting an optionally tagged text snippet
 	 */
-	bool PasteFromText(const FString& InTag, const FString& InText);
+	bool PasteFromText(const FString& InTag, const FString& InText, const bool bIncludeEditCondition = true);
 	
 	FSlateColor GetOuterBackgroundColor() const;
 	FSlateColor GetInnerBackgroundColor() const;
