Copy-pasting set of string properties no longer copies additional quotes around each item
* Fix is to make FPropertyHandleSet::SetValueFromFormattedString the same as FPropertyHandleArray::SetValueFromFormattedString and trim the quotes before application

diff --git a/Engine/Source/Editor/PropertyEditor/Private/PropertyHandleImpl.cpp b/Engine/Source/Editor/PropertyEditor/Private/PropertyHandleImpl.cpp
--- a/Engine/Source/Editor/PropertyEditor/Private/PropertyHandleImpl.cpp
+++ b/Engine/Source/Editor/PropertyEditor/Private/PropertyHandleImpl.cpp
@@ -6267,6 +6267,10 @@ FPropertyAccess::Result FPropertyHandleSet::SetValueFromFormattedString(const FS
 			}
 			const TSharedRef<IPropertyHandle> Property = GetElement( NumElements - 1 );
 			ensure(!Property->AsStruct());
+			
+			// Trim ""
+			FString FormattedString = Value.TrimQuotes();
+			
 			if ( Property->IsValidHandle() && Property->SetValueFromFormattedString(Value, Flags) == FPropertyAccess::Fail)
 			{
 				Result = FPropertyAccess::Fail;
