Use output of trim string on TSet<String> copy-paste

diff --git a/Engine/Source/Editor/PropertyEditor/Private/PropertyHandleImpl.cpp b/Engine/Source/Editor/PropertyEditor/Private/PropertyHandleImpl.cpp
--- a/Engine/Source/Editor/PropertyEditor/Private/PropertyHandleImpl.cpp
+++ b/Engine/Source/Editor/PropertyEditor/Private/PropertyHandleImpl.cpp
@@ -6271,7 +6271,7 @@ FPropertyAccess::Result FPropertyHandleSet::SetValueFromFormattedString(const FS
 			// Trim ""
 			FString FormattedString = Value.TrimQuotes();
 			
-			if ( Property->IsValidHandle() && Property->SetValueFromFormattedString(Value, Flags) == FPropertyAccess::Fail)
+			if ( Property->IsValidHandle() && Property->SetValueFromFormattedString(FormattedString, Flags) == FPropertyAccess::Fail)
 			{
 				Result = FPropertyAccess::Fail;
 			}
