Make asset selector's browse button have customizable tooltip & icon, and do not indicate that Alt can be used for actors by default, because this key is only handled for Content Browser assets.

----------------------------------------
Fix flipped logical operator in CL 46303391.

----------------------------------------

Note: This is two changes cherry picked in and merged together: 46303391 & 46361108 by kristof.morva1

#rnx

diff --git a/Engine/Source/Editor/PropertyEditor/Private/PropertyCustomizationHelpers.cpp b/Engine/Source/Editor/PropertyEditor/Private/PropertyCustomizationHelpers.cpp
--- a/Engine/Source/Editor/PropertyEditor/Private/PropertyCustomizationHelpers.cpp
+++ b/Engine/Source/Editor/PropertyEditor/Private/PropertyCustomizationHelpers.cpp
@@ -950,6 +950,8 @@ void SObjectPropertyEntryBox::Construct( const FArguments& InArgs )
 				.DisplayUseSelected(InArgs._DisplayUseSelected)
 				.DisplayBrowse(InArgs._DisplayBrowse)
 				.OnBrowseOverride(InArgs._OnBrowseOverride)
+				.OnBrowseToolTipOverride(InArgs._OnBrowseToolTipOverride)
+				.OnBrowseIconOverride(InArgs._OnBrowseIconOverride)
 				.EnableContentPicker(InArgs._EnableContentPicker)
 				.PropertyHandle(PropertyHandle)
 				.OwnerAssetDataArray(OwnerAssetDataArray)
diff --git a/Engine/Source/Editor/PropertyEditor/Private/UserInterface/PropertyEditor/SPropertyEditorAsset.cpp b/Engine/Source/Editor/PropertyEditor/Private/UserInterface/PropertyEditor/SPropertyEditorAsset.cpp
--- a/Engine/Source/Editor/PropertyEditor/Private/UserInterface/PropertyEditor/SPropertyEditorAsset.cpp
+++ b/Engine/Source/Editor/PropertyEditor/Private/UserInterface/PropertyEditor/SPropertyEditorAsset.cpp
@@ -837,9 +837,12 @@ void SPropertyEditorAsset::Construct(const FArguments& InArgs, const TSharedPtr<
 	{
 		FSimpleDelegate OnBrowseDelegate = InArgs._OnBrowseOverride.IsBound() ? InArgs._OnBrowseOverride : FSimpleDelegate::CreateSP(this, &SPropertyEditorAsset::OnBrowse);
 
-		// Only the default SPropertyEditorAsset::OnBrowse delegate supports alt+click to directly open the asset, so we only modify the icon if there is no browse override
-		TAttribute<const FSlateBrush*> IconDelegate = !InArgs._OnBrowseOverride.IsBound() ?
-			TAttribute<const FSlateBrush*>::CreateSP(this, &SPropertyEditorAsset::GetOnBrowseIcon) : TAttribute<const FSlateBrush*>();
+		// Only the default SPropertyEditorAsset::OnBrowse delegate supports alt+click to directly open the asset by default, so we only modify the icon if there is no browse override
+		TAttribute<const FSlateBrush*> IconDelegate = InArgs._OnBrowseIconOverride.IsSet() ? InArgs._OnBrowseIconOverride
+			: !InArgs._OnBrowseOverride.IsBound() ? TAttribute<const FSlateBrush*>::CreateSP(this, &SPropertyEditorAsset::GetOnBrowseIcon)
+			: TAttribute<const FSlateBrush*>();
+
+		TAttribute<FText> OnBrowseToolTip = InArgs._OnBrowseToolTipOverride.IsSet() ? InArgs._OnBrowseToolTipOverride : TAttribute<FText>(this, &SPropertyEditorAsset::GetOnBrowseToolTip);
 
 		// Browse to is readonly. There are cases the reference is set elsewhere other than from the content picker, and we want to allow browsing to.
 		bool bIsPickerEnabled = IsEnabledAttribute.Get();
@@ -864,7 +867,7 @@ void SPropertyEditorAsset::Construct(const FArguments& InArgs, const TSharedPtr<
 		[
 			PropertyCustomizationHelpers::MakeBrowseButton(
 				OnBrowseDelegate,
-				TAttribute<FText>( this, &SPropertyEditorAsset::GetOnBrowseToolTip ),
+				OnBrowseToolTip,
 				IsBrowseEnabledAttribute,
 				bIsActor,
 				IconDelegate
@@ -1726,7 +1729,7 @@ FText SPropertyEditorAsset::GetOnBrowseToolTip() const
 		Args.Add(TEXT("Asset"), FText::AsCultureInvariant(Value.Object->GetName()));
 		if (bIsActor)
 		{
-			return FText::Format(LOCTEXT( "SelectSpecificActorInViewport", "Select '{Asset}' in the viewport {AltOpen}"), Args);
+			return FText::Format(LOCTEXT( "SelectSpecificActorInViewport", "Select '{Asset}' in the viewport"), Args);
 		}
 		else
 		{
@@ -1736,7 +1739,7 @@ FText SPropertyEditorAsset::GetOnBrowseToolTip() const
 	
 	if (bIsActor)
 	{
-		return FText::Format(LOCTEXT("SelectActorInViewport", "Select Actor in the viewport {AltOpen}"), Args);
+		return FText::Format(LOCTEXT("SelectActorInViewport", "Select Actor in the viewport"), Args);
 	}
 	else
 	{
@@ -1747,7 +1750,7 @@ FText SPropertyEditorAsset::GetOnBrowseToolTip() const
 const FSlateBrush* SPropertyEditorAsset::GetOnBrowseIcon() const
 {
 	// If the widget is hovered and alt is held down, show the edit icon
-	if (IsHovered() && FSlateApplication::Get().GetModifierKeys().IsAltDown())
+	if (!bIsActor && IsHovered() && FSlateApplication::Get().GetModifierKeys().IsAltDown())
 	{
 		return FAppStyle::Get().GetBrush("Icons.Edit");
 	}
diff --git a/Engine/Source/Editor/PropertyEditor/Private/UserInterface/PropertyEditor/SPropertyEditorAsset.h b/Engine/Source/Editor/PropertyEditor/Private/UserInterface/PropertyEditor/SPropertyEditorAsset.h
--- a/Engine/Source/Editor/PropertyEditor/Private/UserInterface/PropertyEditor/SPropertyEditorAsset.h
+++ b/Engine/Source/Editor/PropertyEditor/Private/UserInterface/PropertyEditor/SPropertyEditorAsset.h
@@ -70,6 +70,8 @@ public:
 		SLATE_ARGUMENT(bool, DisplayUseSelected)
 		SLATE_ARGUMENT(bool, DisplayBrowse)
 		SLATE_EVENT(FSimpleDelegate, OnBrowseOverride)
+		SLATE_ATTRIBUTE(FText, OnBrowseToolTipOverride)
+		SLATE_ATTRIBUTE(const FSlateBrush*, OnBrowseIconOverride)
 		SLATE_ARGUMENT(bool, EnableContentPicker)
 		SLATE_ARGUMENT(bool, DisplayCompactSize)
 		SLATE_ARGUMENT(TSharedPtr<FAssetThumbnailPool>, ThumbnailPool)
diff --git a/Engine/Source/Editor/PropertyEditor/Public/PropertyCustomizationHelpers.h b/Engine/Source/Editor/PropertyEditor/Public/PropertyCustomizationHelpers.h
--- a/Engine/Source/Editor/PropertyEditor/Public/PropertyCustomizationHelpers.h
+++ b/Engine/Source/Editor/PropertyEditor/Public/PropertyCustomizationHelpers.h
@@ -313,6 +313,10 @@ public:
 		SLATE_ARGUMENT(bool, DisplayBrowse)
 		/** Optional delegate called when the 'Browse' button is clicked. Used to override the default editor behavior */
 		SLATE_EVENT(FSimpleDelegate, OnBrowseOverride)
+		/** Optional override for the tooltip of the 'Browse' button */
+		SLATE_ATTRIBUTE(FText, OnBrowseToolTipOverride)
+		/** Optional override for the icon of the 'Browse' button */
+		SLATE_ATTRIBUTE(const FSlateBrush*, OnBrowseIconOverride)
 		/** Whether to enable the content Picker */
 		SLATE_ARGUMENT(bool, EnableContentPicker)
 		/** Whether or not to display a smaller, compact size for the asset thumbnail */ 
