Fix custom property builders were not re-initialized after an array element has been moved to another index in a property panel.
A rebuild happened only in the middle of the operation (when an extra element is added to be used for a swap op), resulting in in array being displayed that has +1 element in it.

#rnx

diff --git a/Engine/Source/Editor/PropertyEditor/Private/PropertyHandleImpl.cpp b/Engine/Source/Editor/PropertyEditor/Private/PropertyHandleImpl.cpp
--- a/Engine/Source/Editor/PropertyEditor/Private/PropertyHandleImpl.cpp
+++ b/Engine/Source/Editor/PropertyEditor/Private/PropertyHandleImpl.cpp
@@ -2307,6 +2307,8 @@ void FPropertyValueImpl::MoveElementTo(int32 OriginalIndex, int32 NewIndex)
 			}
 		}
 	}
+
+	GetPropertyNode()->RebuildChildren();
 }
 
 void FPropertyValueImpl::DuplicateChild( int32 Index )
