Fixed merge conflict between 9956594 and 11308033 which resulted in the non-deterministic visibility of the FDetailMultiTopLevelObjectRootNode.

diff --git a/Engine/Source/Editor/PropertyEditor/Private/DetailMultiTopLevelObjectRootNode.cpp b/Engine/Source/Editor/PropertyEditor/Private/DetailMultiTopLevelObjectRootNode.cpp
--- a/Engine/Source/Editor/PropertyEditor/Private/DetailMultiTopLevelObjectRootNode.cpp
+++ b/Engine/Source/Editor/PropertyEditor/Private/DetailMultiTopLevelObjectRootNode.cpp
@@ -209,7 +209,7 @@ void FDetailMultiTopLevelObjectRootNode::FilterNode( const FDetailFilter& InFilt
 
 bool FDetailMultiTopLevelObjectRootNode::ShouldShowOnlyChildren() const
 {
-	return RootObjectCustomization.IsValid() && RootObjectSet.RootObjects.Num() ? !RootObjectCustomization.Pin()->ShouldDisplayHeader(RootObjectSet) : bShouldShowOnlyChildren;
+	return RootObjectCustomization.IsValid() && RootObjectSet.RootObjects.Num() ? !RootObjectCustomization.Pin()->ShouldDisplayHeader(RootObjectSet) : false;
 }
 
 void FDetailMultiTopLevelObjectRootNode::GenerateWidget_Internal(FDetailWidgetRow& OutRow, TSharedPtr<SDetailMultiTopLevelObjectTableRow> TableRowWidget) const
diff --git a/Engine/Source/Editor/PropertyEditor/Private/DetailMultiTopLevelObjectRootNode.h b/Engine/Source/Editor/PropertyEditor/Private/DetailMultiTopLevelObjectRootNode.h
--- a/Engine/Source/Editor/PropertyEditor/Private/DetailMultiTopLevelObjectRootNode.h
+++ b/Engine/Source/Editor/PropertyEditor/Private/DetailMultiTopLevelObjectRootNode.h
@@ -73,5 +73,4 @@ private:
 	FName NodeName;
 	bool bShouldBeVisible;
 	bool bHasFilterStrings;
-	bool bShouldShowOnlyChildren;
 };
