Remove characters with wrong encoding in comment

#rnx

diff --git a/Engine/Source/Editor/PropertyEditor/Public/DetailsViewArgs.h b/Engine/Source/Editor/PropertyEditor/Public/DetailsViewArgs.h
--- a/Engine/Source/Editor/PropertyEditor/Public/DetailsViewArgs.h
+++ b/Engine/Source/Editor/PropertyEditor/Public/DetailsViewArgs.h
@@ -27,7 +27,7 @@ enum class EEditDefaultsOnlyNodeVisibility : uint8
 	Automatic,
 };
 
-/**Ã©Ã©
+/**
  * Init params for a details view widget
  */
 struct FDetailsViewArgs
