Request a forced refresh of details panel on reset to default behaviour
* Only affect the "PropertyUpdateWidget"

diff --git a/Engine/Source/Editor/PropertyEditor/Private/SDetailCategoryTableRow.cpp b/Engine/Source/Editor/PropertyEditor/Private/SDetailCategoryTableRow.cpp
--- a/Engine/Source/Editor/PropertyEditor/Private/SDetailCategoryTableRow.cpp
+++ b/Engine/Source/Editor/PropertyEditor/Private/SDetailCategoryTableRow.cpp
@@ -134,6 +134,14 @@ void SDetailCategoryTableRow::Construct(const FArguments& InArgs, TSharedRef<FDe
 			{
 				ResetToDefault.Execute();
 			}
+
+			if (TSharedPtr<FDetailTreeNode> OwnerTreeNodePinned = OwnerTreeNode.Pin())
+			{
+				if (TSharedPtr<IDetailsViewPrivate> DetailsView = OwnerTreeNodePinned->GetDetailsViewSharedPtr())
+				{
+					DetailsView->RequestForceRefresh();
+				}
+			}
 		});
 		Args.PropertyPath = nullptr;
 		Args.Category = InOwnerTreeNode->GetNodeName();
diff --git a/Engine/Source/Editor/PropertyEditor/Private/SDetailSingleItemRow.cpp b/Engine/Source/Editor/PropertyEditor/Private/SDetailSingleItemRow.cpp
--- a/Engine/Source/Editor/PropertyEditor/Private/SDetailSingleItemRow.cpp
+++ b/Engine/Source/Editor/PropertyEditor/Private/SDetailSingleItemRow.cpp
@@ -676,7 +676,18 @@ void SDetailSingleItemRow::Construct( const FArguments& InArgs, FDetailLayoutCus
 						if (PropertyPath->IsValid())
 						{
 							FConstructPropertyUpdatedWidgetBuilderArgs Args;
-							Args.ResetToDefaultAction = FExecuteAction::CreateSP(this, &SDetailSingleItemRow::OnResetToDefaultClicked);
+							Args.ResetToDefaultAction =	FExecuteAction::CreateSPLambda(this, [this]
+							{
+								this->OnResetToDefaultClicked();
+
+								if (TSharedPtr<FDetailTreeNode> OwnerTreeNodePinned = OwnerTreeNode.Pin())
+								{
+									if (TSharedPtr<IDetailsViewPrivate> DetailsView = OwnerTreeNodePinned->GetDetailsViewSharedPtr())
+									{
+										DetailsView->RequestForceRefresh();
+									}
+								}
+							});
 							Args.Objects = &Objects;
 							Args.PropertyPath = PropertyPath;
 							Args.Category = Category.IsValid() ? Category->GetCategoryName() : NAME_None;
