Fix for blueprint default properties with custom property customization not using the correct IsEnabled value in certain cases

diff --git a/Engine/Source/Editor/PropertyEditor/Private/DetailPropertyRow.cpp b/Engine/Source/Editor/PropertyEditor/Private/DetailPropertyRow.cpp
--- a/Engine/Source/Editor/PropertyEditor/Private/DetailPropertyRow.cpp
+++ b/Engine/Source/Editor/PropertyEditor/Private/DetailPropertyRow.cpp
@@ -382,11 +382,6 @@ void FDetailPropertyRow::OnItemNodeInitialized( TSharedRef<FDetailCategoryImpl>
 
 		FixEmptyHeaderRowInContainers(PropertyHandle, CustomPropertyWidget);
 
-		// set initial value of enabled attribute to settings from struct customization
-		if (CustomPropertyWidget->IsEnabledAttr.IsSet())
-		{
-			CustomIsEnabledAttrib = CustomPropertyWidget->IsEnabledAttr;
-		}
 
 		// set initial value of auto-expand from struct customization
 		if (CustomPropertyWidget->ForceAutoExpansion.IsSet())
@@ -886,6 +881,7 @@ bool FDetailPropertyRow::GetEnabledState() const
 	bool Result = IsParentEnabled.Get(true);
 
 	Result = Result && CustomIsEnabledAttrib.Get(true);
+	Result = Result && (CustomPropertyWidget == nullptr || CustomPropertyWidget->IsEnabledAttr.Get(true));
 	Result = Result && IsEditConditionMet();
 
 	return Result;
