Editor: add experimental function to refresh all details panels from property editor module.

diff --git a/Engine/Source/Editor/PropertyEditor/Private/PropertyEditorModule.cpp b/Engine/Source/Editor/PropertyEditor/Private/PropertyEditorModule.cpp
--- a/Engine/Source/Editor/PropertyEditor/Private/PropertyEditorModule.cpp
+++ b/Engine/Source/Editor/PropertyEditor/Private/PropertyEditorModule.cpp
@@ -136,6 +136,21 @@ void FPropertyEditorModule::NotifyCustomizationModuleChanged()
 	}
 }
 
+void FPropertyEditorModule::RequestRefreshOfAllDetailsPanels()
+{
+	if (FSlateApplication::IsInitialized())
+	{
+		for (int32 ViewIndex = 0; ViewIndex < AllDetailViews.Num(); ++ViewIndex)
+		{
+			if (AllDetailViews[ViewIndex].IsValid())
+			{
+				TSharedPtr<SDetailsView> DetailViewPin = AllDetailViews[ViewIndex].Pin();
+				DetailViewPin->RequestForceRefresh();
+			}
+		}
+	}
+}
+
 static bool ShouldShowProperty(const FPropertyAndParent& PropertyAndParent, bool bHaveTemplate)
 {
 	const FProperty& Property = PropertyAndParent.Property;
diff --git a/Engine/Source/Editor/PropertyEditor/Public/PropertyEditorModule.h b/Engine/Source/Editor/PropertyEditor/Public/PropertyEditorModule.h
--- a/Engine/Source/Editor/PropertyEditor/Public/PropertyEditorModule.h
+++ b/Engine/Source/Editor/PropertyEditor/Public/PropertyEditorModule.h
@@ -360,6 +360,9 @@ public:
 	 */
 	virtual void NotifyCustomizationModuleChanged();
 
+	UE_EXPERIMENTAL(5.8, "This function may be removed and replaced with some other signaling mechanism.")
+	virtual void RequestRefreshOfAllDetailsPanels();
+
 	/**
 	 * Creates a new detail view
 	 *
