Use TBaseStructure<FDateTime> instead of Z_Construct_UScriptStruct_FDateTime.

diff --git a/Engine/Source/Editor/PropertyEditor/Private/UserInterface/PropertyEditor/SPropertyEditorDateTime.cpp b/Engine/Source/Editor/PropertyEditor/Private/UserInterface/PropertyEditor/SPropertyEditorDateTime.cpp
--- a/Engine/Source/Editor/PropertyEditor/Private/UserInterface/PropertyEditor/SPropertyEditorDateTime.cpp
+++ b/Engine/Source/Editor/PropertyEditor/Private/UserInterface/PropertyEditor/SPropertyEditorDateTime.cpp
@@ -49,8 +49,7 @@ bool SPropertyEditorDateTime::Supports( const TSharedRef< FPropertyEditor >& InP
 	if (Property->IsA(FStructProperty::StaticClass()))
 	{
 		const FStructProperty* StructProp = CastField<const FStructProperty>(Property);
-		extern UScriptStruct* Z_Construct_UScriptStruct_FDateTime();	// It'd be really nice if StaticStruct() worked on types declared in Object.h
-		if (Z_Construct_UScriptStruct_FDateTime() == StructProp->Struct)
+		if (TBaseStructure<FDateTime>::Get() == StructProp->Struct)
 		{
 			return true;
 		}
