[PropertyEditor]
- Fixed Property Node hierarchy was broken when a structure node is attached to a container node.
- Fixed array property inside a structure node not populated correctly because the structure node was never validated.

Every tick, we validate if number of container elements from the actual array/map/set memory matches the number of child property nodes. If there is a mismatch, we rebuild the container property node hierarchy. However, the rebuild can only be based on the container memory and we would lose all the manually attached structure node hierarchy.

When we manipulate the container in the editor, We validate the property node hierarchy to see if we need to refresh the tree. When there is a manually attached structure node to a container or a primitive type node, we didn't verify the attached node but only the parent node.

diff --git a/Engine/Plugins/Runtime/StateTree/Source/StateTreeEditorModule/Private/StateTreePropertyHelpers.cpp b/Engine/Plugins/Runtime/StateTree/Source/StateTreeEditorModule/Private/StateTreePropertyHelpers.cpp
--- a/Engine/Plugins/Runtime/StateTree/Source/StateTreeEditorModule/Private/StateTreePropertyHelpers.cpp
+++ b/Engine/Plugins/Runtime/StateTree/Source/StateTreeEditorModule/Private/StateTreePropertyHelpers.cpp
@@ -122,15 +122,20 @@ void DispatchPostEditToNodes(UObject& Owner, FPropertyChangedChainEvent& InPrope
 			const int32 Index = InPropertyChangedEvent.GetArrayIndex(ArrayProperty->GetName());
 			if (!Helper.IsValidIndex(Index))
 			{
-				check(CurrentPropNode->GetNextNode() == nullptr);
+				// It's possible that a property function node is bound directly to an array property
+				// In which case the there is no index for the array property while the property chain is still going.
+				// We will instead continue traversing the property chain and fetch the property function node property
 				break;
 			}
 
 			if (TargetPath.GetStructID().IsValid())
 			{
+				// Skip adding segment for array property itself if we are not modifying its element,
+				// As we will not indirect to the inner property and will add segment for array property at the end of the loop
 				TargetPath.AddPathSegment(ArrayProperty->GetFName(), Index);
 			}
 
+			// Only indirect ArrayProperty to Inner if we are modifying array element
 			CurrentAddress = Helper.GetRawPtr(Index);
 			CurrentProperty = ArrayProperty->Inner;
 		}
diff --git a/Engine/Source/Editor/PropertyEditor/Private/PropertyNode.cpp b/Engine/Source/Editor/PropertyEditor/Private/PropertyNode.cpp
--- a/Engine/Source/Editor/PropertyEditor/Private/PropertyNode.cpp
+++ b/Engine/Source/Editor/PropertyEditor/Private/PropertyNode.cpp
@@ -713,6 +713,20 @@ void FPropertyNode::MarkChildrenAsRebuilt()
  */
 EPropertyDataValidationResult FPropertyNode::EnsureDataIsValid()
 {
+	auto GetChildStructureNodes = [Self = this](TArray<TSharedPtr<const FPropertyNode>, TInlineAllocator<2>>& OutChildStructurePropertyNodes)
+	{
+		for (const TSharedPtr<FPropertyNode>& ChildNode : Self->ChildNodes)
+		{
+			if (const FComplexPropertyNode* ComplexNode = ChildNode->AsComplexNode())
+			{
+				if (ComplexNode->AsStructureNode())
+				{
+					OutChildStructurePropertyNodes.Emplace(ChildNode);
+				}
+			}
+		}
+	};
+
 	bool bValidateChildren = !HasNodeFlags(EPropertyNodeFlags::SkipChildValidation);
 	bool bValidateChildrenKeyNodes = false;		// by default, we don't check this, since it's just for Map properties
 
@@ -728,12 +742,12 @@ EPropertyDataValidationResult FPropertyNode::EnsureDataIsValid()
 	{
 		CachedReadAddresses.Reset();
 
-		//Figure out if an array mismatch can be ignored
+		// Figure out if an array mismatch can be ignored
 		bool bIgnoreAllMismatch = false;
-		//make sure that force depth-limited trees don't cause a refresh
+		// Make sure that force depth-limited trees don't cause a refresh
 		bIgnoreAllMismatch |= (MaxChildDepthAllowed==0);
 
-		//check my property
+		// Check my property
 		if (Property.IsValid())
 		{
 			FProperty* MyProperty = Property.Get();
@@ -741,25 +755,28 @@ EPropertyDataValidationResult FPropertyNode::EnsureDataIsValid()
 
 			if (!OwnerStruct || OwnerStruct->IsStructTrashed())
 			{
-				//verify that the property is not part of an invalid trash class, treat it as an invalid object if it is which will cause a refresh
+				// Verify that the property is not part of an invalid trash class, treat it as an invalid object if it is which will cause a refresh
 				return EPropertyDataValidationResult::ObjectInvalid;
 			}
 
-			//verify that the number of container children is correct
+			// Cached Child Structure nodes that are manually attached to this property node
+			TArray<TSharedPtr<const FPropertyNode>, TInlineAllocator<2>> CachedChildStructureNodes;
+			GetChildStructureNodes(CachedChildStructureNodes);
+
+			// Verify that the number of container children is correct
 			FArrayProperty* ArrayProperty = CastField<FArrayProperty>(MyProperty);
 			FSetProperty* SetProperty = CastField<FSetProperty>(MyProperty);
 			FMapProperty* MapProperty = CastField<FMapProperty>(MyProperty);
-			FStructProperty* StructProperty = CastField<FStructProperty>(MyProperty);
 			FOptionalProperty* OptionalProperty = CastField<FOptionalProperty>(MyProperty);
 
-			//default to unknown array length
-			int32 NumArrayChildren = -1;
-			//assume all arrays have the same length
-			bool bArraysHaveEqualNum = true;
-			//assume all arrays match the number of property window children
-			bool bArraysMatchChildNum = true;
+			// Default to unknown array length
+			int32 NumContainerChildren = -1;
+			// Assume all arrays have the same length
+			bool bContainersHaveEqualNum = true;
+			// Assume all arrays match the number of property window children
+			bool bContainersMatchChildNum = true;
 
-			bool bArrayHasNewItem = false;
+			bool bContainerHasNewItem = false;
 
 			FProperty* ContainerElementProperty = MyProperty;
 
@@ -784,33 +801,37 @@ EPropertyDataValidationResult FPropertyNode::EnsureDataIsValid()
 			}
 
 			bValidateChildren = DoesChildPropertyRequireValidation(ContainerElementProperty);
+			for (TSharedPtr<const FPropertyNode>& ChildStructureNode : CachedChildStructureNodes)
+			{
+				bValidateChildren |= DoesChildPropertyRequireValidation(ChildStructureNode->Property.Get());
+			}
 
-			//verify that the number of object children are the same too
+			// Verify that the number of object children are the same too
 			FObjectPropertyBase* ObjectProperty = CastField<FObjectPropertyBase>(MyProperty);
-			//check to see, if this an object property, whether the contents are NULL or not.
-			//This is the check to see if an object property was changed from NULL to non-NULL, or vice versa, from non-property window code.
+			// Check to see, if this an object property, whether the contents are NULL or not.
+			// This is the check to see if an object property was changed from NULL to non-NULL, or vice versa, from non-property window code.
 			bool bObjectPropertyNull = true;
 
-			//Edit inline properties can change underneath the window
+			// Edit inline properties can change underneath the window
 			bool bIgnoreChangingChildren = !(
 				HasNodeFlags(EPropertyNodeFlags::EditInlineNew) ||
 				HasNodeFlags(EPropertyNodeFlags::ShowInnerObjectProperties) ||
 			 	HasNodeFlags(EPropertyNodeFlags::DynamicInstance)
 			);
 
-			//ignore this node if the consistency check should happen for theW children
+			// Ignore this node if the consistency check should happen for theW children
 			bool bIgnoreStaticArray = (Property->ArrayDim > 1) && (ArrayIndex == -1);
 
-			//if this node can't possibly have children (or causes a circular reference loop) then ignore this as a object property
+			// If this node can't possibly have children (or causes a circular reference loop) then ignore this as a object property
 			if (bIgnoreChangingChildren || bIgnoreStaticArray || HasNodeFlags(EPropertyNodeFlags::NoChildrenDueToCircularReference))
 			{
-				//this will bypass object property consistency checks
+				// This will bypass object property consistency checks
 				ObjectProperty = nullptr;
 			}
 
 			FReadAddressList ReadAddresses;
 			const bool bSuccess = GetReadAddress( ReadAddresses );
-			//make sure we got the addresses correctly
+			// Make sure we got the addresses correctly
 			if (!bSuccess)
 			{
 				UE_LOG( LogPropertyNode, Verbose, TEXT("Object is invalid %s"), *Property->GetName() );
@@ -820,7 +841,7 @@ EPropertyDataValidationResult FPropertyNode::EnsureDataIsValid()
 			// If an object property with ShowInnerProperties changed object values out from under the property
 			bool bShowInnerObjectPropertiesObjectChanged = false;
 
-			//check for null, if we find one, there is a problem.
+			// Check for null, if we find one, there is a problem.
 			for (int32 Scan = 0; Scan < ReadAddresses.Num(); ++Scan)
 			{
 				uint8* Addr = ReadAddresses.GetAddress(Scan);
@@ -833,34 +854,34 @@ EPropertyDataValidationResult FPropertyNode::EnsureDataIsValid()
 
 				if( ArrayProperty && !bIgnoreAllMismatch)
 				{
-					//ensure that array structures have the proper number of children
+					// Ensure that array structures have the proper number of children
 					FScriptArrayHelper ArrayHelper(ArrayProperty, Addr);
 					int32 ArrayNum = ArrayHelper.Num();
-					//if first child
-					if (NumArrayChildren == -1)
+					// If first child
+					if (NumContainerChildren == -1)
 					{
-						NumArrayChildren = ArrayNum;
+						NumContainerChildren = ArrayNum;
 					}
-					bArrayHasNewItem = GetNumChildNodes() < ArrayNum;
-					//make sure multiple arrays match
-					bArraysHaveEqualNum = bArraysHaveEqualNum && (NumArrayChildren == ArrayNum);
-					//make sure the array matches the number of property node children
-					bArraysMatchChildNum = bArraysMatchChildNum && (GetNumChildNodes() == ArrayNum);
+					bContainerHasNewItem = GetNumChildNodes() < ArrayNum;
+					// Make sure multiple arrays match
+					bContainersHaveEqualNum = bContainersHaveEqualNum && (NumContainerChildren == ArrayNum);
+					// Make sure the array matches the number of property node children
+					bContainersMatchChildNum = bContainersMatchChildNum && (GetNumChildNodes() == ArrayNum);
 				}
 
 				if (SetProperty && !bIgnoreAllMismatch)
 				{
-					// like arrays, ensure that set structures have the proper number of children
+					// Like arrays, ensure that set structures have the proper number of children
 					int32 SetNum = FScriptSetHelper::Num(Addr);
 
-					if (NumArrayChildren == -1)
+					if (NumContainerChildren == -1)
 					{
-						NumArrayChildren = SetNum;
+						NumContainerChildren = SetNum;
 					}
 
-					bArrayHasNewItem = GetNumChildNodes() < SetNum;
-					bArraysHaveEqualNum = bArraysHaveEqualNum && (NumArrayChildren == SetNum);
-					bArraysMatchChildNum = bArraysMatchChildNum && (GetNumChildNodes() == SetNum);
+					bContainerHasNewItem = GetNumChildNodes() < SetNum;
+					bContainersHaveEqualNum = bContainersHaveEqualNum && (NumContainerChildren == SetNum);
+					bContainersMatchChildNum = bContainersMatchChildNum && (GetNumChildNodes() == SetNum);
 				}
 
 				if (MapProperty && !bIgnoreAllMismatch)
@@ -868,14 +889,14 @@ EPropertyDataValidationResult FPropertyNode::EnsureDataIsValid()
 					FScriptMapHelper MapHelper(MapProperty, Addr);
 					int32 MapNum = MapHelper.Num();
 
-					if (NumArrayChildren == -1)
+					if (NumContainerChildren == -1)
 					{
-						NumArrayChildren = MapNum;
+						NumContainerChildren = MapNum;
 					}
 
-					bArrayHasNewItem = GetNumChildNodes() < MapNum;
-					bArraysHaveEqualNum = bArraysHaveEqualNum && (NumArrayChildren == MapNum);
-					bArraysMatchChildNum = bArraysMatchChildNum && (GetNumChildNodes() == MapNum);
+					bContainerHasNewItem = GetNumChildNodes() < MapNum;
+					bContainersHaveEqualNum = bContainersHaveEqualNum && (NumContainerChildren == MapNum);
+					bContainersMatchChildNum = bContainersMatchChildNum && (GetNumChildNodes() == MapNum);
 				}
 
 				if (ObjectProperty && !bIgnoreAllMismatch)
@@ -920,16 +941,18 @@ EPropertyDataValidationResult FPropertyNode::EnsureDataIsValid()
 				{
 					uint32 bIsSet = OptionalProperty->IsSet(Addr);
 					// re-use `bArrayHasNewItem` as its logic is what we end up wanting exactly
-					bArrayHasNewItem = bIsSet != GetNumChildNodes();
+					bContainerHasNewItem = bIsSet != GetNumChildNodes();
 				}
 			}
 
-			//if all arrays match each other but they do NOT match the property structure, cause a rebuild
-			if (bArraysHaveEqualNum && !bArraysMatchChildNum)
+			// If all arrays match each other but they do NOT match the property structure, cause a rebuild
+			// Don't trigger a rebuild if we have any structure child nodes, because we can't retrieve that information from the actual memory
+			// todo: Find a way to support a mix of container elements and manually added structure as child nodes
+			if (CachedChildStructureNodes.IsEmpty() && bContainersHaveEqualNum && !bContainersMatchChildNum)
 			{
 				RebuildChildren();
 
-				if( bArrayHasNewItem && ChildNodes.Num() )
+				if( bContainerHasNewItem && ChildNodes.Num() )
 				{
 					TSharedPtr<FPropertyNode> LastChildNode = ChildNodes.Last();
 					// Don't expand huge children
@@ -943,7 +966,8 @@ EPropertyDataValidationResult FPropertyNode::EnsureDataIsValid()
 				return EPropertyDataValidationResult::ArraySizeChanged;
 			}
 
-			if (bShowInnerObjectPropertiesObjectChanged)
+			// Don't trigger a rebuild if we have any structure child nodes, because we can't retrieve that information from the actual memory
+			if (CachedChildStructureNodes.IsEmpty() && bShowInnerObjectPropertiesObjectChanged)
 			{
 				RebuildChildren();
 				return EPropertyDataValidationResult::EditInlineNewValueChanged;
