[PropertyEditor] Used a typedef for Array of Property Nodes

diff --git a/Engine/Source/Editor/PropertyEditor/Private/PropertyNode.cpp b/Engine/Source/Editor/PropertyEditor/Private/PropertyNode.cpp
--- a/Engine/Source/Editor/PropertyEditor/Private/PropertyNode.cpp
+++ b/Engine/Source/Editor/PropertyEditor/Private/PropertyNode.cpp
@@ -713,7 +713,8 @@ void FPropertyNode::MarkChildrenAsRebuilt()
  */
 EPropertyDataValidationResult FPropertyNode::EnsureDataIsValid()
 {
-	auto GetChildStructureNodes = [Self = this](TArray<TSharedPtr<const FPropertyNode>, TInlineAllocator<2>>& OutChildStructurePropertyNodes)
+	using FPropertyNodeInlineArray = TArray<TSharedPtr<const FPropertyNode>, TInlineAllocator<2>>;
+	auto GetChildStructureNodes = [Self = this](FPropertyNodeInlineArray& OutChildStructurePropertyNodes)
 	{
 		for (const TSharedPtr<FPropertyNode>& ChildNode : Self->ChildNodes)
 		{
@@ -760,7 +761,7 @@ EPropertyDataValidationResult FPropertyNode::EnsureDataIsValid()
 			}
 
 			// Cached Child Structure nodes that are manually attached to this property node
-			TArray<TSharedPtr<const FPropertyNode>, TInlineAllocator<2>> CachedChildStructureNodes;
+			FPropertyNodeInlineArray CachedChildStructureNodes;
 			GetChildStructureNodes(CachedChildStructureNodes);
 
 			// Verify that the number of container children is correct
