Use middle ellipsis for object selector text clipping so that it's more consistent with other windows displaying actor names.

diff --git a/Engine/Source/Editor/PropertyEditor/Private/UserInterface/PropertyEditor/SPropertyEditorAsset.cpp b/Engine/Source/Editor/PropertyEditor/Private/UserInterface/PropertyEditor/SPropertyEditorAsset.cpp
--- a/Engine/Source/Editor/PropertyEditor/Private/UserInterface/PropertyEditor/SPropertyEditorAsset.cpp
+++ b/Engine/Source/Editor/PropertyEditor/Private/UserInterface/PropertyEditor/SPropertyEditorAsset.cpp
@@ -44,6 +44,7 @@
 #include "PropertyEditorConstants.h"
 #include "PropertyEditorUtils.h"
 #include "Misc/EditorPathHelper.h"
+#include "Settings/EditorStyleSettings.h"
 
 #define LOCTEXT_NAMESPACE "PropertyEditor"
 
@@ -657,6 +658,11 @@ void SPropertyEditorAsset::Construct(const FArguments& InArgs, const TSharedPtr<
 				SNew(STextBlock)
 				.Font( FAppStyle::GetFontStyle( PropertyEditorConstants::PropertyFontStyle ) )
 				.Text(this,&SPropertyEditorAsset::OnGetAssetName)
+                .OverflowPolicy(
+                    GetDefault<UEditorStyleSettings>()->bEnableMiddleEllipsis
+                        ? TOptional(ETextOverflowPolicy::MiddleEllipsis)
+                        : NullOpt
+                )
 			]
 		];
 
