Auto-focus the currently selected actor in actor picker's dropdown, so it's much easier to select another actor which is nearby the current one (e.g. same LI, folder, etc).

diff --git a/Engine/Source/Editor/PropertyEditor/Private/UserInterface/PropertyEditor/SPropertyMenuActorPicker.cpp b/Engine/Source/Editor/PropertyEditor/Private/UserInterface/PropertyEditor/SPropertyMenuActorPicker.cpp
--- a/Engine/Source/Editor/PropertyEditor/Private/UserInterface/PropertyEditor/SPropertyMenuActorPicker.cpp
+++ b/Engine/Source/Editor/PropertyEditor/Private/UserInterface/PropertyEditor/SPropertyMenuActorPicker.cpp
@@ -97,13 +97,30 @@ void SPropertyMenuActorPicker::Construct( const FArguments& InArgs )
 
 		InitOptions.ColumnMap.Add(FSceneOutlinerBuiltInColumnTypes::Label(), FSceneOutlinerColumnInfo(ESceneOutlinerColumnVisibility::Visible, 0));
 		InitOptions.ColumnMap.Add(FSceneOutlinerBuiltInColumnTypes::ActorInfo(), FSceneOutlinerColumnInfo(ESceneOutlinerColumnVisibility::Visible, 10));
+
+		TSharedRef<ISceneOutliner> SceneOutliner = SceneOutlinerModule.CreateActorPicker(InitOptions, FOnActorPicked::CreateSP(this, &SPropertyMenuActorPicker::OnActorSelected), nullptr, !bAllowPickingLevelInstanceContent);
+
+		if (CurrentActor)
+		{
+			SceneOutliner->SetSelection([this](ISceneOutlinerTreeItem& Item)
+			{
+				if (const FActorTreeItem* ActorItem = Item.CastTo<FActorTreeItem>())
+				{
+					return ActorItem->Actor == CurrentActor;
+				}
+
+				return false;
+			});
+
+			SceneOutliner->FrameSelectedItems();
+		}
 		
 		MenuContent =
 			SNew(SBox)
 			.WidthOverride(static_cast<float>(PropertyEditorAssetConstants::SceneOutlinerWindowSize.X))
 			.HeightOverride(static_cast<float>(PropertyEditorAssetConstants::SceneOutlinerWindowSize.Y))
 			[
-				SceneOutlinerModule.CreateActorPicker(InitOptions, FOnActorPicked::CreateSP(this, &SPropertyMenuActorPicker::OnActorSelected), nullptr, !bAllowPickingLevelInstanceContent)
+				MoveTemp(SceneOutliner)
 			];
 
 		MenuBuilder.AddWidget(MenuContent.ToSharedRef(), FText::GetEmpty(), true);
