[UE] Add missing support for instanced struct value edit propagation and some fixes for its reset to default UX.

Notes:
- Updated SInstancedStructPicker::StructPicked() to route through the IPropertyHandle value set API when the property is outered to one or more object containers. This internally handles any default value propagation.
- Updated FPropertyNode::GetDiffersFromDefault() to use the archetype's value as default when an outer object node is non-NULL. Previously it was always using a default instanced struct value for comparison.
- Updated FPropertyNode::GetDefaultValueAsString() similarly to export the default value text from the archetype when an outer object node is non-NULL. Previously it was always exporting from a default struct value.

#Synced-CL 48038701

diff --git a/Engine/Source/Editor/PropertyEditor/Private/PropertyNode.cpp b/Engine/Source/Editor/PropertyEditor/Private/PropertyNode.cpp
--- a/Engine/Source/Editor/PropertyEditor/Private/PropertyNode.cpp
+++ b/Engine/Source/Editor/PropertyEditor/Private/PropertyNode.cpp
@@ -2689,18 +2689,34 @@ bool FPropertyNode::GetDiffersFromDefault()
 				{
 					continue;
 				}
-				
-				// Make an instance of the struct to be used as default value to test against.
-				if (DefaultStruct.GetStruct() != Struct)
+
+				const uint8* PropertyDefaultAddress;
+				const uint8* PropertyDefaultBaseAddress;
+
+				const UObject* TopLevelObject = (TopLevelObjectNode && Index < TopLevelObjectNode->GetNumObjects()) ? TopLevelObjectNode->GetUObject(Index) : nullptr;
+				if (TopLevelObject)
 				{
-					DefaultStruct.Initialize(Struct);
+					UObject* ObjectArchetype = FPropertyNode::GetArchetype(TopLevelObject);
+					check(ObjectArchetype);
+
+					PropertyDefaultAddress = GetValueAddress(StructNode->GetValueAddressFromObject(ObjectArchetype), bIsSparse, /*bIsStruct=*/true);
+					PropertyDefaultBaseAddress = BaseNode->GetValueBaseAddress(StructNode->GetValueBaseAddressFromObject(ObjectArchetype), bIsSparse, /*bIsStruct=*/true);
 				}
+				else
+				{
+					// Make an instance of the struct to be used as default value to test against.
+					if (DefaultStruct.GetStruct() != Struct)
+					{
+						DefaultStruct.Initialize(Struct);
+					}
 
-				check(DefaultStruct.IsValid());
+					check(DefaultStruct.IsValid());
+
+					PropertyDefaultAddress = GetValueAddress(DefaultStruct.GetStructMemory(), bIsSparse, /*bIsStruct=*/true);
+					PropertyDefaultBaseAddress = BaseNode->GetValueBaseAddress(DefaultStruct.GetStructMemory(), bIsSparse, /*bIsStruct=*/true);
+				}
 
 				const uint8* PropertyValueAddress = GetValueAddress(StructData->GetStructMemory(), bIsSparse, /*bIsStruct=*/true);
-				const uint8* PropertyDefaultAddress = GetValueAddress(DefaultStruct.GetStructMemory(), bIsSparse, /*bIsStruct=*/true);
-				const uint8* PropertyDefaultBaseAddress = BaseNode->GetValueBaseAddress(DefaultStruct.GetStructMemory(), bIsSparse, /*bIsStruct=*/true);
 
 				// If this is a container property, we must take special measures to use the base address of the property's value; for instance,
 				// the array property's PropertyDefaultBaseAddress points to an FScriptArray*, while PropertyDefaultAddress points to the 
@@ -2712,7 +2728,6 @@ bool FPropertyNode::GetDiffersFromDefault()
 					PropertyDefaultAddress = PropertyDefaultBaseAddress;
 				}
 
-				const UObject* TopLevelObject = (TopLevelObjectNode && Index < TopLevelObjectNode->GetNumObjects()) ? TopLevelObjectNode->GetUObject(Index) : nullptr;
 				bDiffersFromDefault = GetDiffersFromDefault(PropertyValueAddress, PropertyDefaultAddress, PropertyDefaultBaseAddress, Prop, TopLevelObject);
 			}
 		}
@@ -2835,6 +2850,7 @@ FString FPropertyNode::GetDefaultValueAsString(bool bUseDisplayName)
 		FStructOnScope DefaultStruct;
 		FString NodeDefaultValue;
 
+		const uint8* LastPropertyDefaultAddress = nullptr;
 		const FObjectPropertyNode* TopLevelObjectNode = StructNode->FindObjectItemParent();
 		for (int32 StructIndex = 0; StructIndex < Structs.Num(); ++StructIndex)
 		{
@@ -2849,25 +2865,45 @@ FString FPropertyNode::GetDefaultValueAsString(bool bUseDisplayName)
 				continue;
 			}
 
-			if (DefaultStruct.GetStruct() != Struct)
+			const uint8* PropertyDefaultAddress;
+			const uint8* PropertyDefaultBaseAddress;
+
+			const UObject* TopLevelObject = (TopLevelObjectNode && StructIndex < TopLevelObjectNode->GetNumObjects()) ? TopLevelObjectNode->GetUObject(StructIndex) : nullptr;
+			if (TopLevelObject)
+			{
+				UObject* ObjectArchetype = FPropertyNode::GetArchetype(TopLevelObject);
+				check(ObjectArchetype);
+
+				PropertyDefaultAddress = GetValueAddress(StructNode->GetValueAddressFromObject(ObjectArchetype), bIsSparse, /*bIsStruct=*/true);
+				PropertyDefaultBaseAddress = BaseNode->GetValueBaseAddress(StructNode->GetValueBaseAddressFromObject(ObjectArchetype), bIsSparse, /*bIsStruct=*/true);
+			}
+			else
 			{
 				// Make an instance of the struct to be used as default value.
-				DefaultStruct.Initialize(Struct);
+				if (DefaultStruct.GetStruct() != Struct)
+				{
+					DefaultStruct.Initialize(Struct);
+				}
+
 				check(DefaultStruct.IsValid());
 
-				const uint8* PropertyDefaultAddress = GetValueAddress(DefaultStruct.GetStructMemory(), bIsSparse, /*bIsStruct=*/true);
-				const uint8* PropertyDefaultBaseAddress = BaseNode->GetValueBaseAddress(DefaultStruct.GetStructMemory(), bIsSparse, /*bIsStruct=*/true);
+				PropertyDefaultAddress = GetValueAddress(DefaultStruct.GetStructMemory(), bIsSparse, /*bIsStruct=*/true);
+				PropertyDefaultBaseAddress = BaseNode->GetValueBaseAddress(DefaultStruct.GetStructMemory(), bIsSparse, /*bIsStruct=*/true);
+			}
 
-				// If this is a container property, we must take special measures to use the base address of the property's value; for instance,
-				// the array property's PropertyDefaultBaseAddress points to an FScriptArray*, while PropertyDefaultAddress points to the 
-				// FScriptArray's Data pointer.
-				if (bIsContainer)
-				{
-					PropertyDefaultAddress = PropertyDefaultBaseAddress;
-				}
+			// If this is a container property, we must take special measures to use the base address of the property's value; for instance,
+			// the array property's PropertyDefaultBaseAddress points to an FScriptArray*, while PropertyDefaultAddress points to the 
+			// FScriptArray's Data pointer.
+			if (bIsContainer)
+			{
+				PropertyDefaultAddress = PropertyDefaultBaseAddress;
+			}
 
-				const UObject* TopLevelObject = (TopLevelObjectNode && StructIndex < TopLevelObjectNode->GetNumObjects()) ? TopLevelObjectNode->GetUObject(StructIndex) : nullptr;
+			if (NodeDefaultValue.IsEmpty()
+				|| PropertyDefaultAddress != LastPropertyDefaultAddress)
+			{
 				NodeDefaultValue = GetDefaultValueAsString(PropertyDefaultAddress, Prop, Mode, TopLevelObject);
+				LastPropertyDefaultAddress = PropertyDefaultAddress;
 			}
 			
 			if (DefaultValue.IsEmpty())
diff --git a/Engine/Source/Editor/StructUtilsEditor/Private/SInstancedStructPicker.cpp b/Engine/Source/Editor/StructUtilsEditor/Private/SInstancedStructPicker.cpp
--- a/Engine/Source/Editor/StructUtilsEditor/Private/SInstancedStructPicker.cpp
+++ b/Engine/Source/Editor/StructUtilsEditor/Private/SInstancedStructPicker.cpp
@@ -390,19 +390,51 @@ void SInstancedStructPicker::StructPicked(const UScriptStruct* InStruct)
 	{
 		FScopedTransaction Transaction(LOCTEXT("OnStructPicked", "Set Struct"));
 
-		StructProperty->NotifyPreChange();
+		TArray<UObject*> OuterObjects;
+		OuterObjects.Reserve(StructProperty->GetNumOuterObjects());
+		StructProperty->GetOuterObjects(OuterObjects);
 
-		StructProperty->EnumerateRawData([InStruct](void* RawData, const int32 /*DataIndex*/, const int32 /*NumDatas*/)
+		if (OuterObjects.IsEmpty())
 		{
-			if (FInstancedStruct* InstancedStruct = static_cast<FInstancedStruct*>(RawData))
+			StructProperty->NotifyPreChange();
+
+			StructProperty->EnumerateRawData([InStruct](void* RawData, const int32 /*DataIndex*/, const int32 /*NumDatas*/)
+			{
+				if (FInstancedStruct* InstancedStruct = static_cast<FInstancedStruct*>(RawData))
+				{
+					InstancedStruct->InitializeAs(InStruct);
+				}
+				return true;
+			});
+
+			StructProperty->NotifyPostChange(EPropertyChangeType::ValueSet);
+			StructProperty->NotifyFinishedChangingProperties();
+		}
+		else
+		{
+			FInstancedStruct InstancedStruct;
+			InstancedStruct.InitializeAs(InStruct);
+
+			TArray<FString> NewValues;
+			NewValues.Reserve(OuterObjects.Num());
+			for (UObject* OuterObject : OuterObjects)
 			{
-				InstancedStruct->InitializeAs(InStruct);
+				if (!OuterObject)
+				{
+					continue;
+				}
+
+				InstancedStruct.ExportTextItem(
+					NewValues.Emplace_GetRef(),
+					FInstancedStruct(),
+					OuterObject,
+					PPF_None,
+					OuterObject->GetOutermost());
 			}
-			return true;
-		});
 
-		StructProperty->NotifyPostChange(EPropertyChangeType::ValueSet);
-		StructProperty->NotifyFinishedChangingProperties();
+			// This will internally handle propagation for values within default/archetype objects
+			StructProperty->SetPerObjectValues(NewValues, EPropertyValueSetFlags::NotTransactable);
+		}
 
 		// After the type has changed, let's expand, so that the user can edit the newly appeared child properties
 		StructProperty->SetExpanded(true);
