Add support for filtering shown categories when creating a details view.
Add support for hiding headers when creating a details view.

#tests editor
#Synced-CL 48254188

diff --git a/Engine/Source/Editor/PropertyEditor/Private/DetailCategoryBuilderImpl.cpp b/Engine/Source/Editor/PropertyEditor/Private/DetailCategoryBuilderImpl.cpp
--- a/Engine/Source/Editor/PropertyEditor/Private/DetailCategoryBuilderImpl.cpp
+++ b/Engine/Source/Editor/PropertyEditor/Private/DetailCategoryBuilderImpl.cpp
@@ -935,7 +935,8 @@ TSharedRef<ITableRow> FDetailCategoryImpl::GenerateWidgetForTableView(const TSha
 		.InnerCategory(ParentLayout.IsValid() ? ParentLayout->IsLayoutForExternalRoot() : false)
 		.DisplayName(GetDisplayName())
 		.HeaderContent(HeaderContent)
-		.WholeRowHeaderContent(bHeaderContentWholeRowContent);
+		.WholeRowHeaderContent(bHeaderContentWholeRowContent)
+		.ShowHeader(ParentLayout.IsValid() ? ParentLayout->GetCurrentFilter().bShowHeaders : true);
 
 	if (ToolTip.IsSet() && !ToolTip->IsEmpty())
 	{
diff --git a/Engine/Source/Editor/PropertyEditor/Private/DetailFilter.h b/Engine/Source/Editor/PropertyEditor/Private/DetailFilter.h
--- a/Engine/Source/Editor/PropertyEditor/Private/DetailFilter.h
+++ b/Engine/Source/Editor/PropertyEditor/Private/DetailFilter.h
@@ -17,19 +17,22 @@ struct FDetailFilter
 		, bShowFavoritesCategory(false)
 		, bShowOnlyAllowed(false)
 		, bShowLooseProperties(false)
+		, bShowHeaders(true)
 	{}
 
 	bool IsEmptyFilter() const
 	{
 		return FilterStrings.Num() == 0
 			&& VisibleSections.Num() == 0
+			&& VisibleCategories.Num() == 0
 			&& bShowOnlyModified == false
 			&& bShowAllAdvanced == false
 			&& bShowOnlyAllowed == false
 			&& bShowAllChildrenIfCategoryMatches == false
 			&& bShowOnlyKeyable == false
 			&& bShowOnlyAnimated == false
-			&& bShowLooseProperties == false;
+			&& bShowLooseProperties == false
+			&& bShowHeaders == true;
 	}
 
 	/** Any user search terms that items must match */
@@ -50,10 +53,14 @@ struct FDetailFilter
 	bool bShowOnlyAllowed;
 	/* If true, will also show loose properties */
 	bool bShowLooseProperties;
+	/* If true, category headers */
+	bool bShowHeaders;
 	/** The set of allowed properties to show. */
 	TSet<FPropertyPath> PropertyAllowList;
 	/** The set of selected sections to show. If empty, all sections are shown.*/
 	TSet<FName> VisibleSections;
+	/** The set of selected categories to show. If empty, all categories are shown.*/
+	TSet<FName> VisibleCategories;
 
 	/** Delegate that determines whether a property should be forced hidden - evaluated on panel refresh */
 	DECLARE_DELEGATE_RetVal_OneParam(bool, FShouldForceHideProperty, const TSharedRef<class FPropertyNode>&);
diff --git a/Engine/Source/Editor/PropertyEditor/Private/DetailLayoutBuilderImpl.cpp b/Engine/Source/Editor/PropertyEditor/Private/DetailLayoutBuilderImpl.cpp
--- a/Engine/Source/Editor/PropertyEditor/Private/DetailLayoutBuilderImpl.cpp
+++ b/Engine/Source/Editor/PropertyEditor/Private/DetailLayoutBuilderImpl.cpp
@@ -615,10 +615,13 @@ void FDetailLayoutBuilderImpl::FilterDetailLayout( const FDetailFilter& InFilter
 
 		if(RootTreeNode->GetVisibility() == ENodeVisibility::Visible)
 		{
-			FilteredRootTreeNodes.Add(RootTreeNode);
-			if (TSharedPtr<IDetailsViewPrivate> DetailsViewPinned = DetailsView.Pin())
+			if (InFilter.VisibleCategories.IsEmpty() || InFilter.VisibleCategories.Contains(RootTreeNode->GetNodeName()))
 			{
-				DetailsViewPinned->RequestItemExpanded(RootTreeNode, RootTreeNode->ShouldBeExpanded());
+				FilteredRootTreeNodes.Add(RootTreeNode);
+				if (TSharedPtr<IDetailsViewPrivate> DetailsViewPinned = DetailsView.Pin())
+				{
+					DetailsViewPinned->RequestItemExpanded(RootTreeNode, RootTreeNode->ShouldBeExpanded());
+				}
 			}
 		}
 	}
diff --git a/Engine/Source/Editor/PropertyEditor/Private/SDetailCategoryTableRow.cpp b/Engine/Source/Editor/PropertyEditor/Private/SDetailCategoryTableRow.cpp
--- a/Engine/Source/Editor/PropertyEditor/Private/SDetailCategoryTableRow.cpp
+++ b/Engine/Source/Editor/PropertyEditor/Private/SDetailCategoryTableRow.cpp
@@ -28,6 +28,7 @@ void SDetailCategoryTableRow::Construct(const FArguments& InArgs, TSharedRef<FDe
 	DisplayName = InArgs._DisplayName;
 	bIsInnerCategory = InArgs._InnerCategory;
 	bShowBorder = InArgs._ShowBorder;
+	bShowHeader = InArgs._ShowHeader;
 	bIsEmpty = InArgs._IsEmpty;
 
 	TSharedPtr<IDetailsViewPrivate> DetailsView = InOwnerTreeNode->GetDetailsViewSharedPtr();
@@ -171,7 +172,8 @@ void SDetailCategoryTableRow::Construct(const FArguments& InArgs, TSharedRef<FDe
 				SNew(SBorder)
 				.BorderImage(this, &SDetailCategoryTableRow::GetBackgroundImage)
 				.BorderBackgroundColor(this, &SDetailCategoryTableRow::GetInnerBackgroundColor)
-					.Padding(0.0f)
+				.Padding(0.0f)
+				.Visibility(bShowHeader ? EVisibility::Visible : EVisibility::Collapsed)
 				[
 					SNew(SBox)
 					.MinDesiredHeight(PropertyEditorConstants::PropertyRowHeight)
diff --git a/Engine/Source/Editor/PropertyEditor/Private/SDetailCategoryTableRow.h b/Engine/Source/Editor/PropertyEditor/Private/SDetailCategoryTableRow.h
--- a/Engine/Source/Editor/PropertyEditor/Private/SDetailCategoryTableRow.h
+++ b/Engine/Source/Editor/PropertyEditor/Private/SDetailCategoryTableRow.h
@@ -21,6 +21,7 @@ public:
 		: _InnerCategory(false)
 		, _WholeRowHeaderContent(false)
 		, _ShowBorder(true)
+		, _ShowHeader(true)
 		, _IsEmpty(false)
 		, _Objects(nullptr)
 	{}
@@ -29,6 +30,7 @@ public:
 		SLATE_ARGUMENT(TSharedPtr<SWidget>, HeaderContent)
 		SLATE_ARGUMENT(bool, WholeRowHeaderContent)
 		SLATE_ARGUMENT(bool, ShowBorder)
+		SLATE_ARGUMENT(bool, ShowHeader)
 		
 		/** If true, this Category should have no UProperty data associated with it, and will be shown as an 
 		* empty stub with no expansion arrow */
@@ -87,6 +89,7 @@ private:
 
 	bool bIsInnerCategory = false;
 	bool bShowBorder = false;
+	bool bShowHeader = true;
 
 	/**
 	* If true, this Category should have no UProperty data associated with it, and will be shown as an empty stub
diff --git a/Engine/Source/Editor/PropertyEditor/Private/SDetailsView.cpp b/Engine/Source/Editor/PropertyEditor/Private/SDetailsView.cpp
--- a/Engine/Source/Editor/PropertyEditor/Private/SDetailsView.cpp
+++ b/Engine/Source/Editor/PropertyEditor/Private/SDetailsView.cpp
@@ -86,12 +86,17 @@ void SDetailsView::Construct(const FArguments& InArgs, const FDetailsViewArgs& I
 	}
 
 	CurrentFilter.bShowLooseProperties = InDetailsViewArgs.bShowLooseProperties;
+	CurrentFilter.bShowHeaders = InDetailsViewArgs.bShowHeaders;
 	CurrentFilter.ShouldForceHideProperty = InDetailsViewArgs.ShouldForceHideProperty;
 
 	if (DetailsViewArgs.InitialSections.Num() > 0)
 	{
 		CurrentFilter.VisibleSections = DetailsViewArgs.InitialSections;
 	}
+	if (DetailsViewArgs.InitialCategories.Num() > 0)
+	{
+		CurrentFilter.VisibleCategories = DetailsViewArgs.InitialCategories;
+	}
 
 	ColumnSizeData.SetValueColumnWidth(DetailsViewArgs.ColumnWidth);
 	ColumnSizeData.SetRightColumnMinWidth(DetailsViewArgs.RightColumnMinWidth);
diff --git a/Engine/Source/Editor/PropertyEditor/Public/DetailsViewArgs.h b/Engine/Source/Editor/PropertyEditor/Public/DetailsViewArgs.h
--- a/Engine/Source/Editor/PropertyEditor/Public/DetailsViewArgs.h
+++ b/Engine/Source/Editor/PropertyEditor/Public/DetailsViewArgs.h
@@ -88,6 +88,8 @@ struct FDetailsViewArgs
 	bool bShowPropertyMatrixButton : 1;
 	/** Allow options to be changed */
 	bool bShowOptions : 1;
+	/** Allow headers to be hidden */
+	bool bShowHeaders : 1;
 	/** True if you want to show the object label */
 	bool bShowObjectLabel : 1;
 	UE_DEPRECATED(5.0, "bShowActorLabel has been renamed bShowObjectLabel")
@@ -131,6 +133,9 @@ struct FDetailsViewArgs
 	/** Initial sections to be shown enabled in the filter */
 	TSet<FName> InitialSections;
 
+	/** Initial categories to be shown enabled in the filter */
+	TSet<FName> InitialCategories;
+
 public:
 
 	FDetailsViewArgs()
@@ -149,6 +154,7 @@ public:
 		, bSearchInitialKeyFocus(false)
 		, bShowPropertyMatrixButton(true)
 		, bShowOptions(true)
+		, bShowHeaders(true)
 		, bShowObjectLabel(true)
 		, bShowModifiedPropertiesOption(true)
 		, bShowDifferingPropertiesOption(false)
