Add support for Arrays, Sets and Structs to Audio Properties Details Injector
#Synced-CL 48495066

diff --git a/Engine/Source/Editor/PropertyEditor/Public/PropertyCustomizationHelpers.h b/Engine/Source/Editor/PropertyEditor/Public/PropertyCustomizationHelpers.h
--- a/Engine/Source/Editor/PropertyEditor/Public/PropertyCustomizationHelpers.h
+++ b/Engine/Source/Editor/PropertyEditor/Public/PropertyCustomizationHelpers.h
@@ -641,12 +641,16 @@ public:
 protected:
 
 	virtual void SetOnRebuildChildren( FSimpleDelegate InOnRebuildChildren  ) override { OnRebuildChildren = InOnRebuildChildren; } 
-
+	
 	void OnNumChildrenChanged()
 	{
 		OnRebuildChildren.ExecuteIfBound();
 	}
 
+	bool GetShouldGenerateHeader() const {return bGenerateHeader;}
+	bool GetShouldDisplayResetToDefault() const {return bDisplayResetToDefault;}
+	bool GetShouldDisplayElementNum() const {return bDisplayElementNum;}
+
 private:
 
 	FText DisplayName;
