NavigationData: make SupportedAreas visible in the details panel
Fix TitleProperty meta tag not working on VisibleAnywhere properties

#Synced-CL 48908617

diff --git a/Engine/Source/Editor/PropertyEditor/Private/UserInterface/PropertyEditor/SPropertyEditorArrayItem.cpp b/Engine/Source/Editor/PropertyEditor/Private/UserInterface/PropertyEditor/SPropertyEditorArrayItem.cpp
--- a/Engine/Source/Editor/PropertyEditor/Private/UserInterface/PropertyEditor/SPropertyEditorArrayItem.cpp
+++ b/Engine/Source/Editor/PropertyEditor/Private/UserInterface/PropertyEditor/SPropertyEditorArrayItem.cpp
@@ -129,14 +129,12 @@ bool SPropertyEditorArrayItem::Supports( const TSharedRef< class FPropertyEditor
 
 	if (!CastField<const FClassProperty>(Property) && PropertyNode->HasNodeFlags(EPropertyNodeFlags::SingleSelectOnly))
 	{
-		if (Property->GetOwner<FArrayProperty>() &&
-			!(Property->GetOwner<FArrayProperty>()->PropertyFlags & CPF_EditConst))
+		if (Property->GetOwner<FArrayProperty>())
 		{
 			return true;
 		}
 
-		if (Property->GetOwner<FMapProperty>() &&
-			!(Property->GetOwner<FMapProperty>()->PropertyFlags & CPF_EditConst))
+		if (Property->GetOwner<FMapProperty>())
 		{
 			return true;
 		}
diff --git a/Engine/Source/Runtime/NavigationSystem/Public/NavigationData.h b/Engine/Source/Runtime/NavigationSystem/Public/NavigationData.h
--- a/Engine/Source/Runtime/NavigationSystem/Public/NavigationData.h
+++ b/Engine/Source/Runtime/NavigationSystem/Public/NavigationData.h
@@ -36,10 +36,10 @@ struct FSupportedAreaData
 {
 	GENERATED_USTRUCT_BODY()
 
-	UPROPERTY()
+	UPROPERTY(VisibleAnywhere, Category=Areas)
 	FString AreaClassName;
 
-	UPROPERTY()
+	UPROPERTY(VisibleAnywhere, Category=Areas)
 	int32 AreaID;
 
 	UPROPERTY(transient)
@@ -1095,7 +1095,7 @@ protected:
 	TMap<UClass*,FSharedConstNavQueryFilter > QueryFilters;
 
 	/** serialized area class - ID mapping */
-	UPROPERTY()
+	UPROPERTY(VisibleAnywhere, Category=Areas, meta=(TitleProperty="{AreaID}: {AreaClassName}"))
 	TArray<FSupportedAreaData> SupportedAreas;
 
 	/** mapping for SupportedAreas */
