[PropertyEditor] Add an option to add a color bar on every row in the category.
#Synced-CL 48720034

diff --git a/Engine/Source/Editor/PropertyEditor/Private/DetailCategoryBuilderImpl.cpp b/Engine/Source/Editor/PropertyEditor/Private/DetailCategoryBuilderImpl.cpp
--- a/Engine/Source/Editor/PropertyEditor/Private/DetailCategoryBuilderImpl.cpp
+++ b/Engine/Source/Editor/PropertyEditor/Private/DetailCategoryBuilderImpl.cpp
@@ -689,6 +689,21 @@ void FDetailCategoryImpl::SetIsEmpty(bool bInIsEmpty)
 {
 	bIsEmpty = bInIsEmpty;
 }
+
+TOptional<FLinearColor> FDetailCategoryImpl::GetColor() const
+{
+	return Color;
+}
+
+void FDetailCategoryImpl::SetColor(FLinearColor InColor)
+{
+	Color = InColor;
+}
+
+void FDetailCategoryImpl::ClearColor()
+{
+	Color = TOptional<FLinearColor>();
+}
  
 bool FDetailCategoryImpl::ShouldAdvancedBeExpanded() const
 {
diff --git a/Engine/Source/Editor/PropertyEditor/Private/DetailCategoryBuilderImpl.h b/Engine/Source/Editor/PropertyEditor/Private/DetailCategoryBuilderImpl.h
--- a/Engine/Source/Editor/PropertyEditor/Private/DetailCategoryBuilderImpl.h
+++ b/Engine/Source/Editor/PropertyEditor/Private/DetailCategoryBuilderImpl.h
@@ -331,9 +331,9 @@ public:
 	FDetailLayoutCustomization* GetDefaultCustomization(TSharedRef<FPropertyNode> PropertyNode);
 	
 	/**
-    * If true, this Category should have no UProperty data associated with it, and will be shown as an empty stub
-    * with no expansion arrow
-    */
+	 * If true, this Category should have no UProperty data associated with it, and will be shown as an empty stub
+	 * with no expansion arrow
+	 */
 	virtual bool IsEmpty() const override;
 	
 	/**
@@ -342,6 +342,13 @@ public:
 	 */
 	virtual void SetIsEmpty(bool bInIsEmpty) override;
 
+	virtual TOptional<FLinearColor> GetColor() const override;
+
+	virtual void SetColor(FLinearColor InColor) override;
+
+	virtual void ClearColor() override;
+
+
 private:
 	virtual void OnItemExpansionChanged(bool bIsExpanded, bool bShouldSaveState) override;
 
@@ -452,6 +459,8 @@ private:
 	FName CategoryName;
 	/** The sort order of this category (amongst all categories) */
 	int32 SortOrder;
+	/** The color that represents the category. */
+	TOptional<FLinearColor> Color;
 	/** Whether or not to restore the expansion state between sessions */
 	bool bRestoreExpansionState : 1;
 	/** Whether or not the category should be initially collapsed */
diff --git a/Engine/Source/Editor/PropertyEditor/Private/SDetailSingleItemRow.cpp b/Engine/Source/Editor/PropertyEditor/Private/SDetailSingleItemRow.cpp
--- a/Engine/Source/Editor/PropertyEditor/Private/SDetailSingleItemRow.cpp
+++ b/Engine/Source/Editor/PropertyEditor/Private/SDetailSingleItemRow.cpp
@@ -26,6 +26,7 @@
 #include "Styling/StyleColors.h"
 #include "UObject/Field.h"
 #include "UserInterface/PropertyEditor/PropertyEditorConstants.h"
+#include "Widgets/Colors/SColorBlock.h"
 #include "Widgets/Images/SImage.h"
 #include "Widgets/Input/SButton.h"
 #include "Widgets/Input/SComboButton.h"
@@ -825,26 +826,24 @@ void SDetailSingleItemRow::Construct( const FArguments& InArgs, FDetailLayoutCus
 	}
 	else
 	{
-		this->ChildSlot
-		[
-			SNew( SBorder )
+		TSharedRef<SWidget> ContentBorder = SNew(SBorder)
 			.BorderImage(FAppStyle::Get().GetBrush("DetailsView.GridLine"))
-			.Padding(FMargin(0.0f,0.0f,0.0f,1.0f))
+			.Padding(FMargin(0.0f, 0.0f, 0.0f, 1.0f))
 			.Clipping(EWidgetClipping::ClipToBounds)
 			[
 				SNew(SBox)
 				.MinDesiredHeight(PropertyEditorConstants::PropertyRowHeight)
 				[
-					SNew( SHorizontalBox )
+					SNew(SHorizontalBox)
 					+ SHorizontalBox::Slot()
 					.HAlign(HAlign_Fill)
 					.VAlign(VAlign_Fill)
 					[
-						SNew( SBorder )
+						SNew(SBorder)
 						.BorderImage(FAppStyle::Get().GetBrush("DetailsView.Highlight"))
 						.Padding_Lambda(GetHighlightBorderPadding)
 						[
-							SNew( SBorder )
+							SNew(SBorder)
 							.BorderImage(FAppStyle::Get().GetBrush("DetailsView.CategoryMiddle"))
 							.BorderBackgroundColor(this, &SDetailSingleItemRow::GetOuterBackgroundColor)
 							.Padding(0.0f)
@@ -854,8 +853,35 @@ void SDetailSingleItemRow::Construct( const FArguments& InArgs, FDetailLayoutCus
 						]
 					]
 				]
-			]
-		];
+			];
+
+		if (Category.IsValid() && Category->GetColor().IsSet())
+		{
+			this->ChildSlot
+			[
+				SNew(SHorizontalBox)
+				+ SHorizontalBox::Slot()
+				.Padding(0.f)
+				.MaxWidth(2.f)
+				[
+					SNew(SColorBlock)
+					.Color(Category->GetColor().GetValue())
+				]
+				+ SHorizontalBox::Slot()
+				.Padding(0.f)
+				.FillContentWidth(1.0f)
+				[
+					ContentBorder
+				]
+			];
+		}
+		else
+		{
+			this->ChildSlot
+			[
+				ContentBorder
+			];
+		}
 	}
 
 	STableRow< TSharedPtr< FDetailTreeNode > >::ConstructInternal(
diff --git a/Engine/Source/Editor/PropertyEditor/Public/DetailCategoryBuilder.h b/Engine/Source/Editor/PropertyEditor/Public/DetailCategoryBuilder.h
--- a/Engine/Source/Editor/PropertyEditor/Public/DetailCategoryBuilder.h
+++ b/Engine/Source/Editor/PropertyEditor/Public/DetailCategoryBuilder.h
@@ -243,4 +243,19 @@ public:
 	 * as an empty stub with no expansion arrow
 	 */
 	virtual void SetIsEmpty(bool bInIsEmpty) { }
+
+	/**
+	 * If valid, the color that represented the category. It can be used to visually identify the category quickly.
+	 */
+	virtual TOptional<FLinearColor> GetColor() const = 0;
+
+	/**
+	 * Set the color that represents the category.
+	 */
+	virtual void SetColor(FLinearColor InColor) = 0;
+
+	/**
+ 	 * Clear the color that represents the category.
+ 	 */
+	virtual void ClearColor() = 0;
 };
