Fix array did not expand its children on item add if its property editor was added to a parent property inside a customization.

#rnx
#Synced-CL 49037932

diff --git a/Engine/Source/Editor/PropertyEditor/Private/Presentation/PropertyEditor/PropertyEditor.cpp b/Engine/Source/Editor/PropertyEditor/Private/Presentation/PropertyEditor/PropertyEditor.cpp
--- a/Engine/Source/Editor/PropertyEditor/Private/Presentation/PropertyEditor/PropertyEditor.cpp
+++ b/Engine/Source/Editor/PropertyEditor/Private/Presentation/PropertyEditor/PropertyEditor.cpp
@@ -249,7 +249,18 @@ void FPropertyEditor::OnAddItem()
 	}
 
 	// Expand containers when an item is added to them
-	PropertyNode->SetNodeFlags(EPropertyNodeFlags::Expanded, true);
+	FPropertyNode* ParentNode = &PropertyNode.Get();
+	while (ParentNode)
+	{
+		// We expand all parents. Normally, they'd all be expanded (otherwise the property would not be visible and the
+		// user would not be able to add a new item), but with customizations it's possible to have something like:
+		//  > [Some Label Widget]    | 3 Array elements (+)
+		// where the row belongs to a struct property, while the label widget and the array are (for example) two child
+		// properties of said struct. In this case, without expanding parents, adding an array item would try to expand
+		// the property node which does not have its own row, while the row itself would remain collapsed.
+		ParentNode->SetNodeFlags(EPropertyNodeFlags::Expanded, true);
+		ParentNode = ParentNode->GetParentNode();
+	}
 
 	if (Result.GetAccessResult() == FPropertyAccess::Success && UE::Private::CVarExpandAddedItem->GetBool())
 	{
