EditorUsability : PropertyEditor
- Show Favorites/Modified sections in details panel filter bar, use details settings and check "Show Helper Sections"
- Add/Remove category items from favorites using context menu on a category
- Expand/Collapse details panel root categories using shortcuts : CTRL+SHIFT+X and CTRL+SHIFT+C, menu entries and commands actions are now using same code path
- Fixed SHIFT+Click on expansion arrow do not recursively expand advanced category
- Grouping properties based on their original category in favorites and removing Advanced

#Synced-CL 49068058

diff --git a/Engine/Source/Editor/LevelEditor/Private/SActorDetails.cpp b/Engine/Source/Editor/LevelEditor/Private/SActorDetails.cpp
--- a/Engine/Source/Editor/LevelEditor/Private/SActorDetails.cpp
+++ b/Engine/Source/Editor/LevelEditor/Private/SActorDetails.cpp
@@ -1375,4 +1375,19 @@ void SActorDetails::OnObjectPropertyChanged(UObject* ObjectBeingModified, FPrope
 	}
 }
 
+FReply SActorDetails::OnKeyDown(const FGeometry& InGeometry, const FKeyEvent& InKeyEvent)
+{
+	if (DetailsView.IsValid())
+    {
+		// This leaves a chance to events coming from the whole Details tab that do not come directly from Detail view to be handled
+		FReply Reply = DetailsView->OnKeyDown(InGeometry, InKeyEvent);
+		if (Reply.IsEventHandled())
+		{
+			return Reply;
+		}
+    }
+
+	return SCompoundWidget::OnKeyDown(InGeometry, InKeyEvent);
+}
+
 #undef LOCTEXT_NAMESPACE
diff --git a/Engine/Source/Editor/LevelEditor/Private/SActorDetails.h b/Engine/Source/Editor/LevelEditor/Private/SActorDetails.h
--- a/Engine/Source/Editor/LevelEditor/Private/SActorDetails.h
+++ b/Engine/Source/Editor/LevelEditor/Private/SActorDetails.h
@@ -48,7 +48,7 @@ public:
 		TSharedPtr<FUICommandList> InCommandList,
 		TSharedPtr<FTabManager> InTabManager,
 		TWeakPtr<FActorDetailsSCSEditorUICustomization> InLevelEditorSubobjectEditorCustomization);
-	~SActorDetails();
+	virtual ~SActorDetails() override;
 
 	/**
 	 * Return true if this details panel is observing the given selection set.
@@ -80,6 +80,11 @@ public:
 	 */
 	void SetActorDetailsRootCustomization(TSharedPtr<FDetailsViewObjectFilter> InActorDetailsObjectFilter, TSharedPtr<class IDetailRootObjectCustomization> ActorDetailsRootCustomization);
 
+protected:
+	//~ Begin SWidget
+	virtual FReply OnKeyDown(const FGeometry& InGeometry, const FKeyEvent& InKeyEvent) override;
+	//~ End SWidget
+
 private:
 	/**
 	 * A customization that the ActorDetails widget manages that maintains the visibility of the
@@ -120,7 +125,6 @@ private:
 	void OnObjectsReinstanced(const TMap<UObject*, UObject*>& InReinstancedObjects);
 	void OnObjectPropertyChanged(UObject* ObjectBeingModified, FPropertyChangedEvent& PropertyChangedEvent);
 
-private:
 	TSharedPtr<SSplitter> DetailsSplitter;
 	TSharedPtr<class IDetailsView> DetailsView;
 	TSharedPtr<SBox> ComponentsBox;
diff --git a/Engine/Source/Editor/PropertyEditor/Private/DetailAdvancedDropdownNode.cpp b/Engine/Source/Editor/PropertyEditor/Private/DetailAdvancedDropdownNode.cpp
--- a/Engine/Source/Editor/PropertyEditor/Private/DetailAdvancedDropdownNode.cpp
+++ b/Engine/Source/Editor/PropertyEditor/Private/DetailAdvancedDropdownNode.cpp
@@ -214,6 +214,15 @@ bool FAdvancedDropdownNode::GenerateStandaloneWidget(FDetailWidgetRow& OutRow) c
 	return false;
 }
 
+void FAdvancedDropdownNode::OnItemExpansionChanged(bool bIsExpanded, bool bShouldSaveState)
+{
+	// If state does not match, request expansion update to category
+	if (bIsExpanded != IsExpanded.Get())
+	{
+		OnAdvancedDropDownClicked();
+	}
+}
+
 FReply FAdvancedDropdownNode::OnAdvancedDropDownClicked()
 {
 	ParentCategory.OnAdvancedDropdownClicked();
diff --git a/Engine/Source/Editor/PropertyEditor/Private/DetailAdvancedDropdownNode.h b/Engine/Source/Editor/PropertyEditor/Private/DetailAdvancedDropdownNode.h
--- a/Engine/Source/Editor/PropertyEditor/Private/DetailAdvancedDropdownNode.h
+++ b/Engine/Source/Editor/PropertyEditor/Private/DetailAdvancedDropdownNode.h
@@ -2,7 +2,6 @@
 
 #pragma once
 
-#include "CoreMinimal.h"
 #include "Misc/Attribute.h"
 #include "Input/Reply.h"
 #include "IPropertyUtilities.h"
@@ -25,27 +24,26 @@ public:
 	}
 
 private:
-	/** IDetailTreeNode Interface */
+	//~ Begin FDetailTreeNode
 	virtual TSharedPtr<IDetailsView> GetNodeDetailsViewSharedPtr() const override { return ParentCategory.GetNodeDetailsViewSharedPtr(); }
 	virtual TSharedPtr<IDetailsViewPrivate> GetDetailsViewSharedPtr() const override { return ParentCategory.GetDetailsViewSharedPtr(); }
 	virtual TSharedRef< ITableRow > GenerateWidgetForTableView( const TSharedRef<STableViewBase>& OwnerTable, bool bAllowFavoriteSystem) override;
 	virtual bool GenerateStandaloneWidget(FDetailWidgetRow& OutRow) const override;
 	virtual void GetChildren(FDetailNodeList& OutChildren, const bool& bInIgnoreVisibility) override {}
-	virtual void OnItemExpansionChanged( bool bIsExpanded, bool bShouldSaveState) override {}
+	virtual void OnItemExpansionChanged( bool bIsExpanded, bool bShouldSaveState) override;
 	virtual bool ShouldBeExpanded() const override { return false; }
 	virtual ENodeVisibility GetVisibility() const override;
 	virtual void FilterNode( const FDetailFilter& InFilter ) override {}
 	virtual void Tick( float DeltaTime ) override {}
 	virtual bool ShouldShowOnlyChildren() const override { return false; }
 	virtual FName GetNodeName() const override { return NAME_None; }
-
 	virtual EDetailNodeType GetNodeType() const override { return EDetailNodeType::Advanced; }
 	virtual TSharedPtr<IPropertyHandle> CreatePropertyHandle() const override { return nullptr; }
+	//~ End FDetailTreeNode
 
 	/** Called when the advanced drop down arrow is clicked */
 	FReply OnAdvancedDropDownClicked();
 
-private:
 	FDetailCategoryImpl& ParentCategory;
 	TAttribute<bool> IsEnabled;
 	TAttribute<bool> IsExpanded;
diff --git a/Engine/Source/Editor/PropertyEditor/Private/DetailCategoryBuilderImpl.cpp b/Engine/Source/Editor/PropertyEditor/Private/DetailCategoryBuilderImpl.cpp
--- a/Engine/Source/Editor/PropertyEditor/Private/DetailCategoryBuilderImpl.cpp
+++ b/Engine/Source/Editor/PropertyEditor/Private/DetailCategoryBuilderImpl.cpp
@@ -1182,8 +1182,9 @@ void FDetailCategoryImpl::GenerateChildrenForSingleLayout(const FDetailLayout& L
 	FDetailNodeList GeneratedChildren;
 	GenerateNodesFromCustomizations(Customizations, GeneratedChildren);
 
+	// If this is the favorite category, create a group and put all properties based on their original category instead of showing them all flat
 	const FName InstanceName = Layout.GetInstanceName();
-	if (LayoutMap.ShouldShowGroup(InstanceName))
+	if (bFavoriteCategory || LayoutMap.ShouldShowGroup(InstanceName))
 	{
 		TSharedRef<FDetailCategoryGroupNode> GroupNode = MakeShareable(new FDetailCategoryGroupNode(InstanceName, AsShared()));
 		GroupNode->SetChildren(GeneratedChildren);
@@ -1197,6 +1198,20 @@ void FDetailCategoryImpl::GenerateChildrenForSingleLayout(const FDetailLayout& L
 
 void FDetailCategoryImpl::GenerateChildrenForLayouts()
 {
+	if (bFavoriteCategory)
+	{
+		// Merge advanced and simple layout customizations into simple nodes
+		// Favorite category should not have advanced nodes
+		for (const FDetailLayout& Layout : LayoutMap)
+		{
+			TArray<FDetailLayoutCustomization> Customizations;
+			Customizations.Append(Layout.GetSimpleLayouts());
+			Customizations.Append(Layout.GetAdvancedLayouts());
+			GenerateChildrenForSingleLayout(Layout, Customizations, SimpleChildNodes);
+		}
+		return;
+	}
+
 	// note: these can't be ranged-for, because the map may have items added to it by customizations during iteration
 	for (int32 LayoutIndex = 0; LayoutIndex < LayoutMap.Num(); ++LayoutIndex)
 	{
diff --git a/Engine/Source/Editor/PropertyEditor/Private/DetailLayoutHelpers.cpp b/Engine/Source/Editor/PropertyEditor/Private/DetailLayoutHelpers.cpp
--- a/Engine/Source/Editor/PropertyEditor/Private/DetailLayoutHelpers.cpp
+++ b/Engine/Source/Editor/PropertyEditor/Private/DetailLayoutHelpers.cpp
@@ -180,16 +180,8 @@ namespace DetailLayoutHelpers
 								}
 							}
 
-							// Add the property to the favorite
-							const FObjectPropertyNode* RootObjectParent = ChildNodePtr->FindRootObjectItemParent();
-							FName RootInstanceName = NAME_None;
-							if (RootObjectParent != nullptr)
-							{
-								RootInstanceName = RootObjectParent->GetObjectBaseClass()->GetFName();
-							}
-
-							// Duplicate the row
-							FavoritesCategory.AddPropertyNode(ChildNodePtr.ToSharedRef(), RootInstanceName);
+							// Add the property to the favorite by preserving its category, this will duplicate the row
+							FavoritesCategory.AddPropertyNode(ChildNodePtr.ToSharedRef(), CurCategory);
 						}
 
 						if (bIsStruct)
diff --git a/Engine/Source/Editor/PropertyEditor/Private/DetailsViewConfig.h b/Engine/Source/Editor/PropertyEditor/Private/DetailsViewConfig.h
--- a/Engine/Source/Editor/PropertyEditor/Private/DetailsViewConfig.h
+++ b/Engine/Source/Editor/PropertyEditor/Private/DetailsViewConfig.h
@@ -70,6 +70,10 @@ public:
 	/** A map of class name to a set of selected sections for that class. */
 	UPROPERTY()
 	TMap<FName, FDetailsSectionSelection> SelectedSections;
+	
+	/** Show Favorites and Modified sections */
+	UPROPERTY()
+	bool bShowHelperSections = false;
 };
 
 UCLASS(EditorConfig="DetailsView")
diff --git a/Engine/Source/Editor/PropertyEditor/Private/PropertyEditorCommands.cpp b/Engine/Source/Editor/PropertyEditor/Private/PropertyEditorCommands.cpp
new file mode 100644
--- /dev/null
+++ b/Engine/Source/Editor/PropertyEditor/Private/PropertyEditorCommands.cpp
@@ -0,0 +1,25 @@
+// Copyright Epic Games, Inc. All Rights Reserved.
+
+#include "PropertyEditorCommands.h"
+#include "Styling/AppStyle.h"
+
+#define LOCTEXT_NAMESPACE "PropertyEditorCommands"
+
+FPropertyEditorCommands::FPropertyEditorCommands()
+	: TCommands<FPropertyEditorCommands>
+	(
+		TEXT("PropertyEditor"),
+		LOCTEXT("PropertyEditorContext", "Property Editor"),
+		NAME_None,
+		FAppStyle::GetAppStyleSetName()
+	)
+{
+}
+
+void FPropertyEditorCommands::RegisterCommands()
+{
+	UI_COMMAND(ExpandAll, "Expand All", "Expand all elements in this details view", EUserInterfaceActionType::Button, FInputChord(EModifierKey::Control | EModifierKey::Shift, EKeys::X));
+	UI_COMMAND(CollapseAll, "Collapse All", "Collapse all elements in this details view", EUserInterfaceActionType::Button, FInputChord(EModifierKey::Control | EModifierKey::Shift, EKeys::C));
+}
+
+#undef LOCTEXT_NAMESPACE
diff --git a/Engine/Source/Editor/PropertyEditor/Private/PropertyEditorModule.cpp b/Engine/Source/Editor/PropertyEditor/Private/PropertyEditorModule.cpp
--- a/Engine/Source/Editor/PropertyEditor/Private/PropertyEditorModule.cpp
+++ b/Engine/Source/Editor/PropertyEditor/Private/PropertyEditorModule.cpp
@@ -1,41 +1,42 @@
 // Copyright Epic Games, Inc. All Rights Reserved.
 
 #include "PropertyEditorModule.h"
+
 #include "AssetToolsModule.h"
 #include "DetailRowMenuContextPrivate.h"
+#include "DetailsViewStyle.h"
+#include "Engine/UserDefinedEnum.h"
+#include "Framework/Application/SlateApplication.h"
 #include "IAssetTools.h"
 #include "IDetailsView.h"
+#include "Interfaces/IMainFrameModule.h"
 #include "IPropertyChangeListener.h"
 #include "IPropertyTable.h"
 #include "IPropertyTableCellPresenter.h"
 #include "IPropertyTableWidgetHandle.h"
+#include "IStructureDataProvider.h"
+#include "Modules/ModuleManager.h"
+#include "Presentation/PropertyEditor/PropertyEditor.h"
+#include "Presentation/PropertyTable/PropertyTable.h"
 #include "PropertyChangeListener.h"
+#include "PropertyEditorCommands.h"
 #include "PropertyEditorToolkit.h"
 #include "PropertyRowGenerator.h"
 #include "SDetailsView.h"
 #include "SPropertyTreeViewImpl.h"
 #include "SSingleProperty.h"
 #include "SStructureDetailsView.h"
-
-#include "Engine/UserDefinedEnum.h"
 #include "StructUtils/UserDefinedStruct.h"
-#include "Framework/Application/SlateApplication.h"
-#include "Interfaces/IMainFrameModule.h"
-#include "Modules/ModuleManager.h"
-#include "Presentation/PropertyEditor/PropertyEditor.h"
-#include "Presentation/PropertyTable/PropertyTable.h"
 #include "Toolkits/AssetEditorToolkit.h"
+#include "ToolMenus.h"
 #include "UObject/UnrealType.h"
 #include "UserInterface/PropertyTable/PropertyTableWidgetHandle.h"
 #include "UserInterface/PropertyTable/SPropertyTable.h"
 #include "UserInterface/PropertyTable/TextPropertyTableCellPresenter.h"
 #include "Widgets/Colors/SColorPicker.h"
 #include "Widgets/Layout/SBorder.h"
-#include "DetailsViewStyle.h"
-#include "ToolMenus.h"
-#include "IStructureDataProvider.h"
 
-IMPLEMENT_MODULE( FPropertyEditorModule, PropertyEditor );
+IMPLEMENT_MODULE(FPropertyEditorModule, PropertyEditor);
 
 TAutoConsoleVariable<bool> FPropertyEditorModule::CVarAllowPropertyMatrixForNonActors = TAutoConsoleVariable<bool>(
 	TEXT("PropertyEditor.AllowPropertyMatrixForNonActors"),
@@ -109,6 +110,8 @@ void FPropertyEditorModule::ShutdownModule()
 {
 	UToolMenus::UnRegisterStartupCallback(this);
 	UToolMenus::UnregisterOwner(this);
+
+	FPropertyEditorCommands::Unregister();
 	
 	// No need to remove this object from root since the final GC pass doesn't care about root flags
 	StructOnScopePropertyOwner = nullptr;
diff --git a/Engine/Source/Editor/PropertyEditor/Private/SDetailCategoryTableRow.cpp b/Engine/Source/Editor/PropertyEditor/Private/SDetailCategoryTableRow.cpp
--- a/Engine/Source/Editor/PropertyEditor/Private/SDetailCategoryTableRow.cpp
+++ b/Engine/Source/Editor/PropertyEditor/Private/SDetailCategoryTableRow.cpp
@@ -300,6 +300,111 @@ void SDetailCategoryTableRow::PopulateContextMenu(UToolMenu* ToolMenu)
 	}
 }
 
+bool SDetailCategoryTableRow::CanFavoriteRow(bool bInAddToFavorite)
+{
+	const TSharedPtr<FDetailTreeNode> OwnerNode = OwnerTreeNode.Pin();
+	if (!OwnerNode)
+	{
+		return false;
+	}
+
+	// Checks if node is a favorite or included in favorite category
+	TFunction<bool(const TSharedPtr<FDetailTreeNode>&, bool)> IsFavoriteCategoryNode = [&IsFavoriteCategoryNode](const TSharedPtr<FDetailTreeNode>& InNode, bool bInRecurse)
+		{
+			if (!InNode)
+			{
+				return false;
+			}
+
+			if (InNode->GetNodeType() == EDetailNodeType::Category)
+			{
+				const TSharedPtr<FDetailCategoryImpl> CategoryNode = StaticCastSharedPtr<FDetailCategoryImpl>(InNode);
+				if (CategoryNode->IsFavoriteCategory())
+				{
+					return true;
+				}
+			}
+
+			return bInRecurse ? IsFavoriteCategoryNode(InNode->GetParentNode().Pin(), bInRecurse) : false;
+		};
+
+	// Disallow adding Favorites category to Favorites...
+	// Disallow adding subcategories (group) of Favorites to Favorites
+	const bool bIsRowInFavorites = IsFavoriteCategoryNode(OwnerNode, /** Recurse */true);
+	if (bInAddToFavorite && bIsRowInFavorites)
+	{
+		return false;
+	}
+
+	bool bCanFavoriteRow = false;
+	ForEachImmediateChildPropertyNode([bInAddToFavorite, &bCanFavoriteRow](const TSharedPtr<FPropertyNode>& PropertyNode)
+		{
+			if (PropertyNode->IsFavorite() != bInAddToFavorite)
+			{
+				bCanFavoriteRow = true;
+				return false; // Break
+			}
+
+			return true; // Continue
+		}
+	);
+
+	return bCanFavoriteRow;
+}
+
+void SDetailCategoryTableRow::OnFavoriteRow(bool bInAddToFavorite)
+{
+	bool bRefresh = false;
+	ForEachImmediateChildPropertyNode([bInAddToFavorite, &bRefresh](const TSharedPtr<FPropertyNode>& PropertyNode)
+		{
+			PropertyNode->SetFavorite(bInAddToFavorite);
+			bRefresh = true;
+			return true; // Continue
+		}
+	);
+
+	if (bRefresh)
+	{
+		// Refresh the tree
+		ForceRefresh();
+	}
+}
+
+void SDetailCategoryTableRow::ForEachImmediateChildPropertyNode(TFunctionRef<bool(TSharedPtr<FPropertyNode>&)> InFunctor)
+{
+	const TSharedPtr<FDetailTreeNode> TreeNode = OwnerTreeNode.Pin();
+	if (!TreeNode)
+	{
+		return;
+	}
+
+	TArray<TSharedRef<IDetailTreeNode>> ChildNodes;
+	TreeNode->GetChildren(ChildNodes);
+	if (ChildNodes.IsEmpty())
+	{
+		return;
+	}
+
+	for (const TSharedRef<IDetailTreeNode>& Child : ChildNodes)
+	{
+		if (TSharedPtr<IPropertyHandle> PropertyHandle = Child->CreatePropertyHandle())
+		{
+			for (TSharedPtr<FPropertyNode>& PropertyNode : GetPropertyNodesFromHandles({PropertyHandle}))
+			{
+				if (PropertyNode.IsValid())
+				{
+					const bool bContinue = InFunctor(PropertyNode);
+				
+					if (!bContinue)
+					{
+						return;
+					}
+				}
+			}
+		}
+	}
+}
+
 EVisibility SDetailCategoryTableRow::IsSeparatorVisible() const
 {
 	return bIsInnerCategory || IsItemExpanded() ? EVisibility::Collapsed : EVisibility::Visible;
@@ -335,7 +440,7 @@ const FSlateBrush* SDetailCategoryTableRow::GetBackgroundImageForScrollBarWell()
 	return nullptr;
 }
 
-FSlateColor SDetailCategoryTableRow::	GetInnerBackgroundColor() const
+FSlateColor SDetailCategoryTableRow::GetInnerBackgroundColor() const
 {
 	FSlateColor Color = FSlateColor(FLinearColor::White);
 	
diff --git a/Engine/Source/Editor/PropertyEditor/Private/SDetailCategoryTableRow.h b/Engine/Source/Editor/PropertyEditor/Private/SDetailCategoryTableRow.h
--- a/Engine/Source/Editor/PropertyEditor/Private/SDetailCategoryTableRow.h
+++ b/Engine/Source/Editor/PropertyEditor/Private/SDetailCategoryTableRow.h
@@ -2,6 +2,7 @@
 
 #pragma once
 
+#include "DetailsViewStyle.h"
 #include "DetailTreeNode.h"
 #include "Input/Reply.h"
 #include "Layout/Visibility.h"
@@ -11,8 +12,6 @@
 #include "Widgets/DeclarativeSyntaxSupport.h"
 #include "Widgets/Views/STableRow.h"
 #include "Widgets/Views/STableViewBase.h"
-#include "DetailsViewStyle.h"
-
 
 class SDetailCategoryTableRow : public SDetailTableRowBase
 {
@@ -42,14 +41,21 @@ public:
 
 	void Construct(const FArguments& InArgs, TSharedRef<FDetailTreeNode> InOwnerTreeNode, const TSharedRef<STableViewBase>& InOwnerTableView);
 
-	// ~Begin SWidget Interface
+	//~ Begin SWidget
 	virtual FReply OnMouseButtonUp(const FGeometry& MyGeometry, const FPointerEvent& MouseEvent) override;
-	// ~Begin End Interface
+	//~ End SWidget
 
 protected:
+	//~ Begin SDetailTableRowBase
 	virtual void PopulateContextMenu(UToolMenu* ToolMenu) override;
+	virtual bool CanFavoriteRow(bool bInAddToFavorite) override;
+	virtual void OnFavoriteRow(bool bInAddToFavorite) override;
+	//~ End SDetailTableRowBase
 
 private:
+	/** Process immediate children property nodes of this category, return false to break */
+	void ForEachImmediateChildPropertyNode(TFunctionRef<bool(TSharedPtr<FPropertyNode>&)> InFunctor);
+
 	EVisibility IsSeparatorVisible() const;
 
 	/**
diff --git a/Engine/Source/Editor/PropertyEditor/Private/SDetailSingleItemRow.cpp b/Engine/Source/Editor/PropertyEditor/Private/SDetailSingleItemRow.cpp
--- a/Engine/Source/Editor/PropertyEditor/Private/SDetailSingleItemRow.cpp
+++ b/Engine/Source/Editor/PropertyEditor/Private/SDetailSingleItemRow.cpp
@@ -1339,34 +1339,6 @@ void SDetailSingleItemRow::PopulateContextMenu(UToolMenu* ToolMenu)
 				SelectByPropertyAction);
 		}
 
-		// Favorite
-		{
-			if (OwnerTreeNode.Pin()->GetDetailsViewSharedPtr()->IsFavoritingEnabled())
-			{
-				FUIAction FavoriteAction;
-				FavoriteAction.ExecuteAction = FExecuteAction::CreateSP(this, &SDetailSingleItemRow::OnFavoriteMenuToggle);
-				FavoriteAction.CanExecuteAction = FCanExecuteAction::CreateSP(this, &SDetailSingleItemRow::CanFavorite);
-
-				FText FavoriteText = NSLOCTEXT("PropertyView", "FavoriteProperty", "Add to Favorites");
-				FText FavoriteTooltipText = NSLOCTEXT("PropertyView", "FavoriteProperty_ToolTip", "Add this property to your favorites.");
-				FName FavoriteIcon = "DetailsView.PropertyIsFavorite";
-
-				if (IsFavorite())
-				{
-					FavoriteText = NSLOCTEXT("PropertyView", "RemoveFavoriteProperty", "Remove from Favorites");
-					FavoriteTooltipText = NSLOCTEXT("PropertyView", "RemoveFavoriteProperty_ToolTip", "Remove this property from your favorites.");
-					FavoriteIcon = "DetailsView.PropertyIsNotFavorite";
-				}
-
-				EditSection.AddMenuEntry(
-					TEXT("ToggleFavorite"),
-					FavoriteText,
-					FavoriteTooltipText,
-					FSlateIcon(FAppStyle::Get().GetStyleSetName(), FavoriteIcon),
-					FavoriteAction);
-			}
-		}
-
 		if (FPropertyEditorPermissionList::Get().ShouldShowMenuEntries())
 		{
 			// Hide separator line if it only contains the SearchWidget, making the next 2 elements the top of the list
@@ -1457,6 +1429,81 @@ TSharedPtr<IPropertyHandle> SDetailSingleItemRow::GetPrimaryPropertyHandle() con
 	return nullptr;
 }
 
+bool SDetailSingleItemRow::CanFavoriteRow(bool bInAddToFavorite)
+{
+	if (bInAddToFavorite == IsFavorite())
+	{
+		return false;		
+	}
+
+	if (Customization->HasPropertyNode())
+	{
+		return true;
+	}
+
+	if (Customization->HasCustomBuilder())
+	{
+		const TSharedPtr<IPropertyHandle> PropertyHandle = Customization->CustomBuilderRow->GetPropertyHandle();
+		const FString& OriginalPath = Customization->CustomBuilderRow->GetOriginalPath();
+		return PropertyHandle.IsValid() || !OriginalPath.IsEmpty();
+	}
+
+	return false;
+}
+
+void SDetailSingleItemRow::OnFavoriteRow(bool bInAddToFavorite)
+{
+	const TSharedPtr<FDetailTreeNode> OwnerTreeNodePinned = OwnerTreeNode.Pin();
+	if (!OwnerTreeNodePinned.IsValid())
+	{
+		return;
+	}
+
+	const TSharedPtr<IDetailsViewPrivate> DetailsView = OwnerTreeNodePinned->GetDetailsViewSharedPtr();
+	if (!DetailsView)
+	{
+		return;
+	}
+
+	if (Customization->HasPropertyNode())
+	{
+		const TSharedPtr<FPropertyNode> PropertyNode = Customization->GetPropertyNode();
+		PropertyNode->SetFavorite(bInAddToFavorite);
+	}
+	else if (Customization->HasCustomBuilder())
+	{
+		const TSharedPtr<IPropertyHandle> PropertyHandle = Customization->CustomBuilderRow->GetPropertyHandle();
+		const FString& OriginalPath = Customization->CustomBuilderRow->GetOriginalPath();
+
+		if (PropertyHandle.IsValid())
+		{
+			StaticCastSharedPtr<FPropertyHandleBase>(PropertyHandle)->GetPropertyNode()->SetFavorite(bInAddToFavorite); 
+		}
+		else if (!OriginalPath.IsEmpty())
+		{
+			DetailsView->SetCustomBuilderFavorite(OriginalPath, bInAddToFavorite);
+		}
+	}
+
+	// Calculate the scrolling offset (by item) to make sure the mouse stay over the same property
+	int32 ExpandSize = 0;
+	if (OwnerTreeNodePinned->ShouldBeExpanded())
+	{
+		SDetailSingleItemRow_Helper::RecursivelyGetItemShow(OwnerTreeNodePinned.ToSharedRef(), ExpandSize);
+	}
+	else
+	{
+		// if the item is not expand count is 1
+		ExpandSize = 1;
+	}
+
+	// Apply the calculated offset
+	DetailsView->MoveScrollOffset(bInAddToFavorite ? ExpandSize : -ExpandSize);
+
+	// Refresh the tree
+	ForceRefresh();
+}
+
 void SDetailSingleItemRow::OnCopyProperty()
 {
 	if (OwnerTreeNode.IsValid())
@@ -1865,24 +1912,6 @@ void SDetailSingleItemRow::PopulateExtensionWidget()
 	}
 }
 
-bool SDetailSingleItemRow::CanFavorite() const
-{
-	if (Customization->HasPropertyNode())
-	{
-		return true;
-	}
-
-	if (Customization->HasCustomBuilder())
-	{
-		TSharedPtr<IPropertyHandle> PropertyHandle = Customization->CustomBuilderRow->GetPropertyHandle();
-		const FString& OriginalPath = Customization->CustomBuilderRow->GetOriginalPath();
-		
-		return PropertyHandle.IsValid() || !OriginalPath.IsEmpty();
-	}
-
-	return false;
-}
-
 bool SDetailSingleItemRow::IsFavorite() const
 {
 	if (Customization->HasPropertyNode())
@@ -1918,62 +1947,6 @@ bool SDetailSingleItemRow::IsFavorite() const
 	return false;
 }
 
-void SDetailSingleItemRow::OnFavoriteMenuToggle()
-{
-	if (!CanFavorite())
-	{
-		return; 
-	}
-
-	TSharedPtr<FDetailTreeNode> OwnerTreeNodePinned = OwnerTreeNode.Pin();
-	if (!OwnerTreeNodePinned.IsValid())
-	{
-		return;
-	}
-
-	TSharedPtr<IDetailsViewPrivate> DetailsView = OwnerTreeNodePinned->GetDetailsViewSharedPtr();
-
-	bool bNewValue = !IsFavorite();
-	if (Customization->HasPropertyNode())
-	{
-		TSharedPtr<FPropertyNode> PropertyNode = Customization->GetPropertyNode();
-		PropertyNode->SetFavorite(bNewValue);
-	}
-	else if (Customization->HasCustomBuilder())
-	{
-		TSharedPtr<IPropertyHandle> PropertyHandle = Customization->CustomBuilderRow->GetPropertyHandle();
-		const FString& OriginalPath = Customization->CustomBuilderRow->GetOriginalPath();
-
-		if (PropertyHandle.IsValid())
-		{
-			StaticCastSharedPtr<FPropertyHandleBase>(PropertyHandle)->GetPropertyNode()->SetFavorite(bNewValue); 
-		}
-		else if (!OriginalPath.IsEmpty())
-		{		
-			bNewValue = !DetailsView->IsCustomBuilderFavorite(OriginalPath);
-			DetailsView->SetCustomBuilderFavorite(OriginalPath, bNewValue);
-		}
-	}
-
-	// Calculate the scrolling offset (by item) to make sure the mouse stay over the same property
-	int32 ExpandSize = 0;
-	if (OwnerTreeNodePinned->ShouldBeExpanded())
-	{
-		SDetailSingleItemRow_Helper::RecursivelyGetItemShow(OwnerTreeNodePinned.ToSharedRef(), ExpandSize);
-	}
-	else
-	{
-		// if the item is not expand count is 1
-		ExpandSize = 1;
-	}
-
-	// Apply the calculated offset
-	DetailsView->MoveScrollOffset(bNewValue ? ExpandSize : -ExpandSize);
-
-	// Refresh the tree
-	ForceRefresh();
-}
-
 void SDetailSingleItemRow::CreateGlobalExtensionWidgets(TArray<FPropertyRowExtensionButton>& OutExtensions) const
 {
 	// fetch global extension widgets 
diff --git a/Engine/Source/Editor/PropertyEditor/Private/SDetailSingleItemRow.h b/Engine/Source/Editor/PropertyEditor/Private/SDetailSingleItemRow.h
--- a/Engine/Source/Editor/PropertyEditor/Private/SDetailSingleItemRow.h
+++ b/Engine/Source/Editor/PropertyEditor/Private/SDetailSingleItemRow.h
@@ -60,18 +60,21 @@ public:
 	 */
 	void Construct( const FArguments& InArgs, FDetailLayoutCustomization* InCustomization, bool bHasMultipleColumns, TSharedRef<FDetailTreeNode> InOwnerTreeNode, const TSharedRef<STableViewBase>& InOwnerTableView );
 
-	// ~Begin SWidget Interface
+	//~ Begin SWidget
 	virtual FReply OnMouseButtonUp(const FGeometry& MyGeometry, const FPointerEvent& MouseEvent) override;
-	// ~Begin End Interface
+	virtual void Tick(const FGeometry& AllottedGeometry, const double InCurrentTime, const float InDeltaTime) override;
+	//~ End SWidget
 
 	TSharedPtr<FDragDropOperation> CreateDragDropOperation();
 
-	virtual void Tick(const FGeometry& AllottedGeometry, const double InCurrentTime, const float InDeltaTime);
-
 protected:
+	//~ Begin SDetailTableRowBase
 	virtual void PopulateContextMenu(UToolMenu* ToolMenu) override;
 	virtual TArray<TSharedPtr<IPropertyHandle>> GetPropertyHandles(const bool bRecursive = false) const override;
 	virtual TSharedPtr<IPropertyHandle> GetPrimaryPropertyHandle() const override;
+	virtual bool CanFavoriteRow(bool bInAddToFavorite) override;
+	virtual void OnFavoriteRow(bool bInAddToFavorite) override;
+	//~ End SDetailTableRowBase
 
 private:
 	void OnCopyProperty();
@@ -126,9 +129,6 @@ private:
 	bool IsResetToDefaultVisible() const;
 
 	bool IsHighlighted() const;
-
-	void OnFavoriteMenuToggle();
-	bool CanFavorite() const;
 	bool IsFavorite() const;
 	
 	/** UIActions to help populate the PropertyEditorPermissionList, which must first be turned on through FPropertyEditorPermissionList::Get().SetShouldShowMenuEntries */
@@ -158,10 +158,10 @@ private:
 	void SetPropertyColorationTarget(bool bSelectActors = false);
 private:
 	/** Customization for this widget */
-	FDetailLayoutCustomization* Customization;
+	FDetailLayoutCustomization* Customization = nullptr;
 	FDetailWidgetRow WidgetRow;
-	bool bAllowFavoriteSystem;
-	bool bCachedResetToDefaultVisible;
+	bool bAllowFavoriteSystem = false;
+	bool bCachedResetToDefaultVisible = false;
 	TSharedPtr<FPropertyNode> SwappablePropertyNode;
 	TSharedPtr<SButton> ExpanderArrow;
 	TWeakPtr<FDragDropOperation> DragOperation; // last drag initiated by this widget
diff --git a/Engine/Source/Editor/PropertyEditor/Private/SDetailTableRowBase.cpp b/Engine/Source/Editor/PropertyEditor/Private/SDetailTableRowBase.cpp
--- a/Engine/Source/Editor/PropertyEditor/Private/SDetailTableRowBase.cpp
+++ b/Engine/Source/Editor/PropertyEditor/Private/SDetailTableRowBase.cpp
@@ -91,10 +91,17 @@ bool SDetailTableRowBase::IsScrollBarVisible(TWeakPtr<STableViewBase> OwnerTable
 
 void SDetailTableRowBase::PopulateContextMenu(UToolMenu* ToolMenu)
 {
+	const TSharedPtr<FDetailTreeNode> DetailTreeNode = OwnerTreeNode.Pin();
+	
+	if (!DetailTreeNode)
+	{
+		return;
+	}
+
 	FToolMenuSection& ExpansionSection = ToolMenu->FindOrAddSection(TEXT("Expansion"));
 	{
 		FDetailNodeList VisibleChildren;
-		OwnerTreeNode.Pin()->GetChildren( VisibleChildren );
+		DetailTreeNode->GetChildren( VisibleChildren );
 
 		// Open context menu if this node can be expanded 
 		if( VisibleChildren.Num() )
@@ -116,6 +123,40 @@ void SDetailTableRowBase::PopulateContextMenu(UToolMenu* ToolMenu)
 				ExpandAllAction);
 		}
 	}
+	
+	TSharedPtr<IDetailsViewPrivate> DetailsPanel = DetailTreeNode->GetDetailsViewSharedPtr();
+	if (DetailsPanel && DetailsPanel->IsFavoritingEnabled())
+	{
+		FToolMenuSection& FavoritesSection = ToolMenu->FindOrAddSection(TEXT("Favorites"));
+
+		bool bAddFavorites = true;
+		FavoritesSection.AddMenuEntry(
+			TEXT("FavoriteAll"),
+			NSLOCTEXT("PropertyView", "FavoriteRow", "Add To Favorites"),
+			NSLOCTEXT("PropertyView", "FavoriteRow_ToolTip", "Add this row to favorites"),
+			FSlateIcon(FAppStyle::Get().GetStyleSetName(), "DetailsView.PropertyIsFavorite"),
+			FUIAction(
+				FExecuteAction::CreateSP(this, &SDetailTableRowBase::OnFavoriteRow, bAddFavorites),
+				FCanExecuteAction(),
+				FIsActionChecked(),
+				FIsActionButtonVisible::CreateSP(this, &SDetailTableRowBase::CanFavoriteRow, bAddFavorites)
+			)
+		);
+		
+		bAddFavorites = false;
+		FavoritesSection.AddMenuEntry(
+			TEXT("UnFavoriteAll"),
+			NSLOCTEXT("PropertyView", "UnFavoriteRow", "Remove From Favorites"),
+			NSLOCTEXT("PropertyView", "UnFavoriteRow_ToolTip", "Remove this row from favorites"),
+			FSlateIcon(FAppStyle::Get().GetStyleSetName(), "DetailsView.PropertyIsNotFavorite"),
+			FUIAction(
+				FExecuteAction::CreateSP(this, &SDetailTableRowBase::OnFavoriteRow, bAddFavorites),
+				FCanExecuteAction(),
+				FIsActionChecked(),
+				FIsActionButtonVisible::CreateSP(this, &SDetailTableRowBase::CanFavoriteRow, bAddFavorites)
+			)
+		);
+	}
 }
 
 TArray<TSharedPtr<FPropertyNode>> SDetailTableRowBase::GetPropertyNodes(const bool bRecursive) const
@@ -214,3 +255,13 @@ void SDetailTableRowBase::ForceRefresh()
 		}
 	}
 }
+
+void SDetailTableRowBase::OnFavoriteRow(bool bInAddToFavorite)
+{
+	
+}
+
+bool SDetailTableRowBase::CanFavoriteRow(bool bInAddToFavorite)
+{
+	return false;
+}
diff --git a/Engine/Source/Editor/PropertyEditor/Private/SDetailTableRowBase.h b/Engine/Source/Editor/PropertyEditor/Private/SDetailTableRowBase.h
--- a/Engine/Source/Editor/PropertyEditor/Private/SDetailTableRowBase.h
+++ b/Engine/Source/Editor/PropertyEditor/Private/SDetailTableRowBase.h
@@ -59,6 +59,12 @@ protected:
 
 	/** Refreshes widget and associated list views. */
 	virtual void ForceRefresh();
+	
+	/** Called when row favoriting is toggled */
+	virtual void OnFavoriteRow(bool bInAddToFavorite);
+	
+	/** Is row favoriting allowed */
+	virtual bool CanFavoriteRow(bool bInAddToFavorite);
 
 private:
 	void OnExpandAllClicked()
diff --git a/Engine/Source/Editor/PropertyEditor/Private/SDetailsView.cpp b/Engine/Source/Editor/PropertyEditor/Private/SDetailsView.cpp
--- a/Engine/Source/Editor/PropertyEditor/Private/SDetailsView.cpp
+++ b/Engine/Source/Editor/PropertyEditor/Private/SDetailsView.cpp
@@ -16,6 +16,7 @@
 #include "Misc/ConfigCacheIni.h"
 #include "Modules/ModuleManager.h"
 #include "ObjectPropertyNode.h"
+#include "PropertyEditorCommands.h"
 #include "PropertyEditorHelpers.h"
 #include "SDetailNameArea.h"
 #include "Styling/StyleColors.h"
@@ -68,6 +69,8 @@ void SDetailsView::Construct(const FArguments& InArgs, const FDetailsViewArgs& I
 {
 	DetailsViewArgs = InDetailsViewArgs;
 	DetailsNameWidgetOverrideCustomization = DetailsViewArgs.DetailsNameWidgetOverrideCustomization;
+	
+	BindCommands();
 
 	const FDetailsViewConfig* ViewConfig = GetConstViewConfig();
 	if (ViewConfig != nullptr)
@@ -119,7 +122,7 @@ void SDetailsView::Construct(const FArguments& InArgs, const FDetailsViewArgs& I
 	const TSharedRef<SScrollBar> ExternalScrollbar = InDetailsViewArgs.ExternalScrollbar ? InDetailsViewArgs.ExternalScrollbar.ToSharedRef() : SNew(SScrollBar);
 	ExternalScrollbar->SetVisibility( TAttribute<EVisibility>( this, &SDetailsView::GetScrollBarVisibility ) );
 
-	FMenuBuilder DetailViewOptions(true, nullptr, InDetailsViewArgs.OptionsExtender);
+	FMenuBuilder DetailViewOptions(true, InternalCommands, InDetailsViewArgs.OptionsExtender);
 	const bool bDisplayShowSectionHeadingText = InDetailsViewArgs.OptionsExtender.IsValid() && (InDetailsViewArgs.OptionsExtender->NumExtensions() > 0);
 	DetailViewOptions.BeginSection("DetailsViewShowOptions", (bDisplayShowSectionHeadingText ? LOCTEXT("ShowSectionHeading", "Show") : FText::GetEmpty()));
 
@@ -258,22 +261,27 @@ void SDetailsView::Construct(const FArguments& InArgs, const FDetailsViewArgs& I
 		),
 		NAME_None,
 		EUserInterfaceActionType::Check
-	);		
+	);
+	
+	DetailViewOptions.AddMenuEntry(
+		LOCTEXT("ShowHelperSections", "Show Helper Sections"),
+		LOCTEXT("ShowHelperSectionsToolTip", "Specifies whether helper sections like Favorites, Modified are shown in the filter bar"),
+		FSlateIcon(),
+		FUIAction( 
+			FExecuteAction::CreateSP(this, &SDetailsView::OnShowHelperSectionsClicked),
+			FCanExecuteAction(),
+			FIsActionChecked::CreateSP(this, &SDetailsView::IsShowHelperSectionsChecked)
+		),
+		NAME_None,
+		EUserInterfaceActionType::Check
+	);
 
 	DetailViewOptions.EndSection();
 	DetailViewOptions.BeginSection("DetailsViewCategories");
 
-	DetailViewOptions.AddMenuEntry(
-		LOCTEXT("CollapseAll", "Collapse All Categories"),
-		LOCTEXT("CollapseAll_ToolTip", "Collapses all root level categories"),
-		FSlateIcon(),
-		FUIAction(FExecuteAction::CreateSP(this, &SDetailsView::SetRootExpansionStates, /*bExpanded=*/false, /*bRecurse=*/false )));
-
-	DetailViewOptions.AddMenuEntry(
-		LOCTEXT("ExpandAll", "Expand All Categories"),
-		LOCTEXT("ExpandAll_ToolTip", "Expands all root level categories"),
-		FSlateIcon(),
-		FUIAction(FExecuteAction::CreateSP(this, &SDetailsView::SetRootExpansionStates, /*bExpanded=*/true, /*bRecurse=*/false )));
+	// Commands have been registered earlier
+	DetailViewOptions.AddMenuEntry(FPropertyEditorCommands::Get().CollapseAll);
+	DetailViewOptions.AddMenuEntry(FPropertyEditorCommands::Get().ExpandAll);
 
 	TSharedRef<SVerticalBox> VerticalBox = SNew(SVerticalBox);
 		
@@ -1183,6 +1191,24 @@ void SDetailsView::OnShowSectionsClicked()
 	UpdateFilteredDetails();
 }
 
+bool SDetailsView::IsShowHelperSectionsChecked() const
+{
+	const FDetailsViewConfig* ViewConfig = GetConstViewConfig();
+	return ViewConfig && ViewConfig->bShowHelperSections;
+}
+
+void SDetailsView::OnShowHelperSectionsClicked()
+{
+	if (FDetailsViewConfig* ViewConfig = GetMutableViewConfig())
+	{
+		ViewConfig->bShowHelperSections = !ViewConfig->bShowHelperSections;
+
+		SaveViewConfig();
+		RebuildSectionSelector();
+		UpdateFilteredDetails();
+	}
+}
+
 FSlateColor SDetailsView::GetToggleFavoritesColor() const
 {
 	if (DetailsViewArgs.bAllowFavoriteSystem && CurrentFilter.bShowFavoritesCategory)
@@ -1433,6 +1459,28 @@ void SDetailsView::RebuildSectionSelector()
 		CreateSection(FName("All"), NSLOCTEXT("UObjectSection", "All", "All"))
 	];
 
+	const bool bAddHelperSections = IsShowHelperSectionsChecked();
+
+	if (bAddHelperSections
+		&& DetailsViewArgs.bShowOptions 
+		&& DetailsViewArgs.bAllowFavoriteSystem)
+	{
+		SectionSelectorBox->AddSlot()
+		[
+			CreateSection(FName("Favorites"), NSLOCTEXT("UObjectSection", "Favorites", "Favorites"))
+		];
+	}
+
+	if (bAddHelperSections
+		&& DetailsViewArgs.bShowOptions 
+		&& DetailsViewArgs.bShowModifiedPropertiesOption)
+	{
+		SectionSelectorBox->AddSlot()
+		[
+			CreateSection(FName("Modified"), NSLOCTEXT("UObjectSection", "Modified", "Modified"))
+		];
+	}
+
 	SectionSelectorBox->SetVisibility(EVisibility::Visible);
 
 	CurrentFilter.VisibleSections.Reset();
diff --git a/Engine/Source/Editor/PropertyEditor/Private/SDetailsView.h b/Engine/Source/Editor/PropertyEditor/Private/SDetailsView.h
--- a/Engine/Source/Editor/PropertyEditor/Private/SDetailsView.h
+++ b/Engine/Source/Editor/PropertyEditor/Private/SDetailsView.h
@@ -205,6 +205,12 @@ private:
 
 	/** Called when Show Sections is clicked */
 	void OnShowSectionsClicked();
+	
+	/** @return true if Show Helper Sections is checked. */
+	bool IsShowHelperSectionsChecked() const;
+	
+	/** Called when Show Helper Sections is clicked */
+	void OnShowHelperSectionsClicked();
 
 	/** Get the color of the toggle favorites button. */
 	FSlateColor GetToggleFavoritesColor() const;
diff --git a/Engine/Source/Editor/PropertyEditor/Private/SDetailsViewBase.cpp b/Engine/Source/Editor/PropertyEditor/Private/SDetailsViewBase.cpp
--- a/Engine/Source/Editor/PropertyEditor/Private/SDetailsViewBase.cpp
+++ b/Engine/Source/Editor/PropertyEditor/Private/SDetailsViewBase.cpp
@@ -17,7 +17,9 @@
 #include "DetailLayoutHelpers.h"
 #include "Editor.h"
 #include "IDetailCustomization.h"
+#include "LevelEditorActions.h"
 #include "ObjectPropertyNode.h"
+#include "PropertyEditorCommands.h"
 #include "PropertyPermissionList.h"
 #include "ScopedTransaction.h"
 #include "SDetailNameArea.h"
@@ -88,6 +90,46 @@ void SDetailsViewBase::ClearKeyboardFocusIfWithin(const TSharedRef<SWidget>& Wid
 	}
 }
 
+FDetailFilter SDetailsViewBase::CreateDetailFilter()
+{
+	FDetailFilter Filter = CurrentFilter;
+	
+	// When "Favorites" section is active, show favorited items
+	Filter.bShowFavoritesCategory |= CurrentFilter.VisibleSections.Contains(TEXT("Favorites"));
+
+	// When "Modified" section is active, show all items with modified properties
+	if (CurrentFilter.VisibleSections.Contains(TEXT("Modified")))
+	{
+		Filter.VisibleSections.Empty();
+		Filter.bShowOnlyModified |= true;
+	}
+
+	return Filter;
+}
+
+void SDetailsViewBase::BindCommands()
+{
+	if (InternalCommands.IsValid())
+	{
+		return;
+	}
+
+	// Register commands here otherwise they are destroyed...
+	FPropertyEditorCommands::Register();
+	const FPropertyEditorCommands& PropertyEditorCommands = FPropertyEditorCommands::Get();
+
+	InternalCommands = MakeShared<FUICommandList>();
+	InternalCommands->MapAction(
+		PropertyEditorCommands.CollapseAll, 
+		FExecuteAction::CreateSP(this, &SDetailsViewBase::SetRootExpansionStates, /** Expand */false, /** Recurse */false)
+	);
+
+	InternalCommands->MapAction(
+		PropertyEditorCommands.ExpandAll, 
+		FExecuteAction::CreateSP(this, &SDetailsViewBase::SetRootExpansionStates, /** Expand */true, /** Recurse */false)
+	);
+}
+
 void SDetailsViewBase::OnRowReleasedForDetailTree(const TSharedRef<ITableRow>& TableRow)
 {
 	ClearKeyboardFocusIfWithin(TableRow->AsWidget());
@@ -123,34 +165,31 @@ void SDetailsViewBase::SetNodeExpansionState(TSharedRef<FDetailTreeNode> InTreeN
 	TArray< TSharedRef<FDetailTreeNode> > Children;
 	InTreeNode->GetChildren(Children);
 
-	if (Children.Num())
+	RequestItemExpanded(InTreeNode, bExpand);
+	constexpr bool bShouldSaveState = true;
+	InTreeNode->OnItemExpansionChanged(bExpand, bShouldSaveState);
+
+	if (bRecursive && !Children.IsEmpty())
 	{
-		RequestItemExpanded(InTreeNode, bExpand);
-		const bool bShouldSaveState = true;
-		InTreeNode->OnItemExpansionChanged(bExpand, bShouldSaveState);
+		TArray<FString> TypesToIgnoreForRecursiveExpansion;
+		GConfig->GetArray(TEXT("DetailPropertyExpansion"), TEXT("TypesToIgnoreForRecursiveExpansion"), TypesToIgnoreForRecursiveExpansion, GEditorPerProjectIni);
 
-		if (bRecursive)
+		for (const TSharedRef<FDetailTreeNode>& Child : Children)
 		{
-			TArray<FString> TypesToIgnoreForRecursiveExpansion;
-			GConfig->GetArray(TEXT("DetailPropertyExpansion"), TEXT("TypesToIgnoreForRecursiveExpansion"), TypesToIgnoreForRecursiveExpansion, GEditorPerProjectIni);
-
-			for (const TSharedRef<FDetailTreeNode>& Child : Children)
+			if (bExpand)
 			{
-				if (bExpand)
+				const FProperty* ChildProperty = Child->GetPropertyNode().IsValid() ? Child->GetPropertyNode()->GetProperty() : nullptr;
+				if (ChildProperty != nullptr)
 				{
-					const FProperty* ChildProperty = Child->GetPropertyNode().IsValid() ? Child->GetPropertyNode()->GetProperty() : nullptr;
-					if (ChildProperty != nullptr)
+					if (TypesToIgnoreForRecursiveExpansion.Contains(ChildProperty->GetClass()->GetName()))
 					{
-						if (TypesToIgnoreForRecursiveExpansion.Contains(ChildProperty->GetClass()->GetName()))
-						{
-							continue;
-						}
+						continue;
 					}
 				}
+			}
 
 
-				SetNodeExpansionState(Child, bExpand, bRecursive);
-			}
+			SetNodeExpansionState(Child, bExpand, bRecursive);
 		}
 	}
 }
@@ -1402,6 +1441,16 @@ void SDetailsViewBase::Tick( const FGeometry& AllottedGeometry, const double InC
 	}
 }
 
+FReply SDetailsViewBase::OnKeyDown(const FGeometry& InGeometry, const FKeyEvent& InKeyEvent)
+{
+	if (InternalCommands.IsValid() && InternalCommands->ProcessCommandBindings(InKeyEvent))
+	{
+		return FReply::Handled();
+	}
+
+	return IDetailsViewPrivate::OnKeyDown(InGeometry, InKeyEvent);
+}
+
 SDetailsViewBase::FAnimatingNodeCollection::FAnimatingNodeCollection(
 	const TConstArrayView<TSharedPtr<FPropertyPath>>& InNodePaths,
 	const FGuid& InBatchId)
@@ -1900,6 +1949,8 @@ void SDetailsViewBase::UpdateFilteredDetails()
 	TRACE_CPUPROFILER_EVENT_SCOPE(SDetailsViewBase::UpdateFilteredDetails);
 
 	RootTreeNodes.Reset();
+	
+	const FDetailFilter Filter = CreateDetailFilter();
 
 	FDetailNodeList InitialRootNodeList;
 	NumVisibleTopLevelObjectNodes = 0;
@@ -1919,7 +1970,7 @@ void SDetailsViewBase::UpdateFilteredDetails()
 					FilterRootNode(ExternalRootNode);
 				}
 
-				DetailLayout->FilterDetailLayout(CurrentFilter);
+				DetailLayout->FilterDetailLayout(Filter);
 
 				const FDetailNodeList& LayoutRoots = DetailLayout->GetFilteredRootTreeNodes();
 				if (LayoutRoots.Num() > 0)
diff --git a/Engine/Source/Editor/PropertyEditor/Private/SDetailsViewBase.h b/Engine/Source/Editor/PropertyEditor/Private/SDetailsViewBase.h
--- a/Engine/Source/Editor/PropertyEditor/Private/SDetailsViewBase.h
+++ b/Engine/Source/Editor/PropertyEditor/Private/SDetailsViewBase.h
@@ -194,11 +194,12 @@ public:
 	 */
 	virtual TSharedPtr<FDetailsDisplayManager> GetDisplayManager() override;
 
-	// SWidget interface
+	//~ Begin SWidget
 	virtual bool SupportsKeyboardFocus() const override;
 	virtual FReply OnFocusReceived(const FGeometry& MyGeometry, const FFocusEvent& InFocusEvent) override;
 	virtual void Tick( const FGeometry& AllottedGeometry, const double InCurrentTime, const float InDeltaTime ) override;
-	// End of SWidget interface
+	virtual FReply OnKeyDown(const FGeometry& InGeometry, const FKeyEvent& InKeyEvent) override;
+	//~ End SWidget
 
 	// IScrollableWidget interface
 	virtual FVector2D GetScrollDistance() override;
@@ -474,6 +475,11 @@ protected:
 
 	/** Clear the keyboard focus if it's inside the given widget. */
 	void ClearKeyboardFocusIfWithin(const TSharedRef<SWidget>& Widget) const;
+	
+	/** Creates details filter before filtering nodes */
+	FDetailFilter CreateDetailFilter();
+	
+	void BindCommands();
 
 protected:
 	/** The user defined args for the details view */
@@ -596,4 +602,7 @@ protected:
 
 	/** The EnsureDataIsValid function can be skipped with this member, if set. Useful if your implementation doesn't require this kind of validation each Tick. */
 	FOnValidateDetailsViewPropertyNodes CustomValidatePropertyNodesFunction;
+	
+	/** Internal details commands */
+	TSharedPtr<FUICommandList> InternalCommands;
 };
diff --git a/Engine/Source/Editor/PropertyEditor/Private/SStructureDetailsView.cpp b/Engine/Source/Editor/PropertyEditor/Private/SStructureDetailsView.cpp
--- a/Engine/Source/Editor/PropertyEditor/Private/SStructureDetailsView.cpp
+++ b/Engine/Source/Editor/PropertyEditor/Private/SStructureDetailsView.cpp
@@ -1,22 +1,22 @@
 // Copyright Epic Games, Inc. All Rights Reserved.
 
 #include "SStructureDetailsView.h"
-#include "Framework/MultiBox/MultiBoxBuilder.h"
-#include "Widgets/Images/SImage.h"
-#include "Widgets/Images/SLayeredImage.h"
-#include "Widgets/Input/SComboButton.h"
+
 #include "AssetSelection.h"
 #include "DetailCategoryBuilderImpl.h"
-#include "UserInterface/PropertyDetails/PropertyDetailsUtilities.h"
+#include "DetailsViewPropertyGenerationUtilities.h"
+#include "Framework/MultiBox/MultiBoxBuilder.h"
+#include "PropertyEditorCommands.h"
 #include "StructurePropertyNode.h"
+#include "UserInterface/PropertyDetails/PropertyDetailsUtilities.h"
 #include "Widgets/Colors/SColorPicker.h"
+#include "Widgets/Images/SImage.h"
+#include "Widgets/Images/SLayeredImage.h"
+#include "Widgets/Input/SComboButton.h"
 #include "Widgets/Input/SSearchBox.h"
-#include "DetailsViewPropertyGenerationUtilities.h"
-
 
 #define LOCTEXT_NAMESPACE "SStructureDetailsView"
 
-
 SStructureDetailsView::~SStructureDetailsView()
 {
 	auto RootNodeLocal = GetRootNode();
@@ -41,6 +41,8 @@ void SStructureDetailsView::Construct(const FArguments& InArgs)
 	ColumnSizeData.SetRightColumnMinWidth(DetailsViewArgs.RightColumnMinWidth);
 
 	CustomName = InArgs._CustomName;
+	
+	BindCommands();
 
 	// Create the root property now
 	// Only one root node in a structure details view
@@ -55,7 +57,7 @@ void SStructureDetailsView::Construct(const FArguments& InArgs)
 	// See note in SDetailsView for why visibility is set after construction
 	ExternalScrollbar->SetVisibility( TAttribute<EVisibility>(this, &SStructureDetailsView::GetScrollBarVisibility) );
 
-	FMenuBuilder DetailViewOptions( true, nullptr );
+	FMenuBuilder DetailViewOptions(true, InternalCommands);
 
 	FUIAction ShowOnlyModifiedAction( 
 		FExecuteAction::CreateSP(this, &SStructureDetailsView::OnShowOnlyModifiedClicked),
@@ -119,16 +121,9 @@ void SStructureDetailsView::Construct(const FArguments& InArgs)
 		EUserInterfaceActionType::Check
 		);
 
-	DetailViewOptions.AddMenuEntry(
-		LOCTEXT("CollapseAll", "Collapse All Categories"),
-		LOCTEXT("CollapseAll_ToolTip", "Collapses all root level categories"),
-		FSlateIcon(),
-		FUIAction(FExecuteAction::CreateSP(this, &SStructureDetailsView::SetRootExpansionStates, /*bExpanded=*/false, /*bRecurse=*/false)));
-	DetailViewOptions.AddMenuEntry(
-		LOCTEXT("ExpandAll", "Expand All Categories"),
-		LOCTEXT("ExpandAll_ToolTip", "Expands all root level categories"),
-		FSlateIcon(),
-		FUIAction(FExecuteAction::CreateSP(this, &SStructureDetailsView::SetRootExpansionStates, /*bExpanded=*/true, /*bRecurse=*/false)));
+	// Commands have been registered earlier
+	DetailViewOptions.AddMenuEntry(FPropertyEditorCommands::Get().CollapseAll);
+	DetailViewOptions.AddMenuEntry(FPropertyEditorCommands::Get().ExpandAll);
 
 	TSharedRef<SHorizontalBox> FilterBoxRow = SNew( SHorizontalBox )
 		.Visibility(this, &SStructureDetailsView::GetFilterBoxVisibility)
diff --git a/Engine/Source/Editor/PropertyEditor/Public/PropertyEditorCommands.h b/Engine/Source/Editor/PropertyEditor/Public/PropertyEditorCommands.h
new file mode 100644
--- /dev/null
+++ b/Engine/Source/Editor/PropertyEditor/Public/PropertyEditorCommands.h
@@ -0,0 +1,21 @@
+// Copyright Epic Games, Inc. All Rights Reserved.
+
+#pragma once
+
+#include "Framework/Commands/Commands.h"
+#include "Templates/SharedPointer.h"
+
+class FUICommandInfo;
+
+class FPropertyEditorCommands : public TCommands<FPropertyEditorCommands>
+{
+public:
+	FPropertyEditorCommands();
+
+	//~ Begin TCommands
+	virtual void RegisterCommands() override;
+	//~ End TCommands
+
+	TSharedPtr<FUICommandInfo> CollapseAll;
+	TSharedPtr<FUICommandInfo> ExpandAll;
+};
\ No newline at end of file
