Fixed ToolTip Override that was not working for FPropertyHandleBase::CreatePropertyNameWidget since 48526812

#Synced-CL 49356504

diff --git a/Engine/Source/Editor/PropertyEditor/Private/PropertyEditorHelpers.cpp b/Engine/Source/Editor/PropertyEditor/Private/PropertyEditorHelpers.cpp
--- a/Engine/Source/Editor/PropertyEditor/Private/PropertyEditorHelpers.cpp
+++ b/Engine/Source/Editor/PropertyEditor/Private/PropertyEditorHelpers.cpp
@@ -41,15 +41,17 @@ void SPropertyNameWidget::Construct( const FArguments& InArgs, TSharedPtr<FPrope
 	PropertyEditor = InPropertyEditor;
 
 	TWeakPtr<FPropertyEditor> WeakPropertyEditor = PropertyEditor;
-	auto CreateDocumentationToolTipDeferred = [WeakPropertyEditor]()
+	// Note : ToolTipText has to be captured immediately because FPropertyHandleBase::CreatePropertyNameWidget has a ToolTipOverride that cannot be deferred since 
+	//  it's restored to FText::Empty right after SPropertyNameWidget::Construct :
+	auto CreateDocumentationToolTipDeferred = [WeakPropertyEditor, ToolTipText = PropertyEditor->GetToolTipText()]() 
 	{
 		if (TSharedPtr<FPropertyEditor> PropertyEditorPinned = WeakPropertyEditor.Pin())
 		{
-			return IDocumentation::Get()->CreateToolTip(PropertyEditorPinned->GetToolTipText(), NULL, PropertyEditorPinned->GetDocumentationLink(), PropertyEditorPinned->GetDocumentationExcerptName());
+			return IDocumentation::Get()->CreateToolTip(ToolTipText, NULL, PropertyEditorPinned->GetDocumentationLink(), PropertyEditorPinned->GetDocumentationExcerptName());
 		};
 
 		return SNew(SToolTip)
-			.Text(LOCTEXT("EditorDestroyedToolTip", "Invalid ToolTip, source editor closed."));
+		.Text(LOCTEXT("EditorDestroyedToolTip", "Invalid ToolTip, source editor closed."));
 	};
 	TSharedRef<SDeferredToolTip> DeferredDocumentationToolTip = MakeShared<SDeferredToolTip>(FOnGetDeferredToolTip::CreateLambda(CreateDocumentationToolTipDeferred));
 
