Displayed Private DataLayers name in actor details now show the actual Data Layer Instance short name instead of the generic DataLayerAsset class name.

#Synced-CL 49565614

diff --git a/Engine/Source/Editor/PropertyEditor/Private/UserInterface/PropertyEditor/SPropertyEditorAsset.cpp b/Engine/Source/Editor/PropertyEditor/Private/UserInterface/PropertyEditor/SPropertyEditorAsset.cpp
--- a/Engine/Source/Editor/PropertyEditor/Private/UserInterface/PropertyEditor/SPropertyEditorAsset.cpp
+++ b/Engine/Source/Editor/PropertyEditor/Private/UserInterface/PropertyEditor/SPropertyEditorAsset.cpp
@@ -9,6 +9,7 @@
 #include "Containers/VersePath.h"
 #include "Engine/StaticMesh.h"
 #include "Engine/LevelScriptActor.h"
+#include "Engine/DataAsset.h"
 #include "Editor.h"
 #include "Modules/ModuleManager.h"
 #include "UObject/UObjectHash.h"
@@ -1230,6 +1231,10 @@ FText SPropertyEditorAsset::OnGetAssetName() const
 			{
 				Name = AsField->GetDisplayNameText();
 			}
+			else if (UDataAsset* DataAsset = Cast<UDataAsset>(Value.Object))
+			{
+				Name = DataAsset->GetDisplayNameText();
+			}
 			else
 			{
 				Name = FText::AsCultureInvariant(Value.Object->GetName());
diff --git a/Engine/Source/Runtime/Engine/Classes/Engine/DataAsset.h b/Engine/Source/Runtime/Engine/Classes/Engine/DataAsset.h
--- a/Engine/Source/Runtime/Engine/Classes/Engine/DataAsset.h
+++ b/Engine/Source/Runtime/Engine/Classes/Engine/DataAsset.h
@@ -27,6 +27,10 @@ public:
 	ENGINE_API virtual void Serialize(FStructuredArchiveRecord Record) override;
 #endif
 
+#if WITH_EDITOR
+	ENGINE_API virtual FText GetDisplayNameText() const;
+#endif
+
 	/** Asset registry tag NativeClass with full path in export form Class'/Script/ModuleName.ClassName' of the first found parent native class */
 	static ENGINE_API const FName NativeClassTag;
 };
diff --git a/Engine/Source/Runtime/Engine/Private/DataAsset.cpp b/Engine/Source/Runtime/Engine/Private/DataAsset.cpp
--- a/Engine/Source/Runtime/Engine/Private/DataAsset.cpp
+++ b/Engine/Source/Runtime/Engine/Private/DataAsset.cpp
@@ -27,6 +27,13 @@ void UDataAsset::GetAssetRegistryTags(FAssetRegistryTagsContext Context) const
 	Super::GetAssetRegistryTags(Context);
 }
 
+#if WITH_EDITOR
+FText UDataAsset::GetDisplayNameText() const
+{
+	return FText::AsCultureInvariant(GetName());
+}
+#endif
+
 #if WITH_EDITORONLY_DATA
 void UDataAsset::Serialize(FStructuredArchive::FRecord Record)
 {
diff --git a/Engine/Source/Runtime/Engine/Private/WorldPartition/DataLayer/DataLayerAsset.cpp b/Engine/Source/Runtime/Engine/Private/WorldPartition/DataLayer/DataLayerAsset.cpp
--- a/Engine/Source/Runtime/Engine/Private/WorldPartition/DataLayer/DataLayerAsset.cpp
+++ b/Engine/Source/Runtime/Engine/Private/WorldPartition/DataLayer/DataLayerAsset.cpp
@@ -71,6 +71,18 @@ bool UDataLayerAsset::CanEditChange(const FProperty* InProperty) const
 	return true;
 }
 
+FText UDataLayerAsset::GetDisplayNameText() const
+{
+	if (IsPrivate())
+	{
+		if (const UDataLayerInstance* DataLayerInstance = GetTypedOuter<UDataLayerInstance>())
+		{
+			return FText::FromString(DataLayerInstance->GetDataLayerShortName());
+		}
+	}
+	return Super::GetDisplayNameText();
+}
+
 void UDataLayerAsset::OnCreated()
 {
 	SetDebugColor(FColor::MakeRandomSeededColor(GetTypeHash(GetFullName())));
diff --git a/Engine/Source/Runtime/Engine/Public/WorldPartition/DataLayer/DataLayerAsset.h b/Engine/Source/Runtime/Engine/Public/WorldPartition/DataLayer/DataLayerAsset.h
--- a/Engine/Source/Runtime/Engine/Public/WorldPartition/DataLayer/DataLayerAsset.h
+++ b/Engine/Source/Runtime/Engine/Public/WorldPartition/DataLayer/DataLayerAsset.h
@@ -38,6 +38,10 @@ class UDataLayerAsset : public UDataAsset
 	ENGINE_API virtual bool CanEditChange(const FProperty* InProperty) const override;
 	//~ End UObject interface
 
+	//~ Begin UDataAsset interface
+	ENGINE_API virtual FText GetDisplayNameText() const override;
+	//~ End UDataAsset interface
+
 public:
 	virtual bool CanEditDataLayerType() const;
 	ENGINE_API void SetType(EDataLayerType Type);
