[AvalancheTests] Automate 'Sequence - Verify Direction and Role can be changed'

#rnx

#Synced-CL 45465245

diff --git a/Engine/Plugins/VirtualProduction/CustomDetailsView/Source/CustomDetailsView/Private/SCustomDetailsView.h b/Engine/Plugins/VirtualProduction/CustomDetailsView/Source/CustomDetailsView/Private/SCustomDetailsView.h
--- a/Engine/Plugins/VirtualProduction/CustomDetailsView/Source/CustomDetailsView/Private/SCustomDetailsView.h
+++ b/Engine/Plugins/VirtualProduction/CustomDetailsView/Source/CustomDetailsView/Private/SCustomDetailsView.h
@@ -57,7 +57,7 @@ public:
 
 	void OnExpansionChanged(TSharedPtr<ICustomDetailsViewItem> InItem, bool bInExpanded);
 
-	void SetExpansionRecursive(TSharedPtr<ICustomDetailsViewItem> InItem, bool bInExpand);
+	DLLEXPORT void SetExpansionRecursive(TSharedPtr<ICustomDetailsViewItem> InItem, bool bInExpand);
 
 	bool ShouldItemExpand(const TSharedPtr<ICustomDetailsViewItem>& InItem) const;
 
diff --git a/Engine/Source/Editor/PropertyEditor/Private/UserInterface/PropertyEditor/SPropertyComboBox.h b/Engine/Source/Editor/PropertyEditor/Private/UserInterface/PropertyEditor/SPropertyComboBox.h
--- a/Engine/Source/Editor/PropertyEditor/Private/UserInterface/PropertyEditor/SPropertyComboBox.h
+++ b/Engine/Source/Editor/PropertyEditor/Private/UserInterface/PropertyEditor/SPropertyComboBox.h
@@ -12,7 +12,7 @@
 #include "Styling/AppStyle.h"
 #include "Templates/SharedPointer.h"
 #include "Types/SlateEnums.h"
-#include "UserInterface/PropertyEditor/PropertyEditorConstants.h"
+#include "Editor/PropertyEditor/Private/UserInterface/PropertyEditor/PropertyEditorConstants.h"
 #include "Widgets/DeclarativeSyntaxSupport.h"
 #include "Widgets/Input/SComboBox.h"
 
@@ -48,7 +48,7 @@ public:
 	 *	Sets the currently selected item for the combo box.
 	 *	@param InSelectedItem			The name of the item to select.
 	 */
-	void SetSelectedItem( const FString& InSelectedItem );
+	DLLEXPORT void SetSelectedItem( const FString& InSelectedItem );
 
 	/** 
 	 *	Sets the item list for the combo box.
diff --git a/Engine/Source/Editor/Sequencer/Private/Sequencer.h b/Engine/Source/Editor/Sequencer/Private/Sequencer.h
--- a/Engine/Source/Editor/Sequencer/Private/Sequencer.h
+++ b/Engine/Source/Editor/Sequencer/Private/Sequencer.h
@@ -502,7 +502,7 @@ public:
 	/** Functions to push on to the transport controls we use */
 	FReply OnRecord();
 	FReply OnPlayForward(bool bTogglePlay);
-	FReply OnPlayBackward(bool bTogglePlay);
+	DLLEXPORT FReply OnPlayBackward(bool bTogglePlay);
 	FReply OnStepForward(FFrameNumber Increment = FFrameNumber(1));
 	FReply OnStepBackward(FFrameNumber Increment = FFrameNumber(1));
 	FReply OnJumpToStart();
@@ -524,7 +524,7 @@ public:
 
 	ESequencerLoopMode GetLoopMode() const;
 
-	EPlaybackMode::Type GetPlaybackMode() const;
+	DLLEXPORT EPlaybackMode::Type GetPlaybackMode() const;
 
 	/** @return The toolkit that this sequencer is hosted in (if any) */
 	virtual TSharedPtr<IToolkitHost> GetToolkitHost() const override { return ToolkitHost.Pin(); }
