Fixed missing objects in property change event for edit condition toggling

#Synced-CL 49840944

diff --git a/Engine/Source/Editor/PropertyEditor/Private/Presentation/PropertyEditor/PropertyEditor.cpp b/Engine/Source/Editor/PropertyEditor/Private/Presentation/PropertyEditor/PropertyEditor.cpp
--- a/Engine/Source/Editor/PropertyEditor/Private/Presentation/PropertyEditor/PropertyEditor.cpp
+++ b/Engine/Source/Editor/PropertyEditor/Private/Presentation/PropertyEditor/PropertyEditor.cpp
@@ -615,12 +615,18 @@ void FPropertyEditor::ToggleEditConditionState()
 	TArray<TMap<FString,int32>> ArrayIndicesPerObject;
 	ArrayIndicesPerObject.AddDefaulted(ComplexParentNode->GetInstancesNum());
 
+	TArray<UObject*> InstanceObjects;
+	InstanceObjects.Reserve(ComplexParentNode->GetInstancesNum());
 	for (int32 ObjectIndex = 0; ObjectIndex < ComplexParentNode->GetInstancesNum(); ++ObjectIndex)
 	{
 		FPropertyValueImpl::GenerateArrayIndexMapToObjectNode(ArrayIndicesPerObject[ObjectIndex], &PropertyNode.Get());
+		if (TWeakObjectPtr<UObject> Object = ComplexParentNode->GetInstanceAsUObject(ObjectIndex); Object.IsValid())
+		{
+			InstanceObjects.Add(Object.Get());
+		}
 	}
 
-	FPropertyChangedEvent ChangeEvent(PropertyNode->GetProperty(), EPropertyChangeType::ToggleEditable);
+	FPropertyChangedEvent ChangeEvent(PropertyNode->GetProperty(), EPropertyChangeType::ToggleEditable, InstanceObjects);
 	ChangeEvent.SetArrayIndexPerObject(ArrayIndicesPerObject);
 	PropertyNode->NotifyPostChange( ChangeEvent, PropertyUtilities->GetNotifyHook() );
 	PropertyUtilities->NotifyFinishedChangingProperties(ChangeEvent);
