Add a callback API to disable property permission lists for packages if desired

#Synced-CL 49939927

diff --git a/Engine/Plugins/Runtime/ModelViewViewModel/Source/ModelViewViewModelBlueprint/Private/MVVMDeveloperProjectSettings.cpp b/Engine/Plugins/Runtime/ModelViewViewModel/Source/ModelViewViewModelBlueprint/Private/MVVMDeveloperProjectSettings.cpp
--- a/Engine/Plugins/Runtime/ModelViewViewModel/Source/ModelViewViewModelBlueprint/Private/MVVMDeveloperProjectSettings.cpp
+++ b/Engine/Plugins/Runtime/ModelViewViewModel/Source/ModelViewViewModelBlueprint/Private/MVVMDeveloperProjectSettings.cpp
@@ -150,7 +150,7 @@ bool UMVVMDeveloperProjectSettings::IsPropertyAllowed(const UBlueprint* Generati
 	const bool bDoPropertyEditorPermission = UE::MVVM::Private::ShouldDoFieldEditorPermission(GeneratingFor, AuthoritativeClass, Property->GetOwnerClass());
 	if (bDoPropertyEditorPermission)
 	{
-		if (!FPropertyEditorPermissionList::Get().DoesPropertyPassFilter(AuthoritativeClass, Property->GetFName()))
+		if (FPropertyEditorPermissionList::Get().IsEnabledForPackage(GeneratingFor->GetPackage()) && !FPropertyEditorPermissionList::Get().DoesPropertyPassFilter(AuthoritativeClass, Property->GetFName()))
 		{
 			return false;
 		}
diff --git a/Engine/Source/Editor/PropertyEditor/Private/PropertyPermissionList.cpp b/Engine/Source/Editor/PropertyEditor/Private/PropertyPermissionList.cpp
--- a/Engine/Source/Editor/PropertyEditor/Private/PropertyPermissionList.cpp
+++ b/Engine/Source/Editor/PropertyEditor/Private/PropertyPermissionList.cpp
@@ -458,6 +458,11 @@ FPropertyEditorPermissionList& FPropertyEditorPermissionList::Get()
 	return PermissionList;
 }
 
+bool FPropertyEditorPermissionList::IsEnabledForPackage(const UPackage* Package) const
+{
+	return !IsEnabledForPackageDelegate.IsBound() || IsEnabledForPackageDelegate.Execute(Package);
+}
+
 FHiddenPropertyPermissionList& FHiddenPropertyPermissionList::Get()
 {
 	static FHiddenPropertyPermissionList PermissionList;
diff --git a/Engine/Source/Editor/PropertyEditor/Public/PropertyPermissionList.h b/Engine/Source/Editor/PropertyEditor/Public/PropertyPermissionList.h
--- a/Engine/Source/Editor/PropertyEditor/Public/PropertyPermissionList.h
+++ b/Engine/Source/Editor/PropertyEditor/Public/PropertyPermissionList.h
@@ -16,6 +16,7 @@ DECLARE_LOG_CATEGORY_EXTERN(LogPropertyEditorPermissionList, Log, All);
 
 /** Struct, OwnerName */
 DECLARE_MULTICAST_DELEGATE_TwoParams(FPermissionListUpdated, TSoftObjectPtr<const UStruct>, FName);
+DECLARE_DELEGATE_RetVal_OneParam(bool, FIsEnabledForPackage, const UPackage*);
 DECLARE_DELEGATE_RetVal_OneParam(TOptional<bool>, FCustomDetailTreeNodePermissionDelegate, TSharedRef<FDetailTreeNode> DetailTreeNode);
 
 /**
@@ -175,11 +176,17 @@ public:
 	bool ShouldShowMenuEntries() const { return bShouldShowMenuEntries; }
 	/** Turn on or off menu entries to modify the PermissionList from a Details View */
 	void SetShouldShowMenuEntries(bool bShow) { bShouldShowMenuEntries = bShow; }
+	/** Set the delegate that is called in IsEnabledForPackage */
+	void SetIsEnabledForPackageDelegate(const FIsEnabledForPackage& InIsEnabledForPackage ) { IsEnabledForPackageDelegate = InIsEnabledForPackage; }
+	/** A special call to disable the filter for a specific blueprint, if the IsEnabledForPackageDelegate is set and returns false */
+	UE_API bool IsEnabledForPackage(const UPackage* Package) const;
 
 private:
 	FPropertyEditorPermissionList() = default;
 	~FPropertyEditorPermissionList() = default;
 
+	FIsEnabledForPackage IsEnabledForPackageDelegate;
+
 	/** Whether SDetailSingleItemRow should add menu items to add/remove properties to/from the PermissionList */
 	bool bShouldShowMenuEntries = false;
 };
