Made CachedPropertyPermissionList thread-safe, so we can query the property permissions from multiple threads

#rnx

#Synced-CL 49940768

diff --git a/Engine/Source/Editor/PropertyEditor/Private/PropertyPermissionList.cpp b/Engine/Source/Editor/PropertyEditor/Private/PropertyPermissionList.cpp
--- a/Engine/Source/Editor/PropertyEditor/Private/PropertyPermissionList.cpp
+++ b/Engine/Source/Editor/PropertyEditor/Private/PropertyPermissionList.cpp
@@ -257,6 +257,8 @@ void FPropertyPermissionList::SetEnabled(const FDelegateEnablePathFilter& bEnabl
 
 void FPropertyPermissionList::ClearCache()
 {
+	FWriteScopeLock _(CachedPropertyPermissionListRW);
+
 	CachedPropertyPermissionList.Reset();
 }
 
@@ -314,13 +316,15 @@ const FPropertyPermissionListEntry& FPropertyPermissionList::GetCachedPermission
 	}
 #endif
 
-	// Is this struct already cached? If so, we can just return its result
-	if (const FPropertyPermissionListEntry* CachedPermissionList = CachedPropertyPermissionList.Find(Struct))
 	{
-		return *CachedPermissionList;
-	}
+		FReadScopeLock _(CachedPropertyPermissionListRW);
 
-	const FPropertyPermissionListEntry* Entry = RawPropertyPermissionList.Find(Struct);
+		// Is this struct already cached? If so, we can just return its result
+		if (const TUniquePtr<FPropertyPermissionListEntry>* CachedPermissionListPtr = CachedPropertyPermissionList.Find(Struct))
+		{
+			return **CachedPermissionListPtr;
+		}
+	}
 
 	FNamePermissionList NewPermissionList;
 
@@ -356,6 +360,7 @@ const FPropertyPermissionListEntry& FPropertyPermissionList::GetCachedPermission
 	}
 
 	// Resolve the permission list rule for this entry
+	const FPropertyPermissionListEntry* Entry = RawPropertyPermissionList.Find(Struct);
 	EPropertyPermissionListRules EntryRule = EPropertyPermissionListRules::UseExistingPermissionList;
 	if (Entry)
 	{
@@ -420,11 +425,21 @@ const FPropertyPermissionListEntry& FPropertyPermissionList::GetCachedPermission
 		}
 	}
 
-	return CachedPropertyPermissionList.Add(Struct, {
-		MoveTemp(NewPermissionList),
-		TArray<FName>(), // Note: The AdditionalOwnerNames are only used to prevent removal from RawPropertyPermissionList by calls to RemoveOwner()
-		EntryRuleToPropagate
-		});
+	{
+		FWriteScopeLock _(CachedPropertyPermissionListRW);
+
+		// Check the cache again in case another thread beat us to it
+		if (const TUniquePtr<FPropertyPermissionListEntry>* CachedPermissionListPtr = CachedPropertyPermissionList.Find(Struct))
+		{
+			return **CachedPermissionListPtr;
+		}
+
+		return *CachedPropertyPermissionList.Add(Struct, MakeUnique<FPropertyPermissionListEntry>(
+			MoveTemp(NewPermissionList),
+			TArray<FName>(), // Note: The AdditionalOwnerNames are only used to prevent removal from RawPropertyPermissionList by calls to RemoveOwner()
+			EntryRuleToPropagate
+			));
+	}
 }
 
 bool FPropertyPermissionList::HasSpecificList(const UStruct* ObjectStruct) const
diff --git a/Engine/Source/Editor/PropertyEditor/Public/PropertyPermissionList.h b/Engine/Source/Editor/PropertyEditor/Public/PropertyPermissionList.h
--- a/Engine/Source/Editor/PropertyEditor/Public/PropertyPermissionList.h
+++ b/Engine/Source/Editor/PropertyEditor/Public/PropertyPermissionList.h
@@ -160,7 +160,8 @@ private:
 	friend uint32 GetTypeHash(const FPermissionListUpdate& PermisisonList);
 	TSet< FPermissionListUpdate > PendingUpdates;
 	/** Lazily-constructed combined cache of both the flattened class PermissionList and struct PermissionList */
-	mutable TMap<TWeakObjectPtr<const UStruct>, FPropertyPermissionListEntry> CachedPropertyPermissionList;
+	mutable TMap<TWeakObjectPtr<const UStruct>, TUniquePtr<FPropertyPermissionListEntry>> CachedPropertyPermissionList;
+	mutable FRWLock CachedPropertyPermissionListRW;
 
 	/** Get or create the cached PermissionList for a specific UStruct */
 	const FNamePermissionList& GetCachedPermissionListForStruct(const UStruct* Struct) const;
