Add FEnumProperty support to FPropertyHandleInt

#Synced-CL 49757048

diff --git a/Engine/Source/Editor/PropertyEditor/Private/PropertyHandleImpl.cpp b/Engine/Source/Editor/PropertyEditor/Private/PropertyHandleImpl.cpp
--- a/Engine/Source/Editor/PropertyEditor/Private/PropertyHandleImpl.cpp
+++ b/Engine/Source/Editor/PropertyEditor/Private/PropertyHandleImpl.cpp
@@ -3872,7 +3872,7 @@ bool FPropertyHandleInt::Supports( TSharedRef<FPropertyNode> PropertyNode )
 		return false;
 	}
 
-	const bool bIsInteger = 
+	bool bIsSupported = 
 			Property->IsA<FInt8Property>()
 		||	Property->IsA<FInt16Property>()
 		||  Property->IsA<FIntProperty>()
@@ -3881,13 +3881,34 @@ bool FPropertyHandleInt::Supports( TSharedRef<FPropertyNode> PropertyNode )
 		||  Property->IsA<FUInt32Property>()
 		||  Property->IsA<FUInt64Property>();
 
-	// The value is an integer
-	return bIsInteger;
+	if (!bIsSupported)
+	{
+		if (const FEnumProperty* EnumProperty = CastField<FEnumProperty>(Property))
+		{
+			// We only support an enum if it's underlying type is NOT uint8, which is instead handled by FPropertyHandleByte
+			const FProperty* UnderlyingProperty = EnumProperty->GetUnderlyingProperty();
+			bIsSupported = UnderlyingProperty && !UnderlyingProperty->IsA<FByteProperty>();
+		}
+	}
+
+	// The value is an integer (or enum that isn't uint8-based)
+	return bIsSupported;
 }
 
 template <typename PropertyClass, typename ValueType>
 ValueType GetIntegerValue(void* PropValue, FPropertyValueImpl& Implementation )
 {
+	if (Implementation.IsPropertyTypeOf(FEnumProperty::StaticClass()))
+	{
+		const FEnumProperty* EnumProperty = CastFieldChecked<FEnumProperty>(Implementation.GetPropertyNode()->GetProperty());
+		const FNumericProperty* UnderlyingProperty = EnumProperty->GetUnderlyingProperty();
+		const bool bIsSignedInteger = UnderlyingProperty->CanHoldValue(-1);
+		if (bIsSignedInteger)
+		{
+			return static_cast<ValueType>(UnderlyingProperty->GetSignedIntPropertyValue(PropValue));	
+		}
+		return static_cast<ValueType>(UnderlyingProperty->GetUnsignedIntPropertyValue(PropValue));
+	}
 	check( Implementation.IsPropertyTypeOf( PropertyClass::StaticClass() ) );
 	return Implementation.GetPropertyValue<PropertyClass>(PropValue);
 }
@@ -3988,7 +4009,15 @@ FPropertyAccess::Result FPropertyHandleInt::SetValue(const int8& NewValue, EProp
 	// Clamp the value from any meta data ranges stored on the property value
 	int8 FinalValue = ClampIntegerValueFromMetaData<int8>( NewValue, *this, *Implementation->GetPropertyNode());
 
-	const FString ValueStr = LexToString(FinalValue);
+	FString ValueStr;
+	if (const FEnumProperty* EnumProp = CastField<FEnumProperty>(GetProperty()))
+	{
+		ValueStr = EnumProp->GetEnum()->GetNameStringByValue(FinalValue);
+	}
+	else
+	{
+		ValueStr = LexToString(FinalValue);
+	}
 	Res = Implementation->ImportText(ValueStr, Flags);
 
 	return Res;
@@ -4001,7 +4030,15 @@ FPropertyAccess::Result FPropertyHandleInt::SetValue(const int16& NewValue, EPro
 	// Clamp the value from any meta data ranges stored on the property value
 	int16 FinalValue = ClampIntegerValueFromMetaData<int16>(NewValue, *this, *Implementation->GetPropertyNode());
 
-	const FString ValueStr = LexToString(FinalValue);
+	FString ValueStr;
+	if (const FEnumProperty* EnumProperty = CastField<FEnumProperty>(GetProperty()))
+	{
+		ValueStr = EnumProperty->GetEnum()->GetNameStringByValue(FinalValue);
+	}
+	else
+	{
+		ValueStr = LexToString(FinalValue);
+	}
 	Res = Implementation->ImportText(ValueStr, Flags);
 
 	return Res;
@@ -4014,7 +4051,15 @@ FPropertyAccess::Result FPropertyHandleInt::SetValue( const int32& NewValue, EPr
 	// Clamp the value from any meta data ranges stored on the property value
 	int32 FinalValue = ClampIntegerValueFromMetaData<int32>( NewValue, *this, *Implementation->GetPropertyNode());
 
-	const FString ValueStr = LexToString(FinalValue);
+	FString ValueStr;
+	if (const FEnumProperty* EnumProperty = CastField<FEnumProperty>(GetProperty()))
+	{
+		ValueStr = EnumProperty->GetEnum()->GetNameStringByValue(FinalValue);
+	}
+	else
+	{
+		ValueStr = LexToString(FinalValue);
+	}
 	Res = Implementation->ImportText( ValueStr, Flags );
 
 	return Res;
@@ -4027,7 +4072,15 @@ FPropertyAccess::Result FPropertyHandleInt::SetValue(const int64& NewValue, EPro
 	// Clamp the value from any meta data ranges stored on the property value
 	int64 FinalValue = ClampIntegerValueFromMetaData<int64>(NewValue, *this, *Implementation->GetPropertyNode());
 
-	const FString ValueStr = LexToString(FinalValue);
+	FString ValueStr;
+	if (const FEnumProperty* EnumProperty = CastField<FEnumProperty>(GetProperty()))
+	{
+		ValueStr = EnumProperty->GetEnum()->GetNameStringByValue(FinalValue);
+	}
+	else
+	{
+		ValueStr = LexToString(FinalValue);
+	}
 	Res = Implementation->ImportText(ValueStr, Flags);
 	return Res;
 }
@@ -4038,7 +4091,15 @@ FPropertyAccess::Result FPropertyHandleInt::SetValue(const uint16& NewValue, EPr
 	// Clamp the value from any meta data ranges stored on the property value
 	uint16 FinalValue = ClampIntegerValueFromMetaData<uint16>(NewValue, *this, *Implementation->GetPropertyNode());
 
-	const FString ValueStr = LexToString(FinalValue);
+	FString ValueStr;
+	if (const FEnumProperty* EnumProperty = CastField<FEnumProperty>(GetProperty()))
+	{
+		ValueStr = EnumProperty->GetEnum()->GetNameStringByValue(FinalValue);
+	}
+	else
+	{
+		ValueStr = LexToString(FinalValue);
+	}
 	Res = Implementation->ImportText(ValueStr, Flags);
 
 	return Res;
@@ -4051,7 +4112,15 @@ FPropertyAccess::Result FPropertyHandleInt::SetValue(const uint32& NewValue, EPr
 	// Clamp the value from any meta data ranges stored on the property value
 	uint32 FinalValue = ClampIntegerValueFromMetaData<uint32>(NewValue, *this, *Implementation->GetPropertyNode());
 
-	const FString ValueStr = LexToString(FinalValue);
+	FString ValueStr;
+	if (const FEnumProperty* EnumProp = CastField<FEnumProperty>(GetProperty()))
+	{
+		ValueStr = EnumProp->GetEnum()->GetNameStringByValue(FinalValue);
+	}
+	else
+	{
+		ValueStr = LexToString(FinalValue);
+	}
 	Res = Implementation->ImportText(ValueStr, Flags);
 
 	return Res;
@@ -4063,7 +4132,15 @@ FPropertyAccess::Result FPropertyHandleInt::SetValue(const uint64& NewValue, EPr
 	// Clamp the value from any meta data ranges stored on the property value
 	uint64 FinalValue = ClampIntegerValueFromMetaData<uint64>(NewValue, *this, *Implementation->GetPropertyNode());
 
-	const FString ValueStr = LexToString(FinalValue);
+	FString ValueStr;
+	if (const FEnumProperty* EnumProperty = CastField<FEnumProperty>(GetProperty()))
+	{
+		ValueStr = EnumProperty->GetEnum()->GetNameStringByValue(FinalValue);
+	}
+	else
+	{
+		ValueStr = LexToString(FinalValue);
+	}
 	Res = Implementation->ImportText(ValueStr, Flags);
 	return Res;
 }
diff --git a/Engine/Source/Editor/PropertyEditor/Private/Tests/DetailsView/PropertyHandle/DetailsViewPropertyHandleEnumTest.cpp b/Engine/Source/Editor/PropertyEditor/Private/Tests/DetailsView/PropertyHandle/DetailsViewPropertyHandleEnumTest.cpp
new file mode 100644
--- /dev/null
+++ b/Engine/Source/Editor/PropertyEditor/Private/Tests/DetailsView/PropertyHandle/DetailsViewPropertyHandleEnumTest.cpp
@@ -0,0 +1,75 @@
+// Copyright Epic Games, Inc. All Rights Reserved.
+
+#if WITH_DEV_AUTOMATION_TESTS
+
+#include "DetailsViewPropertyHandleTestBase.h"
+#include "PropertyHandleImpl.h"
+
+#define DETAILS_VIEW_PROPERTY_HANDLE_ENUM_TEST(_ClassName, _TestDir) TEST_CLASS_WITH_BASE_AND_FLAGS(_ClassName, _TestDir, TDetailsViewPropertyHandleEnumTestBase, EAutomationTestFlags::EditorContext | EAutomationTestFlags::EngineFilter)
+
+template<typename Derived, typename AsserterType>
+using TDetailsViewPropertyHandleEnumTestBase = FDetailsViewPropertyHandleTestBase<Derived, AsserterType>;
+
+/**
+ * Tests that uint8-based enums resolve to FPropertyHandleByte
+ */
+DETAILS_VIEW_PROPERTY_HANDLE_ENUM_TEST(FDetailsViewPropertyHandleByteEnumTest, "Editor.PropertyEditor.DetailsView.PropertyHandleByteEnum")
+{
+	FDetailsViewPropertyHandleByteEnumTest() : FDetailsViewPropertyHandleTestBase("Properties", "TestByteEnum") {}
+
+	BEFORE_EACH()
+	{
+		FDetailsViewPropertyHandleTestBase::Setup();
+		ASSERT_THAT(IsNotNull(PropertyHandle));
+	}
+
+	TEST_METHOD(GetValue_uint8)
+	{
+		TestObject->TestByteEnum = ETestByteEnum::Value1;
+
+		uint8 Value = 0;
+		const FPropertyAccess::Result Result = PropertyHandle->GetValue(Value);
+		ASSERT_THAT(AreEqual(FPropertyAccess::Result::Success, Result, TEXT("Property handle retrieved value check")));
+		ASSERT_THAT(AreEqual(static_cast<uint8>(ETestByteEnum::Value1), Value, TEXT("The retrieved value is correct check")));
+	}
+
+	TEST_METHOD(SetValue_uint8)
+	{
+		const FPropertyAccess::Result Result = PropertyHandle->SetValue(static_cast<uint8>(ETestByteEnum::Value2));
+		ASSERT_THAT(AreEqual(FPropertyAccess::Result::Success, Result, TEXT("Property handle set value check")));
+		ASSERT_THAT(AreEqual(ETestByteEnum::Value2, TestObject->TestByteEnum, TEXT("Property of the editing object is set correctly check")));
+	}
+};
+
+/**
+ * Tests that uint32-based enums resolve to FPropertyHandleInt
+ */
+DETAILS_VIEW_PROPERTY_HANDLE_ENUM_TEST(FDetailsViewPropertyHandleIntEnumTest, "Editor.PropertyEditor.DetailsView.PropertyHandleIntEnum")
+{
+	FDetailsViewPropertyHandleIntEnumTest() : FDetailsViewPropertyHandleTestBase("Properties", "TestIntEnum") {}
+
+	BEFORE_EACH()
+	{
+		FDetailsViewPropertyHandleTestBase::Setup();
+		ASSERT_THAT(IsNotNull(PropertyHandle));
+	}
+
+	TEST_METHOD(GetValue_uint32)
+	{
+		TestObject->TestIntEnum = ETestIntEnum::Value1;
+
+		uint32 Value = 0;
+		const FPropertyAccess::Result Result = PropertyHandle->GetValue(Value);
+		ASSERT_THAT(AreEqual(FPropertyAccess::Result::Success, Result, TEXT("Property handle retrieved value check")));
+		ASSERT_THAT(AreEqual(static_cast<uint32>(ETestIntEnum::Value1), Value, TEXT("The retrieved value is correct check")));
+	}
+
+	TEST_METHOD(SetValue_uint32)
+	{
+		const FPropertyAccess::Result Result = PropertyHandle->SetValue(static_cast<uint32>(ETestIntEnum::Value2));
+		ASSERT_THAT(AreEqual(FPropertyAccess::Result::Success, Result, TEXT("Property handle set value check")));
+		ASSERT_THAT(AreEqual(ETestIntEnum::Value2, TestObject->TestIntEnum, TEXT("Property of the editing object is set correctly check")));
+	}
+};
+
+#endif // WITH_DEV_AUTOMATION_TESTS
diff --git a/Engine/Source/Editor/PropertyEditor/Private/Tests/DetailsView/PropertyHandle/DetailsViewPropertyHandleTestClass.h b/Engine/Source/Editor/PropertyEditor/Private/Tests/DetailsView/PropertyHandle/DetailsViewPropertyHandleTestClass.h
--- a/Engine/Source/Editor/PropertyEditor/Private/Tests/DetailsView/PropertyHandle/DetailsViewPropertyHandleTestClass.h
+++ b/Engine/Source/Editor/PropertyEditor/Private/Tests/DetailsView/PropertyHandle/DetailsViewPropertyHandleTestClass.h
@@ -7,6 +7,24 @@
 
 #if WITH_EDITORONLY_DATA
 
+/** Test enum with uint8 underlying type (should resolve to FPropertyHandleByte) */
+UENUM()
+enum class ETestByteEnum : uint8
+{
+	Value0 = 0,
+	Value1 = 1,
+	Value2 = 2
+};
+
+/** Test enum with uint32 underlying type (should resolve to FPropertyHandleInt) */
+UENUM()
+enum class ETestIntEnum : uint32
+{
+	Value0 = 0,
+	Value1 = 1,
+	Value2 = 2
+};
+
 UCLASS()
 class UDetailsViewPropertyHandleTestValueClass : public UObject
 {
@@ -24,6 +42,14 @@ public:
 	
 	UPROPERTY(EditAnywhere, Category = "Properties")
 	TArray<TSoftObjectPtr<UDetailsViewPropertyHandleTestValueClass>> TestValueSoftPtrArray;
+
+	/** uint8-based enum. */
+	UPROPERTY(EditAnywhere, Category = "Properties")
+	ETestByteEnum TestByteEnum = ETestByteEnum::Value0;
+
+	/** uint32-based enum */
+	UPROPERTY(EditAnywhere, Category = "Properties")
+	ETestIntEnum TestIntEnum = ETestIntEnum::Value0;
 };
 
 #endif // WITH_EDITORONLY_DATA
\ No newline at end of file
