Fixed edit change conditions not correctly gathering the edited UObject if it's an instanced struct that was added to the UI, as a complex parent was found on the way up that was not the owning object node.

#Synced-CL 50298389

diff --git a/Engine/Source/Editor/PropertyEditor/Private/Presentation/PropertyEditor/PropertyEditor.cpp b/Engine/Source/Editor/PropertyEditor/Private/Presentation/PropertyEditor/PropertyEditor.cpp
--- a/Engine/Source/Editor/PropertyEditor/Private/Presentation/PropertyEditor/PropertyEditor.cpp
+++ b/Engine/Source/Editor/PropertyEditor/Private/Presentation/PropertyEditor/PropertyEditor.cpp
@@ -609,23 +609,36 @@ void FPropertyEditor::ToggleEditConditionState()
 	PropertyNode->NotifyPreChange( PropertyNode->GetProperty(), PropertyUtilities->GetNotifyHook() );
 
 	PropertyNode->ToggleEditConditionState();
+	
+	TArray<UObject*> InstanceObjects;
 
-	const FComplexPropertyNode* ComplexParentNode = PropertyNode->FindComplexParent();
-
-	TArray<TMap<FString,int32>> ArrayIndicesPerObject;
-	ArrayIndicesPerObject.AddDefaulted(ComplexParentNode->GetInstancesNum());
+	const FPropertyNode* Cur = &PropertyNode.Get();
+	const FComplexPropertyNode* Found = nullptr;
 
-	TArray<UObject*> InstanceObjects;
-	InstanceObjects.Reserve(ComplexParentNode->GetInstancesNum());
-	for (int32 ObjectIndex = 0; ObjectIndex < ComplexParentNode->GetInstancesNum(); ++ObjectIndex)
+	TArray<TMap<FString, int32>> ArrayIndicesPerObject;
+	while( InstanceObjects.Num() == 0 )
 	{
-		FPropertyValueImpl::GenerateArrayIndexMapToObjectNode(ArrayIndicesPerObject[ObjectIndex], &PropertyNode.Get());
-		if (TWeakObjectPtr<UObject> Object = ComplexParentNode->GetInstanceAsUObject(ObjectIndex); Object.IsValid())
+		Found = Cur->AsComplexNode();
+		if( Found )
+		{
+			ArrayIndicesPerObject.AddDefaulted(Found->GetInstancesNum());
+			for (int32 ObjectIndex = 0; ObjectIndex < Found->GetInstancesNum(); ++ObjectIndex)
+			{
+				FPropertyValueImpl::GenerateArrayIndexMapToObjectNode(ArrayIndicesPerObject[ObjectIndex], &PropertyNode.Get());
+				if (TWeakObjectPtr<UObject> Object = Found->GetInstanceAsUObject(ObjectIndex); Object.IsValid())
+				{
+					InstanceObjects.Add(Object.Get());
+				}
+			}
+		}
+		Cur = Cur->GetParentNode();
+		if( !Cur )
 		{
-			InstanceObjects.Add(Object.Get());
+			// There is a break in the parent chain
+			break;
 		}
 	}
-
+	
 	FPropertyChangedEvent ChangeEvent(PropertyNode->GetProperty(), EPropertyChangeType::ToggleEditable, InstanceObjects);
 	ChangeEvent.SetArrayIndexPerObject(ArrayIndicesPerObject);
 	PropertyNode->NotifyPostChange( ChangeEvent, PropertyUtilities->GetNotifyHook() );
