Fixed set and map property editors that would not show "Multiple Values" when editing multiple objects with a different number of items (they both behave like the array property editor now)

#Synced-CL 50508017

diff --git a/Engine/Source/Editor/PropertyEditor/Private/PropertyEditorHelpers.cpp b/Engine/Source/Editor/PropertyEditor/Private/PropertyEditorHelpers.cpp
--- a/Engine/Source/Editor/PropertyEditor/Private/PropertyEditorHelpers.cpp
+++ b/Engine/Source/Editor/PropertyEditor/Private/PropertyEditorHelpers.cpp
@@ -1486,6 +1486,18 @@ namespace PropertyEditorHelpers
 
 		return nullptr;
 	}
+
+	FText GetNumElementsText(const FText& InContainerText, const FPropertyEditor& InPropertyEditor)
+	{
+		void* DummyValueAddress;
+		if (InPropertyEditor.GetPropertyHandle()->GetValueData(DummyValueAddress) == FPropertyAccess::MultipleValues)
+		{
+			return PropertyEditorConstants::DefaultUndeterminedText;
+		}
+
+		return FText::Format(LOCTEXT("NumItemsFmt", "{0} {1} {0}|plural(one=element,other=elements)"), FText::AsNumber(InPropertyEditor.GetPropertyNode()->GetNumChildNodes()), InContainerText);
+	}
+
 }
 
 #undef LOCTEXT_NAMESPACE
diff --git a/Engine/Source/Editor/PropertyEditor/Private/PropertyEditorHelpers.h b/Engine/Source/Editor/PropertyEditor/Private/PropertyEditorHelpers.h
--- a/Engine/Source/Editor/PropertyEditor/Private/PropertyEditorHelpers.h
+++ b/Engine/Source/Editor/PropertyEditor/Private/PropertyEditorHelpers.h
@@ -351,5 +351,13 @@ namespace PropertyEditorHelpers
 	 * @return nullptr if the class could not be found, or if the path doesn't correspond to a class object
 	 */
 	const UClass* FindOrLoadClass(const FString& ClassPath);
+
+	/** Returns the text to inform the users of the number of elements in a given container property 
+	 * @param InContainerText the name of the container to use in the text (e.g. "Map" for a TMap)
+	 * @param InPropertyEditor container to use in the text (e.g. "Map" for a TMap)
+	 * 
+	 * @return the localized text
+	 */
+	FText GetNumElementsText(const FText& InContainerText, const FPropertyEditor& InPropertyEditor);
 }
 
diff --git a/Engine/Source/Editor/PropertyEditor/Private/UserInterface/PropertyEditor/SPropertyEditorArray.cpp b/Engine/Source/Editor/PropertyEditor/Private/UserInterface/PropertyEditor/SPropertyEditorArray.cpp
--- a/Engine/Source/Editor/PropertyEditor/Private/UserInterface/PropertyEditor/SPropertyEditorArray.cpp
+++ b/Engine/Source/Editor/PropertyEditor/Private/UserInterface/PropertyEditor/SPropertyEditorArray.cpp
@@ -59,23 +59,8 @@ void SPropertyEditorArray::GetDesiredWidth(float& OutMinDesiredWidth, float& Out
 
 FText SPropertyEditorArray::GetArrayTextValue() const
 {
-	void* ValueAddress;
-	FPropertyAccess::Result GetValResult = PropertyEditor->GetPropertyHandle()->GetValueData(ValueAddress);
-
-	if (GetValResult == FPropertyAccess::MultipleValues)
-	{
-		return PropertyEditorConstants::DefaultUndeterminedText;
-	}
-
-	const int32 NumChildNodes = PropertyEditor->GetPropertyNode()->GetNumChildNodes();
-	if (NumChildNodes <= 1)
-	{
-		return FText::Format(LOCTEXT("SingleArrayItemFmt", "{0} Array element"), FText::AsNumber(NumChildNodes));
-	}
-	else
-	{
-		return FText::Format(LOCTEXT("NumArrayItemsFmt", "{0} Array elements"), FText::AsNumber(NumChildNodes));
-	}
+	check(PropertyEditor.IsValid());
+	return PropertyEditorHelpers::GetNumElementsText(LOCTEXT("ArrayText", "Array"), *PropertyEditor.Get());
 }
 
 bool SPropertyEditorArray::CanEdit() const
diff --git a/Engine/Source/Editor/PropertyEditor/Private/UserInterface/PropertyEditor/SPropertyEditorMap.cpp b/Engine/Source/Editor/PropertyEditor/Private/UserInterface/PropertyEditor/SPropertyEditorMap.cpp
--- a/Engine/Source/Editor/PropertyEditor/Private/UserInterface/PropertyEditor/SPropertyEditorMap.cpp
+++ b/Engine/Source/Editor/PropertyEditor/Private/UserInterface/PropertyEditor/SPropertyEditorMap.cpp
@@ -6,6 +6,7 @@
 #include "Internationalization/Internationalization.h"
 #include "Layout/Children.h"
 #include "Presentation/PropertyEditor/PropertyEditor.h"
+#include "PropertyEditorHelpers.h"
 #include "PropertyNode.h"
 #include "UObject/UnrealType.h"
 #include "Widgets/Text/STextBlock.h"
@@ -53,7 +54,8 @@ void SPropertyEditorMap::GetDesiredWidth(float& OutMinDesiredWidth, float& OutMa
 
 FText SPropertyEditorMap::GetMapTextValue() const
 {
-	return FText::Format(LOCTEXT("NumMapItemsFmt", "{0} Map elements"), FText::AsNumber( PropertyEditor->GetPropertyNode()->GetNumChildNodes()));
+	check(PropertyEditor.IsValid());
+	return PropertyEditorHelpers::GetNumElementsText(LOCTEXT("MapText", "Map"), *PropertyEditor.Get());
 }
 
 FText SPropertyEditorMap::GetMapTooltipText() const
diff --git a/Engine/Source/Editor/PropertyEditor/Private/UserInterface/PropertyEditor/SPropertyEditorSet.cpp b/Engine/Source/Editor/PropertyEditor/Private/UserInterface/PropertyEditor/SPropertyEditorSet.cpp
--- a/Engine/Source/Editor/PropertyEditor/Private/UserInterface/PropertyEditor/SPropertyEditorSet.cpp
+++ b/Engine/Source/Editor/PropertyEditor/Private/UserInterface/PropertyEditor/SPropertyEditorSet.cpp
@@ -6,6 +6,7 @@
 #include "Internationalization/Internationalization.h"
 #include "Layout/Children.h"
 #include "Presentation/PropertyEditor/PropertyEditor.h"
+#include "PropertyEditorHelpers.h"
 #include "PropertyNode.h"
 #include "UObject/UnrealType.h"
 #include "Widgets/Text/STextBlock.h"
@@ -53,7 +54,8 @@ void SPropertyEditorSet::GetDesiredWidth( float& OutMinDesiredWidth, float& OutM
 
 FText SPropertyEditorSet::GetSetTextValue() const
 {
-	return FText::Format(LOCTEXT("NumSetItemsFmt", "{0} Set elements"), FText::AsNumber(PropertyEditor->GetPropertyNode()->GetNumChildNodes()));
+	check(PropertyEditor.IsValid());
+	return PropertyEditorHelpers::GetNumElementsText(LOCTEXT("SetText", "Set"), *PropertyEditor.Get());
 }
 
 FText SPropertyEditorSet::GetSetTooltipText() const
