Fixup enum binding compatability with different types of integer handlers (uint8)

#Synced-CL 50735392

diff --git a/Engine/Source/Editor/PropertyEditor/Private/PropertyHandleImpl.cpp b/Engine/Source/Editor/PropertyEditor/Private/PropertyHandleImpl.cpp
--- a/Engine/Source/Editor/PropertyEditor/Private/PropertyHandleImpl.cpp
+++ b/Engine/Source/Editor/PropertyEditor/Private/PropertyHandleImpl.cpp
@@ -2542,7 +2542,7 @@ void FPropertyHandleBase::ClearResetToDefaultCustomized()
 {
 	if (TSharedPtr<FPropertyNode> PropertyNode = Implementation->GetPropertyNode(); PropertyNode.IsValid())
 	{
-		PropertyNode->SetNodeFlags(EPropertyNodeFlags::HasCustomResetToDefault, true);
+		PropertyNode->SetNodeFlags(EPropertyNodeFlags::HasCustomResetToDefault, false);
 	}
 }
 
@@ -3896,21 +3896,35 @@ bool FPropertyHandleInt::Supports( TSharedRef<FPropertyNode> PropertyNode )
 }
 
 template <typename PropertyClass, typename ValueType>
-ValueType GetIntegerValue(void* PropValue, FPropertyValueImpl& Implementation )
+ValueType GetIntegerValue(void* PropValue, FPropertyValueImpl& Implementation)
 {
+	if (Implementation.IsPropertyTypeOf(PropertyClass::StaticClass()))
+	{
+		return Implementation.GetPropertyValue<PropertyClass>(PropValue);
+	}
+
+	const FNumericProperty* NumericProperty = nullptr;
 	if (Implementation.IsPropertyTypeOf(FEnumProperty::StaticClass()))
 	{
 		const FEnumProperty* EnumProperty = CastFieldChecked<FEnumProperty>(Implementation.GetPropertyNode()->GetProperty());
-		const FNumericProperty* UnderlyingProperty = EnumProperty->GetUnderlyingProperty();
-		const bool bIsSignedInteger = UnderlyingProperty->CanHoldValue(-1);
+		NumericProperty = EnumProperty->GetUnderlyingProperty();
+	}
+	else if (Implementation.IsPropertyTypeOf(FNumericProperty::StaticClass()))
+	{
+		NumericProperty = CastFieldChecked<const FNumericProperty>(Implementation.GetPropertyNode()->GetProperty());
+	}
+
+	if (NumericProperty)
+	{
+		const bool bIsSignedInteger = NumericProperty->CanHoldValue(-1);
 		if (bIsSignedInteger)
 		{
-			return static_cast<ValueType>(UnderlyingProperty->GetSignedIntPropertyValue(PropValue));	
+			return static_cast<ValueType>(NumericProperty->GetSignedIntPropertyValue(PropValue));
 		}
-		return static_cast<ValueType>(UnderlyingProperty->GetUnsignedIntPropertyValue(PropValue));
+		return static_cast<ValueType>(NumericProperty->GetUnsignedIntPropertyValue(PropValue));
 	}
-	check( Implementation.IsPropertyTypeOf( PropertyClass::StaticClass() ) );
-	return Implementation.GetPropertyValue<PropertyClass>(PropValue);
+
+	return 0;
 }
 
 FPropertyAccess::Result FPropertyHandleInt::GetValue(int8& OutValue) const
@@ -3964,6 +3978,19 @@ FPropertyAccess::Result FPropertyHandleInt::GetValue( int64& OutValue ) const
 	return Res;
 }
 
+FPropertyAccess::Result FPropertyHandleInt::GetValue(uint8& OutValue) const
+{
+	void* PropValue = nullptr;
+	FPropertyAccess::Result Res = Implementation->GetValueData(PropValue);
+
+	if (Res == FPropertyAccess::Success)
+	{
+		OutValue = GetIntegerValue<FByteProperty, uint8>(PropValue, *Implementation);
+	}
+
+	return Res;
+}
+
 FPropertyAccess::Result FPropertyHandleInt::GetValue(uint16& OutValue) const
 {
 	void* PropValue = nullptr;
@@ -4085,6 +4112,26 @@ FPropertyAccess::Result FPropertyHandleInt::SetValue(const int64& NewValue, EPro
 	return Res;
 }
 
+FPropertyAccess::Result FPropertyHandleInt::SetValue(const uint8& NewValue, EPropertyValueSetFlags::Type Flags)
+{
+	FPropertyAccess::Result Res;
+	// Clamp the value from any meta data ranges stored on the property value
+	uint8 FinalValue = ClampIntegerValueFromMetaData<uint8>(NewValue, *this, *Implementation->GetPropertyNode());
+
+	FString ValueStr;
+	if (const FEnumProperty* EnumProperty = CastField<FEnumProperty>(GetProperty()))
+	{
+		ValueStr = EnumProperty->GetEnum()->GetNameStringByValue(FinalValue);
+	}
+	else
+	{
+		ValueStr = LexToString(FinalValue);
+	}
+	Res = Implementation->ImportText(ValueStr, Flags);
+
+	return Res;
+}
+
 FPropertyAccess::Result FPropertyHandleInt::SetValue(const uint16& NewValue, EPropertyValueSetFlags::Type Flags)
 {
 	FPropertyAccess::Result Res;
diff --git a/Engine/Source/Editor/PropertyEditor/Private/PropertyHandleImpl.h b/Engine/Source/Editor/PropertyEditor/Private/PropertyHandleImpl.h
--- a/Engine/Source/Editor/PropertyEditor/Private/PropertyHandleImpl.h
+++ b/Engine/Source/Editor/PropertyEditor/Private/PropertyHandleImpl.h
@@ -583,6 +583,9 @@ public:
 	virtual FPropertyAccess::Result GetValue(int64& OutValue) const override;
 	virtual FPropertyAccess::Result SetValue(const int64& InValue, EPropertyValueSetFlags::Type Flags = EPropertyValueSetFlags::DefaultFlags) override;
 
+	virtual FPropertyAccess::Result GetValue(uint8& OutValue) const override;
+	virtual FPropertyAccess::Result SetValue(const uint8& InValue, EPropertyValueSetFlags::Type Flags = EPropertyValueSetFlags::DefaultFlags) override;
+
 	virtual FPropertyAccess::Result GetValue(uint16& OutValue) const override;
 	virtual FPropertyAccess::Result SetValue(const uint16& InValue, EPropertyValueSetFlags::Type Flags = EPropertyValueSetFlags::DefaultFlags) override;
 
