Add GetNumChildren to DetailGroup.

#rnx

#Synced-CL 50656828

diff --git a/Engine/Source/Editor/PropertyEditor/Private/DetailGroup.h b/Engine/Source/Editor/PropertyEditor/Private/DetailGroup.h
--- a/Engine/Source/Editor/PropertyEditor/Private/DetailGroup.h
+++ b/Engine/Source/Editor/PropertyEditor/Private/DetailGroup.h
@@ -44,6 +44,9 @@ public:
 	/** @return The name of the group */
 	virtual FName GetGroupName() const override { return GetRowName(); }
 
+	/** @return The number of child rows under this group */
+	virtual int32 GetNumChildren() const override { return GroupChildren.Num(); }
+
 	/** @return The localized display name of the group */
 	const FText& GetGroupDisplayName() const { return LocalizedDisplayName; }
 	
diff --git a/Engine/Source/Editor/PropertyEditor/Public/IDetailGroup.h b/Engine/Source/Editor/PropertyEditor/Public/IDetailGroup.h
--- a/Engine/Source/Editor/PropertyEditor/Public/IDetailGroup.h
+++ b/Engine/Source/Editor/PropertyEditor/Public/IDetailGroup.h
@@ -118,6 +118,11 @@ public:
 	*/
 	virtual FName GetGroupName() const = 0;
 
+	/**
+	 * Returns the number of child rows under this group
+	 */
+	virtual int32 GetNumChildren() const = 0;
+
 	/**
 	* Return the property row associated with the specified property handle
 	*/
