[PropertyEditor] Amend the comment for RegisterCustomPropertyTypeLayout to include  instance object properties, since it does work for them.

#Synced-CL 50865973

diff --git a/Engine/Source/Editor/PropertyEditor/Public/PropertyEditorModule.h b/Engine/Source/Editor/PropertyEditor/Public/PropertyEditorModule.h
--- a/Engine/Source/Editor/PropertyEditor/Public/PropertyEditorModule.h
+++ b/Engine/Source/Editor/PropertyEditor/Public/PropertyEditorModule.h
@@ -292,7 +292,7 @@ public:
 
 	/**
 	 * Registers a property type customization
-	 * A property type is a specific FProperty type, a struct, or enum type
+	 * A property type is a specific FProperty type, a struct, instanced Object property, or enum type
 	 *
 	 * @param PropertyTypeName		The name of the property type to customize.  For structs and enums this is the name of the struct class or enum	(not StructProperty or ByteProperty) 
 	 * @param PropertyTypeLayoutDelegate	The delegate to call when querying for a custom layout of the property type
