Log a warning if an EditCondition involves a property that is marked deprecated

#Synced-CL 50915523

diff --git a/Engine/Source/Editor/PropertyEditor/Private/EditConditionContext.cpp b/Engine/Source/Editor/PropertyEditor/Private/EditConditionContext.cpp
--- a/Engine/Source/Editor/PropertyEditor/Private/EditConditionContext.cpp
+++ b/Engine/Source/Editor/PropertyEditor/Private/EditConditionContext.cpp
@@ -116,18 +116,33 @@ const T* FindTypedField(const TSharedPtr<FPropertyNode>& PropertyNode, const FSt
 		return nullptr;
 	}
 
+	auto ShouldLog = [&]() -> bool
+		{
+			TPair<FName, FString> FieldKey(Property->GetOwnerStruct()->GetFName(), FieldName);
+			if (!AlreadyLogged.Find(FieldKey))
+			{
+				AlreadyLogged.Add(FieldKey);
+				return true;
+			}
+			return false;
+		};
+
 	const FProperty* Field = FindFProperty<FProperty>(Property->GetOwnerStruct(), *FieldName);
 	if (Field == nullptr)
 	{
-		TPair<FName, FString> FieldKey(Property->GetOwnerStruct()->GetFName(), FieldName);
-		if (!AlreadyLogged.Find(FieldKey))
+		if (ShouldLog())
 		{
-			AlreadyLogged.Add(FieldKey);
 			UE_LOG(LogEditCondition, Error, TEXT("EditCondition parsing failed: Field name \"%s\" was not found in class \"%s\"."), *FieldName, *Property->GetOwnerStruct()->GetName());
 		}
-
 		return nullptr;
 	}
+	else if (Field->HasAnyPropertyFlags(CPF_Deprecated))
+	{
+		if (ShouldLog())
+		{
+			UE_LOG(LogEditCondition, Warning, TEXT("EditCondition parsing failed: Field name \"%s\" was found in class \"%s\" but marked deprecated."), *FieldName, *Property->GetOwnerStruct()->GetName());
+		}
+	}
 
 	return CastField<T>(Field);
 }
