FPropertyPermissionList thread-safety and caching improvements

* Mitigated a race condition in FPropertyPermissionList when the cache was cleared while another thread was reading a cached entry.
* Allow the cached result to be queried directly, to avoid repeated shared map look-ups in cases where we're frequently querying data for the same struct.

#rnx

#Synced-CL 51168648

diff --git a/Engine/Source/Editor/PropertyEditor/Private/PropertyPermissionList.cpp b/Engine/Source/Editor/PropertyEditor/Private/PropertyPermissionList.cpp
--- a/Engine/Source/Editor/PropertyEditor/Private/PropertyPermissionList.cpp
+++ b/Engine/Source/Editor/PropertyEditor/Private/PropertyPermissionList.cpp
@@ -264,14 +264,20 @@ void FPropertyPermissionList::ClearCache()
 
 bool FPropertyPermissionList::HasFiltering(const UStruct* ObjectStruct) const
 {
-	return ObjectStruct ? GetCachedPermissionListForStruct(ObjectStruct).HasFiltering() : false;
+	if (ObjectStruct && IsEnabled())
+	{
+		TSharedRef<const FCachedPermissionList> CachedPermissionList = GetCachedPermissionListEntryForStruct(ObjectStruct);
+		return CachedPermissionList->PermissionList->HasFiltering();
+	}
+	return false;
 }
 
 bool FPropertyPermissionList::DoesPropertyPassFilter(const UStruct* ObjectStruct, FName PropertyName) const
 {
 	if (ObjectStruct && IsEnabled())
 	{
-		return GetCachedPermissionListForStruct(ObjectStruct).PassesFilter(PropertyName);
+		TSharedRef<const FCachedPermissionList> CachedPermissionList = GetCachedPermissionListEntryForStruct(ObjectStruct);
+		return CachedPermissionList->PermissionList->PassesFilter(PropertyName);
 	}
 	return true;
 }
@@ -289,18 +295,23 @@ bool FPropertyPermissionList::DoesDetailTreeNodePassFilter(const UStruct* Object
 			}
 		}
 
-		return GetCachedPermissionListForStruct(ObjectStruct).PassesFilter(DetailTreeNode->GetNodeName());
+		TSharedRef<const FCachedPermissionList> CachedPermissionList = GetCachedPermissionListEntryForStruct(ObjectStruct);
+		return CachedPermissionList->PermissionList->PassesFilter(DetailTreeNode->GetNodeName());
 	}
 	
 	return true;
 }
 
-const FNamePermissionList& FPropertyPermissionList::GetCachedPermissionListForStruct(const UStruct* Struct) const
+TSharedPtr<const FNamePermissionList> FPropertyPermissionList::GetCachedPermissionListForStructIfEnabled(const UStruct* Struct) const
 {
-	return GetCachedPermissionListEntryForStructHelper(Struct).PermissionList;
+	if (Struct && IsEnabled())
+	{
+		return GetCachedPermissionListEntryForStruct(Struct)->PermissionList;
+	}
+	return nullptr;
 }
 
-const FPropertyPermissionListEntry& FPropertyPermissionList::GetCachedPermissionListEntryForStructHelper(const UStruct* Struct) const
+TSharedRef<const FPropertyPermissionList::FCachedPermissionList> FPropertyPermissionList::GetCachedPermissionListEntryForStruct(const UStruct* Struct) const
 {
 	check(Struct);
 
@@ -320,9 +331,9 @@ const FPropertyPermissionListEntry& FPropertyPermissionList::GetCachedPermission
 		FReadScopeLock _(CachedPropertyPermissionListRW);
 
 		// Is this struct already cached? If so, we can just return its result
-		if (const TUniquePtr<FPropertyPermissionListEntry>* CachedPermissionListPtr = CachedPropertyPermissionList.Find(Struct))
+		if (const TSharedRef<const FCachedPermissionList>* CachedPermissionListPtr = CachedPropertyPermissionList.Find(Struct))
 		{
-			return **CachedPermissionListPtr;
+			return *CachedPermissionListPtr;
 		}
 	}
 
@@ -354,9 +365,9 @@ const FPropertyPermissionListEntry& FPropertyPermissionList::GetCachedPermission
 	EPropertyPermissionListRules SuperEntryRule = EPropertyPermissionListRules::UseExistingPermissionList;
 	if (SuperStruct)
 	{
-		const FPropertyPermissionListEntry& SuperPermissionList = GetCachedPermissionListEntryForStructHelper(SuperStruct);
-		SuperEntryRule = SuperPermissionList.Rules;
-		NewPermissionList.Append(SuperPermissionList.PermissionList);
+		TSharedRef<const FCachedPermissionList> SuperPermissionList = GetCachedPermissionListEntryForStruct(SuperStruct);
+		SuperEntryRule = SuperPermissionList->Rules;
+		NewPermissionList.Append(*SuperPermissionList->PermissionList);
 	}
 
 	// Resolve the permission list rule for this entry
@@ -429,14 +440,13 @@ const FPropertyPermissionListEntry& FPropertyPermissionList::GetCachedPermission
 		FWriteScopeLock _(CachedPropertyPermissionListRW);
 
 		// Check the cache again in case another thread beat us to it
-		if (const TUniquePtr<FPropertyPermissionListEntry>* CachedPermissionListPtr = CachedPropertyPermissionList.Find(Struct))
+		if (const TSharedRef<const FCachedPermissionList>* CachedPermissionListPtr = CachedPropertyPermissionList.Find(Struct))
 		{
-			return **CachedPermissionListPtr;
+			return *CachedPermissionListPtr;
 		}
 
-		return *CachedPropertyPermissionList.Add(Struct, MakeUnique<FPropertyPermissionListEntry>(
+		return CachedPropertyPermissionList.Add(Struct, MakeShared<FCachedPermissionList>(
 			MoveTemp(NewPermissionList),
-			TArray<FName>(), // Note: The AdditionalOwnerNames are only used to prevent removal from RawPropertyPermissionList by calls to RemoveOwner()
 			EntryRuleToPropagate
 			));
 	}
diff --git a/Engine/Source/Editor/PropertyEditor/Public/PropertyPermissionList.h b/Engine/Source/Editor/PropertyEditor/Public/PropertyPermissionList.h
--- a/Engine/Source/Editor/PropertyEditor/Public/PropertyPermissionList.h
+++ b/Engine/Source/Editor/PropertyEditor/Public/PropertyPermissionList.h
@@ -127,6 +127,9 @@ public:
 	/** Clear CachedPropertyPermissionList to cause the PermissionListed property list to be regenerated next time it's queried */
 	UE_API void ClearCache();
 
+	/** Gets a read-only reference to the cached permission list for the given struct, or null if the filtering is disabled. A non-null value will remain valid to read, but the info in it may become stale if this struct is re-cached because its raw rules change. */
+	UE_API TSharedPtr<const FNamePermissionList> GetCachedPermissionListForStructIfEnabled(const UStruct* Struct) const;
+
 	/** If true, PermissionListUpdatedDelegate will not broadcast when modifying the permission lists */
 	bool bSuppressUpdateDelegate = false;
 
@@ -159,13 +162,24 @@ private:
 	friend bool operator==(const FPermissionListUpdate& A, const FPermissionListUpdate& B);
 	friend uint32 GetTypeHash(const FPermissionListUpdate& PermisisonList);
 	TSet< FPermissionListUpdate > PendingUpdates;
+
+	struct FCachedPermissionList
+	{
+		FCachedPermissionList(FNamePermissionList&& InPermissionList, EPropertyPermissionListRules InRules)
+			: PermissionList(MakeShared<FNamePermissionList>(MoveTemp(InPermissionList)))
+			, Rules(InRules)
+		{
+		}
+
+		TSharedRef<const FNamePermissionList> PermissionList;
+		EPropertyPermissionListRules Rules = EPropertyPermissionListRules::UseExistingPermissionList;
+	};
 	/** Lazily-constructed combined cache of both the flattened class PermissionList and struct PermissionList */
-	mutable TMap<TWeakObjectPtr<const UStruct>, TUniquePtr<FPropertyPermissionListEntry>> CachedPropertyPermissionList;
+	mutable TMap<TWeakObjectPtr<const UStruct>, TSharedRef<const FCachedPermissionList>> CachedPropertyPermissionList;
 	mutable FRWLock CachedPropertyPermissionListRW;
 
 	/** Get or create the cached PermissionList for a specific UStruct */
-	const FNamePermissionList& GetCachedPermissionListForStruct(const UStruct* Struct) const;
-	const FPropertyPermissionListEntry& GetCachedPermissionListEntryForStructHelper(const UStruct* Struct) const;
+	TSharedRef<const FCachedPermissionList> GetCachedPermissionListEntryForStruct(const UStruct* Struct) const;
 };
 
 class FPropertyEditorPermissionList : public FPropertyPermissionList
