FPropertyPermissionList BP dirtying/ticking improvements

* When a BP is re-compiled, only remove the relevent entries from FPropertyPermissionList rather than clear the entire cache. Also ignore the initial compile-on-load.
* Skip the deferred update handling for commandlets, as it relies on ticking, and commandlets don't tick. Only UI was binding to PermissionListUpdatedDelegate and that also doesn't exist for a commandlet.

#rnx
#Synced-CL 51168648

diff --git a/Engine/Source/Editor/PropertyEditor/Private/PropertyPermissionList.cpp b/Engine/Source/Editor/PropertyEditor/Private/PropertyPermissionList.cpp
--- a/Engine/Source/Editor/PropertyEditor/Private/PropertyPermissionList.cpp
+++ b/Engine/Source/Editor/PropertyEditor/Private/PropertyPermissionList.cpp
@@ -42,25 +42,20 @@ uint32 GetTypeHash(const FPropertyPermissionList::FPermissionListUpdate& Permisi
 
 FPropertyPermissionList::FPropertyPermissionList()
 {
-	if (GEditor)
+	if (!IsRunningCommandlet()) // Commandlets don't Tick
 	{
-		RegisterOnBlueprintCompiled();
+		OnTickHandle = FTSTicker::GetCoreTicker().AddTicker(FTickerDelegate::CreateRaw(this, &FPropertyPermissionList::Tick), 1.0f);
 	}
-	else
-	{
-		FCoreDelegates::OnPostEngineInit.AddRaw(this, &FPropertyPermissionList::RegisterOnBlueprintCompiled);
-	}
-	OnTickHandle = FTSTicker::GetCoreTicker().AddTicker(FTickerDelegate::CreateRaw(this, &FPropertyPermissionList::Tick), 1.0f);
+	FCoreUObjectDelegates::OnObjectPostCDOCompiled.AddRaw(this, &FPropertyPermissionList::OnObjectPostCDOCompiled);
 }
 
 FPropertyPermissionList::~FPropertyPermissionList()
 {
-	FTSTicker::RemoveTicker(OnTickHandle);
-	FCoreDelegates::OnPostEngineInit.RemoveAll(this);
-	if (GEditor)
+	if (OnTickHandle.IsValid())
 	{
-		GEditor->OnBlueprintCompiled().RemoveAll(this);
+		FTSTicker::RemoveTicker(OnTickHandle);
 	}
+	FCoreUObjectDelegates::OnObjectPostCDOCompiled.RemoveAll(this);
 }
 
 bool FPropertyPermissionList::Tick(float DeltaTime)
@@ -79,11 +74,17 @@ bool FPropertyPermissionList::Tick(float DeltaTime)
 	return true;
 }
 
-void FPropertyPermissionList::RegisterOnBlueprintCompiled()
+void FPropertyPermissionList::OnObjectPostCDOCompiled(UObject* CDO, const FObjectPostCDOCompiledContext& Context)
 {
-	if (GEditor)
+	if (Context.bIsRegeneratingOnLoad)
+	{
+		// Can skip compile-on-load as the class layout can't have changed from what was cached
+		return;
+	}
+
 	{
-		GEditor->OnBlueprintCompiled().AddRaw(this, &FPropertyPermissionList::ClearCache);
+		FWriteScopeLock _(CachedPropertyPermissionListRW);
+		ClearCacheForClass_NoLock(CDO->GetClass());
 	}
 }
 
@@ -93,12 +94,28 @@ void FPropertyPermissionList::ClearCacheAndBroadcast(TSoftObjectPtr<const UStruc
 	// and lazily repopulated any time the raw PermissionList changes.
 	ClearCache();
 
-	if (!bSuppressUpdateDelegate)
+	if (!bSuppressUpdateDelegate && OnTickHandle.IsValid())
 	{
 		PendingUpdates.Add({ObjectStruct, OwnerName});
 	}
 }
 
+void FPropertyPermissionList::ClearCacheForClass_NoLock(const UClass* Class)
+{
+	if (CachedPropertyPermissionList.Remove(Class) == 0)
+	{
+		// If this class wasn't cached, then none of its derived types can be either
+		return;
+	}
+
+	TArray<UClass*> DerivedClasses;
+	GetDerivedClasses(Class, DerivedClasses, /*bRecursive*/false);
+	for (const UClass* DerivedClass : DerivedClasses)
+	{
+		ClearCacheForClass_NoLock(DerivedClass);
+	}
+}
+
 void FPropertyPermissionList::AddPermissionList(TSoftObjectPtr<const UStruct> Struct, const FNamePermissionList& PermissionList, const EPropertyPermissionListRules Rules, const TConstArrayView<FName> InAdditionalOwnerNames)
 {
 	FPropertyPermissionListEntry& Entry = RawPropertyPermissionList.FindOrAdd(Struct);
diff --git a/Engine/Source/Editor/PropertyEditor/Public/PropertyPermissionList.h b/Engine/Source/Editor/PropertyEditor/Public/PropertyPermissionList.h
--- a/Engine/Source/Editor/PropertyEditor/Public/PropertyPermissionList.h
+++ b/Engine/Source/Editor/PropertyEditor/Public/PropertyPermissionList.h
@@ -140,9 +140,12 @@ protected:
 	UE_API ~FPropertyPermissionList();
 
 private:
-	UE_API bool Tick(float DeltaTime);
-	UE_API void RegisterOnBlueprintCompiled();
-	UE_API void ClearCacheAndBroadcast(TSoftObjectPtr<const UStruct> ObjectStruct = nullptr, FName OwnerName = NAME_None);
+	bool Tick(float DeltaTime);
+	void OnObjectPostCDOCompiled(UObject* CDO, const FObjectPostCDOCompiledContext& Context);
+	void ClearCacheAndBroadcast(TSoftObjectPtr<const UStruct> ObjectStruct = nullptr, FName OwnerName = NAME_None);
+
+	/** Clear the given class (and any derived classes) from the cache. You must hold a write lock on CachedPropertyPermissionListRW before calling this function! */
+	void ClearCacheForClass_NoLock(const UClass* Class);
 
 	/** Delegate to return whether DoesPropertyPassFilter should perform its PermissionList check or always return true */
 	FDelegateEnablePathFilter EnablePermissionList;
