Allow specifying a default value in TOptional properties, which is then used in the property editor UI when the option value is set

#Synced-CL 51001024

diff --git a/Engine/Source/Editor/PropertyEditor/Private/PropertyHandleImpl.cpp b/Engine/Source/Editor/PropertyEditor/Private/PropertyHandleImpl.cpp
--- a/Engine/Source/Editor/PropertyEditor/Private/PropertyHandleImpl.cpp
+++ b/Engine/Source/Editor/PropertyEditor/Private/PropertyHandleImpl.cpp
@@ -5829,7 +5829,7 @@ FPropertyAccess::Result FPropertyHandleOptional::SetOptionalValue(FProperty* New
 	for (int i = 0; i < ReadAddresses.Num(); i++)
 	{
 		void* Addr = ReadAddresses.GetAddress(i);
-		OptionalProperty->MarkSetAndGetInitializedValuePointerToReplace(Addr);
+		void* ValuePtr = OptionalProperty->MarkSetAndGetInitializedValuePointerToReplace(Addr);
 
 		if (NewProperty)
 		{
@@ -5838,13 +5838,21 @@ FPropertyAccess::Result FPropertyHandleOptional::SetOptionalValue(FProperty* New
 
 		// TODO: This should be re-worked to leverage using 'import_text' instead of directly setting new values as we are duplicating logic somewhat here (non-trivial task)
 		FProperty* ValueProperty = OptionalProperty->GetValueProperty();
+
+		static const FName DefaultSomeMetaDataName = "DefaultSome";
+		const FString* DefaultSomeString = OptionalProperty->FindMetaData(DefaultSomeMetaDataName);
+		if (ValueProperty && DefaultSomeString)
+		{
+			ValueProperty->ImportText_Direct(**DefaultSomeString, ValuePtr, nullptr, 0);
+		}
+
 		if (ValueProperty && NewObjectClass)
 		{
 			if (ValueProperty->IsA(FClassProperty::StaticClass()) || ValueProperty->IsA(FSoftClassProperty::StaticClass()))
 			{
 				if (FObjectPropertyBase* ClassProperty = CastField<FObjectPropertyBase>(ValueProperty))
 				{
-					void* ClassPropertyValuePtr = ClassProperty->ContainerPtrToValuePtr<void>(Addr);
+					void* ClassPropertyValuePtr = ClassProperty->ContainerPtrToValuePtr<void>(ValuePtr);
 					ClassProperty->SetObjectPropertyValue(ClassPropertyValuePtr, const_cast<UClass*>(NewObjectClass));
 				}
 			}
