Add support for capturing expansion state per custom built details view via ViewIdentifier's.
Initial usage example in SStateCentricViewMainStateNode.cpp

#Synced-CL 51389752
#tests editor

#Synced-CL 51389752

diff --git a/Engine/Plugins/Runtime/StateTree/Source/StateTreeEditorModule/Private/StateCentricView/View/SStateCentricViewMainStateNode.cpp b/Engine/Plugins/Runtime/StateTree/Source/StateTreeEditorModule/Private/StateCentricView/View/SStateCentricViewMainStateNode.cpp
--- a/Engine/Plugins/Runtime/StateTree/Source/StateTreeEditorModule/Private/StateCentricView/View/SStateCentricViewMainStateNode.cpp
+++ b/Engine/Plugins/Runtime/StateTree/Source/StateTreeEditorModule/Private/StateCentricView/View/SStateCentricViewMainStateNode.cpp
@@ -289,6 +289,8 @@ TSharedRef<SWidget> SStateCentricViewMainStateNode::Construct_GenerateNodeBody(c
 
 			return false;
 		});
+		DetailsViewArgs.ViewIdentifier = FName(ViewState->Name.ToString() + "_SStateCentricViewMainStateNode");
+		DetailsViewArgs.ExpansionPersistenceMethod = FDetailsViewArgs::EExpansionPersistenceMethod::TransientPerViewIdentifier;
 
 		TSharedRef<IDetailsView> StateDetailsView = PropertyEditorModule.CreateDetailView(DetailsViewArgs);
 		StateDetailsView->SetExtensionHandler(MakeShared<FStateTreeBindingExtension>());
diff --git a/Engine/Source/Editor/PropertyEditor/Private/DetailCategoryBuilderImpl.cpp b/Engine/Source/Editor/PropertyEditor/Private/DetailCategoryBuilderImpl.cpp
--- a/Engine/Source/Editor/PropertyEditor/Private/DetailCategoryBuilderImpl.cpp
+++ b/Engine/Source/Editor/PropertyEditor/Private/DetailCategoryBuilderImpl.cpp
@@ -19,7 +19,11 @@
 #include "SDetailCategoryTableRow.h"
 #include "StructurePropertyNode.h"
 #include "Styling/StyleColors.h"
- 
+
+FString FDetailCategoryImpl::DetailCategories = TEXT("DetailCategories");
+FString FDetailCategoryImpl::TransientPerViewDetailCategories = TEXT("TransientPerViewDetailCategories");
+FString FDetailCategoryImpl::PersistentPerViewDetailCategories = TEXT("PersistentPerViewDetailCategories");
+
 static void AddLayoutToList(const FDetailLayoutCustomization& Layout, TArray<FDetailLayoutCustomization>& List)
 {
 	if (Layout.bCustom)
@@ -1017,13 +1021,53 @@ void FDetailCategoryImpl::OnItemExpansionChanged(bool bIsExpanded, bool bShouldS
 {
 	if (bRestoreExpansionState && bShouldSaveState)
 	{
+		
 		// Save the collapsed state of this section
-		GConfig->SetBool(TEXT("DetailCategories"), *CategoryPathName, bIsExpanded, GEditorPerProjectIni);
+		GConfig->SetBool(*GetDetailCategoriesExpansionSection(), *GetCategoryPathConfigKey(CategoryPathName), bIsExpanded, GEditorPerProjectIni);
 	}
 
 	OnExpansionChangedDelegate.ExecuteIfBound(bIsExpanded);
 }
 
+const FString& FDetailCategoryImpl::GetDetailCategoriesExpansionSection() const
+{
+	if (const TSharedPtr<IDetailsViewPrivate> DetailsView = GetDetailsViewSharedPtr())
+	{
+		if (DetailsView->GetExpansionPersistenceMethod() == FDetailsViewArgs::EExpansionPersistenceMethod::TransientPerViewIdentifier)
+		{
+			return TransientPerViewDetailCategories;
+		}
+		else if (DetailsView->GetExpansionPersistenceMethod() == FDetailsViewArgs::EExpansionPersistenceMethod::PersistentPerViewIdentifier)
+		{
+			return PersistentPerViewDetailCategories;
+		}
+	}
+
+	return DetailCategories;
+}
+
+FString FDetailCategoryImpl::GetCategoryPathConfigKey(const FString& InPath) const
+{
+	if (const TSharedPtr<IDetailsViewPrivate> DetailsView = GetDetailsViewSharedPtr())
+	{
+		if (DetailsView->GetExpansionPersistenceMethod() == FDetailsViewArgs::EExpansionPersistenceMethod::TransientPerViewIdentifier
+			|| DetailsView->GetExpansionPersistenceMethod() == FDetailsViewArgs::EExpansionPersistenceMethod::PersistentPerViewIdentifier)
+		{
+			if (ensure(!DetailsView->GetViewIdentifier().IsNone()))
+			{
+				TStringBuilder<128> ViewIdentifierString;
+				ViewIdentifierString.Append(DetailsView->GetViewIdentifier().ToString());
+				ViewIdentifierString.Append(".");
+				ViewIdentifierString.Append(InPath);
+				return ViewIdentifierString.ToString();
+			}
+		}
+
+	}
+
+	return InPath;
+}
+
 bool FDetailCategoryImpl::ShouldBeExpanded() const
 {
 	if (bHasFilterStrings)
@@ -1035,7 +1079,7 @@ bool FDetailCategoryImpl::ShouldBeExpanded() const
 		// Collapse by default if there are no simple child nodes
 		bool bShouldBeExpanded = !ContainsOnlyAdvanced() && !bShouldBeInitiallyCollapsed;
 		// Save the collapsed state of this section
-		GConfig->GetBool(TEXT("DetailCategories"), *CategoryPathName, bShouldBeExpanded, GEditorPerProjectIni);
+		GConfig->GetBool(*GetDetailCategoriesExpansionSection(), *GetCategoryPathConfigKey(CategoryPathName), bShouldBeExpanded, GEditorPerProjectIni);
 		return bShouldBeExpanded;
 	}
 	else
diff --git a/Engine/Source/Editor/PropertyEditor/Private/DetailCategoryBuilderImpl.h b/Engine/Source/Editor/PropertyEditor/Private/DetailCategoryBuilderImpl.h
--- a/Engine/Source/Editor/PropertyEditor/Private/DetailCategoryBuilderImpl.h
+++ b/Engine/Source/Editor/PropertyEditor/Private/DetailCategoryBuilderImpl.h
@@ -158,6 +158,12 @@ private:
  */
 class FDetailCategoryImpl : public IDetailCategoryBuilder, public FDetailTreeNode, public TSharedFromThis<FDetailCategoryImpl>
 {
+public:
+
+	static FString DetailCategories;
+	static FString TransientPerViewDetailCategories;
+	static FString PersistentPerViewDetailCategories;
+
 public:
 	FDetailCategoryImpl(FName InCategoryName, TSharedRef<FDetailLayoutBuilderImpl> InDetailLayout);
 	~FDetailCategoryImpl();
@@ -352,6 +358,12 @@ public:
 private:
 	virtual void OnItemExpansionChanged(bool bIsExpanded, bool bShouldSaveState) override;
 
+	/** Return string key for the given struct in config .ini for this details view's categories */
+	const FString& GetDetailCategoriesExpansionSection() const;
+
+	/** Return string key for the given struct in config .ini for this details view */
+	FString GetCategoryPathConfigKey(const FString& InPath) const;
+
 	/**
 	 * Generates children for all layouts.
 	 */
diff --git a/Engine/Source/Editor/PropertyEditor/Private/IDetailsViewPrivate.h b/Engine/Source/Editor/PropertyEditor/Private/IDetailsViewPrivate.h
--- a/Engine/Source/Editor/PropertyEditor/Private/IDetailsViewPrivate.h
+++ b/Engine/Source/Editor/PropertyEditor/Private/IDetailsViewPrivate.h
@@ -5,6 +5,7 @@
 #include "AssetThumbnail.h"
 #include "DetailsDisplayManager.h"
 #include "DetailTreeNode.h"
+#include "DetailsViewArgs.h"
 #include "IDetailsView.h"
 #include "PropertyNode.h"
 
@@ -69,6 +70,20 @@ public:
 	 */
 	virtual bool GetCustomSavedExpansionState(const FString& NodePath) const = 0;
 
+	/**
+	 * Gets the View Identifier for this details view. If any.
+	 *
+	 * @return Text set for view identifier for this details view on creation
+	 */
+	virtual FName GetViewIdentifier() const = 0;
+
+	/**
+	 * Gets expansion persistence method for this details view. See FDetailsViewArgs.
+	 *
+	 * @return Text set for view identifier for this details view on creation
+	 */
+	virtual FDetailsViewArgs::EExpansionPersistenceMethod GetExpansionPersistenceMethod() const = 0;
+
 	/**
 	 * @return True if the property is visible
 	 */
diff --git a/Engine/Source/Editor/PropertyEditor/Private/PropertyEditorModule.cpp b/Engine/Source/Editor/PropertyEditor/Private/PropertyEditorModule.cpp
--- a/Engine/Source/Editor/PropertyEditor/Private/PropertyEditorModule.cpp
+++ b/Engine/Source/Editor/PropertyEditor/Private/PropertyEditorModule.cpp
@@ -23,6 +23,7 @@
 #include "PropertyEditorToolkit.h"
 #include "PropertyRowGenerator.h"
 #include "SDetailsView.h"
+#include "SDetailsViewBase.h"
 #include "SPropertyTreeViewImpl.h"
 #include "SSingleProperty.h"
 #include "SStructureDetailsView.h"
@@ -104,6 +105,8 @@ void FPropertyEditorModule::StartupModule()
 	FDetailsViewStyle::InitializeDetailsViewStyles();
 
 	UToolMenus::RegisterStartupCallback(FSimpleMulticastDelegate::FDelegate::CreateRaw(this, &FPropertyEditorModule::RegisterMenus));
+
+	FCoreDelegates::GetOnPostEngineInit().Add(FSimpleMulticastDelegate::FDelegate::CreateStatic(SDetailsViewBase::ResetTransientPerViewDetailExpansion));
 }
 
 void FPropertyEditorModule::ShutdownModule()
diff --git a/Engine/Source/Editor/PropertyEditor/Private/SDetailsViewBase.cpp b/Engine/Source/Editor/PropertyEditor/Private/SDetailsViewBase.cpp
--- a/Engine/Source/Editor/PropertyEditor/Private/SDetailsViewBase.cpp
+++ b/Engine/Source/Editor/PropertyEditor/Private/SDetailsViewBase.cpp
@@ -4,6 +4,7 @@
 
 #include "Algo/AllOf.h"
 #include "Algo/AnyOf.h"
+#include "Containers/StringFwd.h"
 #include "Engine/Engine.h"
 #include "GameFramework/Actor.h"
 #include "Misc/ConfigCacheIni.h"
@@ -25,6 +26,40 @@
 #include "SDetailNameArea.h"
 #include "ThumbnailRendering/ThumbnailManager.h"
 
+static FString DetailPropertyExpansion = TEXT("DetailPropertyExpansion");
+static FString TransientPerViewDetailPropertyExpansion = TEXT("TransientPerViewDetailPropertyExpansion");
+static FString PersistentPerViewDetailPropertyExpansion = TEXT("PersistentPerViewDetailPropertyExpansion");
+
+static FString DetailCustomWidgetExpansion = TEXT("DetailCustomWidgetExpansion");
+static FString TransientPerViewDetailCustomWidgetExpansion = TEXT("TransientPerViewDetailCustomWidgetExpansion");
+static FString PersistentPerViewDetailCustomWidgetExpansion = TEXT("PersistentPerViewDetailCustomWidgetExpansion");
+
+void SDetailsViewBase::ResetTransientPerViewDetailExpansion()
+{
+	GConfig->RemoveSectionFromBranch(*TransientPerViewDetailPropertyExpansion, *GEditorPerProjectIni);
+	GConfig->RemoveSectionFromBranch(*TransientPerViewDetailCustomWidgetExpansion, *GEditorPerProjectIni);
+	GConfig->RemoveSectionFromBranch(*FDetailCategoryImpl::TransientPerViewDetailCategories, *GEditorPerProjectIni);
+}
+
+void SDetailsViewBase::ResetPersistentPerViewDetailExpansion()
+{
+	GConfig->RemoveSectionFromBranch(*PersistentPerViewDetailPropertyExpansion, *GEditorPerProjectIni);
+	GConfig->RemoveSectionFromBranch(*PersistentPerViewDetailCustomWidgetExpansion, *GEditorPerProjectIni);
+	GConfig->RemoveSectionFromBranch(*FDetailCategoryImpl::PersistentPerViewDetailCategories, *GEditorPerProjectIni);
+}
+
+static FAutoConsoleCommand ResetTransientPerViewDetailExpansionCommand(
+	TEXT("PropertyEditor.ResetTransientPerViewDetailExpansion"),
+	TEXT("Removes all transient per view detail expansions. These or automatically purged without read on editor startup to avoid incremental file bloat."),
+	FConsoleCommandDelegate::CreateStatic(&SDetailsViewBase::ResetTransientPerViewDetailExpansion)
+);
+
+static FAutoConsoleCommand ResetPersistentPerViewDetailExpansionCommand(
+	TEXT("PropertyEditor.ResetPersistentPerViewDetailExpansion"),
+	TEXT("Removes all persistent per view detail expansions."),
+	FConsoleCommandDelegate::CreateStatic(&SDetailsViewBase::ResetPersistentPerViewDetailExpansion)
+);
+
 SDetailsViewBase::SDetailsViewBase() :
 	  NumVisibleTopLevelObjectNodes(0)
 	, bIsLocked(false)
@@ -704,7 +739,7 @@ void SDetailsViewBase::GetExpansionState(UStruct* BaseStruct, FStringPrefixTree&
 	TArray<FString> DetailPropertyExpansionStrings;
 	for (UStruct* Struct = BaseStruct; Struct && ((BaseStruct == Struct) || (Struct != AActor::StaticClass())); Struct = Struct->GetSuperStruct())
 	{
-		GConfig->GetSingleLineArray(TEXT("DetailPropertyExpansion"), *Struct->GetName(), DetailPropertyExpansionStrings, GEditorPerProjectIni);
+		GConfig->GetSingleLineArray(*GetDetailPropertyExpansionSection(), *GetStructExpansionConfigKey(BaseStruct), DetailPropertyExpansionStrings, GEditorPerProjectIni);
 	}
 
 	OutPrefixTree.InsertAll(DetailPropertyExpansionStrings);
@@ -935,6 +970,16 @@ bool SDetailsViewBase::GetCustomSavedExpansionState(const FString& NodePath) con
 	return ExpandedDetailNodes.Contains(NodePath);
 }
 
+FName SDetailsViewBase::GetViewIdentifier() const
+{
+	return DetailsViewArgs.ViewIdentifier;
+}
+
+FDetailsViewArgs::EExpansionPersistenceMethod SDetailsViewBase::GetExpansionPersistenceMethod() const
+{
+	return DetailsViewArgs.ExpansionPersistenceMethod;
+}
+
 bool SDetailsViewBase::IsPropertyVisible( const FPropertyAndParent& PropertyAndParent ) const
 {
 	return IsPropertyVisibleDelegate.IsBound() ? IsPropertyVisibleDelegate.Execute(PropertyAndParent) : true;
@@ -1564,6 +1609,52 @@ static void SetExpandedItems(TSharedPtr<FPropertyNode> InPropertyNode, const FSt
 	}
 }
 
+const FString& SDetailsViewBase::GetDetailPropertyExpansionSection() const
+{
+	if (DetailsViewArgs.ExpansionPersistenceMethod == FDetailsViewArgs::EExpansionPersistenceMethod::TransientPerViewIdentifier)
+	{
+		return TransientPerViewDetailPropertyExpansion;
+	}
+	else if (DetailsViewArgs.ExpansionPersistenceMethod == FDetailsViewArgs::EExpansionPersistenceMethod::PersistentPerViewIdentifier)
+	{
+		return PersistentPerViewDetailPropertyExpansion;
+	}
+
+	return DetailPropertyExpansion;
+}
+
+const FString& SDetailsViewBase::GetDetailCustomWidgetExpansionSection() const
+{
+	if (DetailsViewArgs.ExpansionPersistenceMethod == FDetailsViewArgs::EExpansionPersistenceMethod::TransientPerViewIdentifier)
+	{
+		return TransientPerViewDetailCustomWidgetExpansion;
+	}
+	else if (DetailsViewArgs.ExpansionPersistenceMethod == FDetailsViewArgs::EExpansionPersistenceMethod::PersistentPerViewIdentifier)
+	{
+		return PersistentPerViewDetailCustomWidgetExpansion;
+	}
+
+	return DetailCustomWidgetExpansion;
+}
+
+FString SDetailsViewBase::GetStructExpansionConfigKey(UStruct* InStruct) const
+{
+	if (DetailsViewArgs.ExpansionPersistenceMethod == FDetailsViewArgs::EExpansionPersistenceMethod::TransientPerViewIdentifier
+		|| DetailsViewArgs.ExpansionPersistenceMethod == FDetailsViewArgs::EExpansionPersistenceMethod::PersistentPerViewIdentifier)
+	{
+		if (ensure(!DetailsViewArgs.ViewIdentifier.IsNone()))
+		{
+			TStringBuilder<128> ViewIdentifierString;
+			ViewIdentifierString.Append(DetailsViewArgs.ViewIdentifier.ToString());
+			ViewIdentifierString.Append(".");
+			ViewIdentifierString.Append(InStruct->GetName());
+			return ViewIdentifierString.ToString();
+		}
+	}
+
+	return InStruct->GetName();
+}
+
 void SDetailsViewBase::SavePreSearchExpandedItems()
 {
 	PreSearchExpandedItems.Clear();
@@ -1717,13 +1808,13 @@ void SDetailsViewBase::SaveExpandedItems(TSharedRef<FPropertyNode> StartNode)
 			if (!bShouldSave)
 			{
 				TArray<FString> DummyExpandedPropertyItems;
-				GConfig->GetSingleLineArray(TEXT("DetailPropertyExpansion"), *Struct->GetName(), DummyExpandedPropertyItems, GEditorPerProjectIni);
+				GConfig->GetSingleLineArray(*GetDetailPropertyExpansionSection(), *GetStructExpansionConfigKey(Struct), DummyExpandedPropertyItems, GEditorPerProjectIni);
 				bShouldSave = DummyExpandedPropertyItems.Num() > 0;
 			}
 
 			if (bShouldSave)
 			{
-				GConfig->SetSingleLineArray(TEXT("DetailPropertyExpansion"), *Struct->GetName(), ExpandedPropertyItems, GEditorPerProjectIni);
+				GConfig->SetSingleLineArray(*GetDetailPropertyExpansionSection(), *GetStructExpansionConfigKey(Struct), ExpandedPropertyItems, GEditorPerProjectIni);
 			}
 		}
 	}
@@ -1742,13 +1833,13 @@ void SDetailsViewBase::SaveExpandedItems(TSharedRef<FPropertyNode> StartNode)
 		if (!bShouldSave)
 		{
 			FString DummyExpandedCustomItemsString;
-			GConfig->GetString(TEXT("DetailCustomWidgetExpansion"), *BestBaseStruct->GetName(), DummyExpandedCustomItemsString, GEditorPerProjectIni);
+			GConfig->GetString(*GetDetailCustomWidgetExpansionSection(), *GetStructExpansionConfigKey(BestBaseStruct), DummyExpandedCustomItemsString, GEditorPerProjectIni);
 			bShouldSave = !DummyExpandedCustomItemsString.IsEmpty();
 		}
 
 		if (bShouldSave)
 		{
-			GConfig->SetString(TEXT("DetailCustomWidgetExpansion"), *BestBaseStruct->GetName(), *ExpandedCustomItemsBuilder, GEditorPerProjectIni);
+			GConfig->SetString(*GetDetailCustomWidgetExpansionSection(), *GetStructExpansionConfigKey(BestBaseStruct), *ExpandedCustomItemsBuilder, GEditorPerProjectIni);
 		}
 	}
 }
@@ -1809,23 +1900,15 @@ void SDetailsViewBase::RestoreExpandedItems(TSharedRef<FPropertyNode> StartNode,
 	}
 	else
 	{
-		//while a valid class, and we're either the same as the base class (for multiple actors being selected and base class is AActor) OR we're not down to AActor yet)
-		TArray<FString> DetailPropertyExpansionStrings;
-		for (UStruct* Struct = BestBaseStruct; Struct && ((BestBaseStruct == Struct) || (Struct != AActor::StaticClass())); Struct = Struct->GetSuperStruct())
-		{
-			GConfig->GetSingleLineArray(TEXT("DetailPropertyExpansion"), *Struct->GetName(), DetailPropertyExpansionStrings, GEditorPerProjectIni);
-		}
-
 		FStringPrefixTree PrefixTree;
-		PrefixTree.InsertAll(DetailPropertyExpansionStrings);
-
+		GetExpansionState(BestBaseStruct, PrefixTree);
 		SetExpandedItems(StartNode, PrefixTree, false);
 	}
 
 	if (BestBaseStruct)
 	{
 		FString ExpandedCustomItems;
-		GConfig->GetString(TEXT("DetailCustomWidgetExpansion"), *BestBaseStruct->GetName(), ExpandedCustomItems, GEditorPerProjectIni);
+		GConfig->GetString(*GetDetailCustomWidgetExpansionSection(), *GetStructExpansionConfigKey(BestBaseStruct), ExpandedCustomItems, GEditorPerProjectIni);
 
 		TArray<FString> ExpandedCustomItemsArray;
 		ExpandedCustomItems.ParseIntoArray(ExpandedCustomItemsArray, TEXT(","), true);
diff --git a/Engine/Source/Editor/PropertyEditor/Private/SDetailsViewBase.h b/Engine/Source/Editor/PropertyEditor/Private/SDetailsViewBase.h
--- a/Engine/Source/Editor/PropertyEditor/Private/SDetailsViewBase.h
+++ b/Engine/Source/Editor/PropertyEditor/Private/SDetailsViewBase.h
@@ -47,6 +47,14 @@ typedef STreeView< TSharedRef<FDetailTreeNode> > SDetailTree;
 
 class SDetailsViewBase : public IDetailsViewPrivate
 {
+public:
+
+	/** Global method to remove saved transient detail expansions */
+	static void ResetTransientPerViewDetailExpansion();
+
+	/** Global method to remove saved persistent detail expansions */
+	static void ResetPersistentPerViewDetailExpansion();
+
 public:
 	SDetailsViewBase();
 	virtual ~SDetailsViewBase();
@@ -140,6 +148,8 @@ public:
 	void SetNodeExpansionState(TSharedRef<FDetailTreeNode> InTreeNode, bool bExpand, bool bRecursive) override;
 	void SaveCustomExpansionState(const FString& NodePath, bool bIsExpanded) override;
 	bool GetCustomSavedExpansionState(const FString& NodePath) const override;
+	virtual FName GetViewIdentifier() const override;
+	virtual FDetailsViewArgs::EExpansionPersistenceMethod GetExpansionPersistenceMethod() const override;
 	virtual void NotifyFinishedChangingProperties(const FPropertyChangedEvent& PropertyChangedEvent) override;
 	void RefreshTree() override;
 	virtual void RequestForceRefresh() override;
@@ -293,6 +303,15 @@ protected:
         double CreationTime = 0.0f;
 	};
 
+	/** Return string key for the given struct in config .ini for this details view's properties */
+	const FString& GetDetailPropertyExpansionSection() const;
+
+	/** Return string key for the given struct in config .ini for this details view's custom widgets */
+	const FString& GetDetailCustomWidgetExpansionSection() const;
+
+	/** Return string key for the given struct in config .ini for this details view */
+	FString GetStructExpansionConfigKey(UStruct* InStruct) const;
+
 	/** Called when a node collection animation completes */
 	void HandleNodeAnimationComplete(FAnimatingNodeCollection* InAnimatedNode);
 
diff --git a/Engine/Source/Editor/PropertyEditor/Public/DetailsViewArgs.h b/Engine/Source/Editor/PropertyEditor/Public/DetailsViewArgs.h
--- a/Engine/Source/Editor/PropertyEditor/Public/DetailsViewArgs.h
+++ b/Engine/Source/Editor/PropertyEditor/Public/DetailsViewArgs.h
@@ -44,6 +44,20 @@ struct FDetailsViewArgs
 		ComponentsAndActorsUseNameArea = 1<<2 | ActorsUseNameArea,
 	};
 
+	/** Enum to set expansion persistence. Used to support saving expansion per details view. */
+	enum class EExpansionPersistenceMethod : uint8
+	{
+		/** Expansion will use default persistence (Global shared per struct type) */
+		Default = 0,
+		/** 
+		 * Expansion will be stored transiently per details view identifier and persist in a single editor session.
+		 * The Expansion will be written to EditorPerProjectUserSettings.ini, but will be cleared without read on editor startup.
+		 */
+		TransientPerViewIdentifier = 1,
+		/** Expansion will be stored per details view identifier and persist on disk. */
+		PersistentPerViewIdentifier = 2,
+	};
+
 	/** Controls how CPF_DisableEditOnInstance nodes will be treated */
 	EEditDefaultsOnlyNodeVisibility DefaultsOnlyVisibility;
 	/** The command list from the host of the details view, allowing child widgets to bind actions with a bound chord */
@@ -66,6 +80,8 @@ struct FDetailsViewArgs
 	TSharedPtr<FExtender> OptionsExtender;
 	/** Identifier for this details view; NAME_None if this view is anonymous */
 	FName ViewIdentifier;
+	/** Persistence method for expansion in this details view */
+	EExpansionPersistenceMethod ExpansionPersistenceMethod;
 	/** Notify hook to call when properties are changed */
 	FNotifyHook* NotifyHook;
 	/** Settings for displaying the name area (@see ENameAreaSettings) */
@@ -143,6 +159,7 @@ public:
 		, ScrollbarAlignment(HAlign_Right)
 	    , DetailsNameWidgetOverrideCustomization(nullptr)
 		, ViewIdentifier(NAME_None)
+		, ExpansionPersistenceMethod(EExpansionPersistenceMethod::Default)
 		, NotifyHook(nullptr)
 		, NameAreaSettings(ActorsUseNameArea)
 		, ColumnWidth(.6f)
