[PropertyEditor] Fix struct UPROPERTY rows losing expansion state across ForceRefresh

Plain struct UPROPERTY rows in the details panel collapsed after ForceRefresh() calls
because FDetailItemNode never routed them through the ExpandedDetailNodes save/restore
channel. Named custom builders, named groups, and external property rows were all handled
correctly; plain FStructProperty rows fell through both persistence mechanisms.

FInstancedStruct was unaffected because its customization uses AddCustomBuilder with a
non-None GetName(), satisfying the HasCustomBuilder condition.

Fix: extend FDetailItemNode to treat FStructProperty rows the same as external property
rows in both directions:

- OnItemExpansionChanged: add FStructProperty to the SaveExpansionState condition so
expansion state is written to ExpandedDetailNodes on user interact.

- InitPropertyEditor: add FStructProperty to the HasExternalPropertyRow() restore block
so GetSavedExpansionState() is consulted when the node is constructed after a refresh.

#Synced-CL 52070808

diff --git a/Engine/Source/Editor/PropertyEditor/Private/DetailItemNode.cpp b/Engine/Source/Editor/PropertyEditor/Private/DetailItemNode.cpp
--- a/Engine/Source/Editor/PropertyEditor/Private/DetailItemNode.cpp
+++ b/Engine/Source/Editor/PropertyEditor/Private/DetailItemNode.cpp
@@ -172,7 +172,8 @@ void FDetailItemNode::InitPropertyEditor()
 
 	Customization.PropertyRow->OnItemNodeInitialized( ParentCategory.Pin().ToSharedRef(), IsParentEnabled, ParentGroup.IsValid() ? ParentGroup.Pin() : nullptr );
 
-	if (Customization.HasExternalPropertyRow())
+	if (Customization.HasExternalPropertyRow()
+		|| CastField<FStructProperty>(NodeProperty))
 	{
 		const bool bSaveState = false;
 		SetExpansionState(ParentCategory.Pin()->GetSavedExpansionState(*this), bSaveState);
@@ -497,7 +498,8 @@ void FDetailItemNode::OnItemExpansionChanged( bool bInIsExpanded, bool bShouldSa
 	if (ParentCategory.IsValid() && bShouldSaveState &&
 			(  (Customization.HasCustomBuilder() && Customization.CustomBuilderRow->GetCustomBuilderName() != NAME_None)
 			|| (Customization.HasGroup() && Customization.DetailGroup->GetGroupName() != NAME_None)
-			|| (Customization.HasExternalPropertyRow())))
+			|| (Customization.HasExternalPropertyRow())
+			|| (Customization.HasPropertyNode() && CastField<FStructProperty>(Customization.GetPropertyNode()->GetProperty()))))
 	{
 		ParentCategory.Pin()->SaveExpansionState(*this);
 	}
