Details Panel - Favorites test added

#rnx

#Synced-CL 52128587

diff --git a/Engine/Source/Editor/PropertyEditor/Private/SDetailsView.cpp b/Engine/Source/Editor/PropertyEditor/Private/SDetailsView.cpp
--- a/Engine/Source/Editor/PropertyEditor/Private/SDetailsView.cpp
+++ b/Engine/Source/Editor/PropertyEditor/Private/SDetailsView.cpp
@@ -357,6 +357,7 @@ void SDetailsView::Construct(const FArguments& InArgs, const FDetailsViewArgs& I
 				.ButtonStyle( FAppStyle::Get(), "SimpleButton" )
 				.OnClicked(this, &SDetailsView::OnToggleFavoritesClicked)
 				.ToolTipText(LOCTEXT("ToggleFavorites", "Toggle Favorites"))
+				.AddMetaData<FTagMetaData>(TEXT("Details.ToggleFavorites"))
 				[
 					SNew(SImage)
 					.ColorAndOpacity(this, &SDetailsView::GetToggleFavoritesColor)
