Fix incorrect color in thumbnail.
The previous code in the construct of SPropertyEditorAsset was getting the currently selected asset as if it was a restriction on the supported types, and was overriding the color with it, preventing the automatic color change based on the currently selected asset to work. Now the color override will only be set if there is an actual class selected.
In the case of an array, when the element where added/remove to the array, the existing entries would get their color overriden with the existing asset, preventing the color to be updated dynamically based on the asset type.

[RN] Asset thumbnail: fixed color that was not updated to match the asset type color in some cases.
#tests Played with different asset pickers (in array or without array), and validated that the color was updated as expected.
#Synced-CL 52573034

diff --git a/Engine/Source/Editor/PropertyEditor/Private/UserInterface/PropertyEditor/SPropertyEditorAsset.cpp b/Engine/Source/Editor/PropertyEditor/Private/UserInterface/PropertyEditor/SPropertyEditorAsset.cpp
--- a/Engine/Source/Editor/PropertyEditor/Private/UserInterface/PropertyEditor/SPropertyEditorAsset.cpp
+++ b/Engine/Source/Editor/PropertyEditor/Private/UserInterface/PropertyEditor/SPropertyEditorAsset.cpp
@@ -694,19 +694,12 @@ void SPropertyEditorAsset::Construct(const FArguments& InArgs, const TSharedPtr<
 		}
 		else
 		{
-			TSharedPtr<IAssetTypeActions> AssetTypeActions;
-			if (ObjectClass != nullptr)
+			//Only override the color if we have a meaningful type specified.
+			//=> When override is done, it will internally prevent to switch the color strip to the current asset type color, when updated the asset.
+			if (ObjectClass != nullptr && ObjectClass != UObject::StaticClass())
 			{
-				const UClass* EffectiveClass = ObjectClass;
-				if (EffectiveClass->GetPathName() != Value.AssetData.AssetClassPath.ToString())
-				{
-					if (UClass* AssetDataClass = FindObject<UClass>(Value.AssetData.AssetClassPath); AssetDataClass && AssetDataClass->IsChildOf(EffectiveClass))
-					{
-						EffectiveClass = AssetDataClass;
-					}
-				}
 				FAssetToolsModule& AssetToolsModule = FModuleManager::LoadModuleChecked<FAssetToolsModule>(TEXT("AssetTools"));
-				AssetTypeActions = AssetToolsModule.Get().GetAssetTypeActionsForClass(EffectiveClass).Pin();
+				TSharedPtr<IAssetTypeActions> AssetTypeActions = AssetToolsModule.Get().GetAssetTypeActionsForClass(ObjectClass).Pin();
 
 				if (AssetTypeActions.IsValid())
 				{
